Made the "Parameters" object in remote call payload optional, if no value are provided the call will be made with default initialized parameters if needed
Fixed a crash when shutting down the WebControl server
Fixed return value and output parameters not being properly returned through the RemoteControl API
Fixed check for content-type being application/json
Fixed error code when request is proper but object resolving fails

diff --git a/Engine/Plugins/VirtualProduction/RemoteControl/Source/WebRemoteControl/Private/WebRemoteControl.cpp b/Engine/Plugins/VirtualProduction/RemoteControl/Source/WebRemoteControl/Private/WebRemoteControl.cpp
--- a/Engine/Plugins/VirtualProduction/RemoteControl/Source/WebRemoteControl/Private/WebRemoteControl.cpp
+++ b/Engine/Plugins/VirtualProduction/RemoteControl/Source/WebRemoteControl/Private/WebRemoteControl.cpp
@@ -79,7 +79,7 @@ namespace RemotePayloadSerializer
 			{
 				// this should mean we reached the parameters field, record the start and ending offset
 				case EJsonNotation::ObjectStart:
-					if (JsonReader->GetIdentifier() == TEXT("parameters"))
+					if (JsonReader->GetIdentifier() == TEXT("Parameters"))
 					{
 						ParamBlockStart = Reader.Tell() - sizeof(PayloadCharType);
 						if (JsonReader->SkipObject())
@@ -136,36 +136,35 @@ namespace RemotePayloadSerializer
 			}
 		}
 
-		// if we properly identified the object path and function name as well as identified the starting / end position for parameters
-		// deserialize the parameter payload
-		if (ErrorText.IsEmpty())
+		bool bSuccess = ErrorText.IsEmpty();
+		// if we properly identified the object path and function name and properly resolved the call
+		if (bSuccess && IRemoteControlModule::Get().ResolveCall(ObjectPath, FunctionName, OutCall.CallRef, &ErrorText))
 		{
-			Reader.Seek(ParamBlockStart);
-			Reader.SetLimitSize(ParamBlockEnd + 1);
+			// Initialize the param struct with default parameters
+			OutCall.bGenerateTransaction = bGenerateTransaction;
+			OutCall.ParamStruct = FStructOnScope(OutCall.CallRef.Function.Get());
 
-			// Resolve the Object and Function
-			if (IRemoteControlModule::Get().ResolveCall(ObjectPath, FunctionName, OutCall.CallRef, &ErrorText))
+			// if some parameters were provided, deserialize them
+			if (ParamBlockStart > 0)
 			{
-				OutCall.bGenerateTransaction = bGenerateTransaction;
-				FStructOnScope FuncParam(OutCall.CallRef.Function.Get());
+				Reader.Seek(ParamBlockStart);
+				Reader.SetLimitSize(ParamBlockEnd + 1);
+
 				FJsonStructDeserializerBackend Backend(Reader);
-				if (FStructDeserializer::Deserialize((void*)FuncParam.GetStructMemory(), *const_cast<UStruct*>(FuncParam.GetStruct()), Backend, FStructDeserializerPolicies()))
-				{
-					OutCall.ParamStruct = MoveTemp(FuncParam);
-				}
-				else
+				if (!FStructDeserializer::Deserialize((void*)OutCall.ParamStruct.GetStructMemory(), *const_cast<UStruct*>(OutCall.ParamStruct.GetStruct()), Backend, FStructDeserializerPolicies()))
 				{
 					ErrorText = TEXT("Parameters object improperly formatted.");
+					bSuccess = false;
 				}
 			}
 		}
 
+		// Print deserialization or resolving error, having resolving error is still considered successful 
 		if (!ErrorText.IsEmpty())
 		{
 			UE_LOG(LogRemoteControl, Error, TEXT("Web Remote Call deserialization error: %s"), *ErrorText);
-			return false;
 		}
-		return true;
+		return bSuccess;
 	}
 
 	bool SerializeCall(const FRCCall& InCall, TArray<uint8>& OutPayload, bool bOnlyReturn = false)
@@ -192,7 +191,7 @@ namespace RemotePayloadSerializer
 		{
 			Policies.PropertyFilter = [](const UProperty* CurrentProp, const UProperty* ParentProperty) -> bool
 			{
-				return CurrentProp->GetFName() == ReturnValuePropName || ParentProperty != nullptr;
+				return CurrentProp->HasAnyPropertyFlags(CPF_ReturnParm|CPF_OutParm) || ParentProperty != nullptr;
 			};
 		}
 
@@ -300,19 +299,21 @@ namespace RemotePayloadSerializer
 			}
 		}
 
-		if (ErrorText.IsEmpty())
+		bool bSuccess = ErrorText.IsEmpty();
+		if (bSuccess)
 		{
 			// if we properly identified the object path, property name and access type as well as identified the starting / end position for property value
 			// resolve the object reference
 			IRemoteControlModule::Get().ResolveObject(AccessValue, ObjectPath, PropertyName, OutObjectRef, &ErrorText);
 		}
 
+		// Print deserialization or resolving error, having resolving error is still considered successful 
 		if (!ErrorText.IsEmpty())
 		{
 			UE_LOG(LogRemoteControl, Error, TEXT("Web Remote Object Access deserialization error: %s"), *ErrorText);
-			return false;
 		}
-		return true;
+
+		return bSuccess;
 	}
 }
 
@@ -389,7 +390,7 @@ public:
 				AddContentTypeHeaders(Response.Get(), TEXT("application/json"));
 
 
-				if (IsRequestContentType(Request, TEXT("application/json")))
+				if (!IsRequestContentType(Request, TEXT("application/json")))
 				{
 					Response->Code = EHttpServerResponseCodes::BadRequest;
 					OnComplete(MoveTemp(Response));
@@ -407,9 +408,10 @@ public:
 					return true;
 				}
 
+				// if we haven't resolve the object or function, return not found
 				if (!Call.IsValid())
 				{
-					Response->Code = EHttpServerResponseCodes::Forbidden;
+					Response->Code = EHttpServerResponseCodes::NotFound;
 					OnComplete(MoveTemp(Response));
 					return true;
 				}
@@ -438,7 +440,7 @@ public:
 				AddCORSHeaders(Response.Get());
 				AddContentTypeHeaders(Response.Get(), TEXT("application/json"));
 
-				if (IsRequestContentType(Request, TEXT("application/json")))
+				if (!IsRequestContentType(Request, TEXT("application/json")))
 				{
 					Response->Code = EHttpServerResponseCodes::BadRequest;
 					OnComplete(MoveTemp(Response));
@@ -458,6 +460,14 @@ public:
 					return true;
 				}
 
+				// if we haven't found the object, return a not found error code
+				if (!ObjectRef.IsValid())
+				{
+					Response->Code = EHttpServerResponseCodes::NotFound;
+					OnComplete(MoveTemp(Response));
+					return true;
+				}
+
 				switch (ObjectRef.Access)
 				{
 				case ERCAccess::READ_ACCESS:
@@ -499,7 +509,7 @@ public:
 				AddCORSHeaders(Response.Get());
 				AddContentTypeHeaders(Response.Get(), TEXT("application/json"));
 
-				if (IsRequestContentType(Request, TEXT("application/json")))
+				if (!IsRequestContentType(Request, TEXT("application/json")))
 				{
 					Response->Code = EHttpServerResponseCodes::BadRequest;
 					OnComplete(MoveTemp(Response));
@@ -542,10 +552,27 @@ public:
 	{
 		if (HttpRouter)
 		{
-			HttpRouter->UnbindRoute(RemoteEventRouteHandle);
-			HttpRouter->UnbindRoute(RemotePropertyRouteHandle);
-			HttpRouter->UnbindRoute(RemoteCallRouteHandle);
-			HttpRouter->UnbindRoute(RemoteRouteOptionsHandle);
+			if (RemoteEventRouteHandle)
+			{
+				HttpRouter->UnbindRoute(RemoteEventRouteHandle);
+				RemoteEventRouteHandle.Reset();
+			}
+			if (RemotePropertyRouteHandle)
+			{
+				HttpRouter->UnbindRoute(RemotePropertyRouteHandle);
+				RemotePropertyRouteHandle.Reset();
+			}
+			if (RemoteCallRouteHandle)
+			{
+				HttpRouter->UnbindRoute(RemoteCallRouteHandle);
+				RemoteCallRouteHandle.Reset();
+			}
+			if (RemoteRouteOptionsHandle)
+			{
+				HttpRouter->UnbindRoute(RemoteRouteOptionsHandle);
+				RemoteRouteOptionsHandle.Reset();
+			}
+			
 		}
 		HttpRouter.Reset();
 
@@ -558,7 +585,7 @@ public:
 private:
 	bool IsRequestContentType(const FHttpServerRequest& Request, const FString& ContentType)
 	{
-		if (const TArray<FString>* ContentTypeHeaders = Request.Headers.Find(TEXT("ContentType")))
+		if (const TArray<FString>* ContentTypeHeaders = Request.Headers.Find(TEXT("Content-Type")))
 		{
 			return ContentTypeHeaders->Num() > 0 && (*ContentTypeHeaders)[0] == ContentType;
 		}
