Added a ResetToDefault option when setting properties through the WebControl API
Merged from CL 7141689

diff --git a/Engine/Plugins/VirtualProduction/RemoteControl/Source/RemoteControl/Private/RemoteControlModule.cpp b/Engine/Plugins/VirtualProduction/RemoteControl/Source/RemoteControl/Private/RemoteControlModule.cpp
--- a/Engine/Plugins/VirtualProduction/RemoteControl/Source/RemoteControl/Private/RemoteControlModule.cpp
+++ b/Engine/Plugins/VirtualProduction/RemoteControl/Source/RemoteControl/Private/RemoteControlModule.cpp
@@ -213,7 +213,7 @@ public:
 
 	virtual bool GetObjectProperties(const FRCObjectReference& ObjectAccess, IStructSerializerBackend& Backend) override
 	{
-		if (ObjectAccess.IsValid())
+		if (ObjectAccess.IsValid() && ObjectAccess.Access == ERCAccess::READ_ACCESS)
 		{
 			UObject* Object = ObjectAccess.Object.Get();
 			FStructSerializerPolicies Policies;
@@ -241,7 +241,7 @@ public:
 
 	virtual bool SetObjectProperties(const FRCObjectReference& ObjectAccess, IStructDeserializerBackend& Backend) override
 	{
-		if (ObjectAccess.IsValid())
+		if (ObjectAccess.IsValid() && (ObjectAccess.Access == ERCAccess::WRITE_ACCESS || ObjectAccess.Access == ERCAccess::WRITE_TRANSACTION_ACCESS))
 		{
 			bool bGenerateTransaction = ObjectAccess.Access == ERCAccess::WRITE_TRANSACTION_ACCESS;
 #if WITH_EDITOR
@@ -284,6 +284,36 @@ public:
 		return false;
 	}
 
+	virtual bool ResetObjectProperties(const FRCObjectReference& ObjectAccess) override
+	{
+		if (ObjectAccess.IsValid() && (ObjectAccess.Access == ERCAccess::WRITE_ACCESS || ObjectAccess.Access == ERCAccess::WRITE_TRANSACTION_ACCESS))
+		{
+			bool bGenerateTransaction = ObjectAccess.Access == ERCAccess::WRITE_TRANSACTION_ACCESS;
+#if WITH_EDITOR
+			FScopedTransaction Transaction(LOCTEXT("RemoteResetPropertyTransaction", "Remote Reset Object Property"), bGenerateTransaction);
+#endif
+			UObject* Object = ObjectAccess.Object.Get();
+			if (bGenerateTransaction)
+			{
+				Object->Modify();
+				Object->PreEditChange(ObjectAccess.Property.Get());
+			}
+					
+			ObjectAccess.Property->InitializeValue(ObjectAccess.Property->template ContainerPtrToValuePtr<void>(ObjectAccess.Object.Get()));
+
+			// if we are generating a transaction, also generate post edit property event, event if the change ended up unsuccessful
+			// this is to match the pre edit change call that can unregister components for example
+			if (bGenerateTransaction)
+			{
+				FPropertyChangedEvent PropertyEvent(ObjectAccess.Property.Get());
+				Object->PostEditChangeProperty(PropertyEvent);
+			}
+			return true;
+		}
+		return false;
+	}
+
+
 private:
 };
 
diff --git a/Engine/Plugins/VirtualProduction/RemoteControl/Source/RemoteControl/Public/IRemoteControlModule.h b/Engine/Plugins/VirtualProduction/RemoteControl/Source/RemoteControl/Public/IRemoteControlModule.h
--- a/Engine/Plugins/VirtualProduction/RemoteControl/Source/RemoteControl/Public/IRemoteControlModule.h
+++ b/Engine/Plugins/VirtualProduction/RemoteControl/Source/RemoteControl/Public/IRemoteControlModule.h
@@ -62,7 +62,8 @@ enum class ERCAccess : uint8
 struct FRCObjectReference
 {
 	FRCObjectReference()
-		: Object(nullptr)
+		: Access(ERCAccess::NO_ACCESS)
+		, Object(nullptr)
 		, Property(nullptr)
 	{}
 
@@ -142,4 +143,11 @@ public:
 	 * @return true if the deserialization succeeded
 	 */
 	virtual bool SetObjectProperties(const FRCObjectReference& ObjectAccess, IStructDeserializerBackend& Backend) = 0;
+
+	/**
+	 * Reset the property or the object the Object Reference is pointing to
+	 * @param ObjectAccess the object reference to reset, it should be a write access reference
+	 * @return true if the reset succeeded.
+	 */
+	virtual bool ResetObjectProperties(const FRCObjectReference& ObjectAccess) = 0;
 };
\ No newline at end of file
diff --git a/Engine/Plugins/VirtualProduction/RemoteControl/Source/WebRemoteControl/Private/WebRemoteControl.cpp b/Engine/Plugins/VirtualProduction/RemoteControl/Source/WebRemoteControl/Private/WebRemoteControl.cpp
--- a/Engine/Plugins/VirtualProduction/RemoteControl/Source/WebRemoteControl/Private/WebRemoteControl.cpp
+++ b/Engine/Plugins/VirtualProduction/RemoteControl/Source/WebRemoteControl/Private/WebRemoteControl.cpp
@@ -205,7 +205,7 @@ namespace RemotePayloadSerializer
 		return true;
 	}
 
-	bool DeserializeObjectRef(TArray<uint8>& InTCHARPayload, FRCObjectReference& OutObjectRef, TPair<int32, int32>& OutPropertyBlockRange)
+	bool DeserializeObjectRef(TArray<uint8>& InTCHARPayload, FRCObjectReference& OutObjectRef, bool& OutResetToDefault, TPair<int32, int32>& OutPropertyBlockRange)
 	{
 		FMemoryReader Reader(InTCHARPayload);
 		TSharedRef<TJsonReader<PayloadCharType>> JsonReader = TJsonReader<PayloadCharType>::Create(&Reader);
@@ -258,7 +258,7 @@ namespace RemotePayloadSerializer
 					}
 					else
 					{
-						ErrorText = TEXT("property_value object improperly formatted.");
+						ErrorText = TEXT("property value object improperly formatted.");
 					}
 				}
 				else
@@ -268,6 +268,17 @@ namespace RemotePayloadSerializer
 				break;
 				// this means we should be done with the request object
 			case EJsonNotation::ObjectEnd:
+				break;
+				// read the reset to default property
+			case EJsonNotation::Boolean:
+				if (JsonReader->GetIdentifier() == TEXT("ResetToDefault"))
+				{
+					OutResetToDefault = JsonReader->GetValueAsBoolean();
+				}
+				else
+				{
+					ErrorText = TEXT("unexpected boolean field.");
+				}
 				break;
 				// any other request field should be string properties
 			case EJsonNotation::String:
@@ -389,10 +400,11 @@ public:
 				AddCORSHeaders(Response.Get());
 				AddContentTypeHeaders(Response.Get(), TEXT("application/json"));
 
+				// Initialize the request as a bad request
+				Response->Code = EHttpServerResponseCodes::BadRequest;
 
 				if (!IsRequestContentType(Request, TEXT("application/json")))
 				{
-					Response->Code = EHttpServerResponseCodes::BadRequest;
 					OnComplete(MoveTemp(Response));
 					return true;
 				}
@@ -403,7 +415,6 @@ public:
 				FRCCall Call;
 				if (!RemotePayloadSerializer::DeserializeCall(WorkingBuffer, Call))
 				{
-					Response->Code = EHttpServerResponseCodes::BadRequest;
 					OnComplete(MoveTemp(Response));
 					return true;
 				}
@@ -440,9 +451,11 @@ public:
 				AddCORSHeaders(Response.Get());
 				AddContentTypeHeaders(Response.Get(), TEXT("application/json"));
 
+				// Initialize the request as a bad request
+				Response->Code = EHttpServerResponseCodes::BadRequest;
+
 				if (!IsRequestContentType(Request, TEXT("application/json")))
 				{
-					Response->Code = EHttpServerResponseCodes::BadRequest;
 					OnComplete(MoveTemp(Response));
 					return true;
 				}
@@ -452,10 +465,10 @@ public:
 
 				
 				FRCObjectReference ObjectRef;
+				bool bResetToDefault = false;
 				TPair<int32, int32> PropertyBlockRange;
-				if (!RemotePayloadSerializer::DeserializeObjectRef(WorkingBuffer, ObjectRef, PropertyBlockRange))
+				if (!RemotePayloadSerializer::DeserializeObjectRef(WorkingBuffer, ObjectRef, bResetToDefault, PropertyBlockRange))
 				{
-					Response->Code = EHttpServerResponseCodes::BadRequest;
 					OnComplete(MoveTemp(Response));
 					return true;
 				}
@@ -485,18 +498,28 @@ public:
 				case ERCAccess::WRITE_ACCESS:
 				case ERCAccess::WRITE_TRANSACTION_ACCESS:
 				{
-					FMemoryReader Reader(WorkingBuffer);
-					Reader.Seek(PropertyBlockRange.Key);
-					Reader.SetLimitSize(PropertyBlockRange.Value + 1);
-					FJsonStructDeserializerBackend DeserializerBackend(Reader);
-					if (IRemoteControlModule::Get().SetObjectProperties(ObjectRef, DeserializerBackend))
+					if (bResetToDefault)
 					{
-						Response->Code = EHttpServerResponseCodes::Ok;
+						if (IRemoteControlModule::Get().ResetObjectProperties(ObjectRef))
+						{
+							Response->Code = EHttpServerResponseCodes::Ok;
+						}
+					}
+					else if (PropertyBlockRange.Key > 0)
+					{
+						FMemoryReader Reader(WorkingBuffer);
+						Reader.Seek(PropertyBlockRange.Key);
+						Reader.SetLimitSize(PropertyBlockRange.Value + 1);
+						FJsonStructDeserializerBackend DeserializerBackend(Reader);
+						if (IRemoteControlModule::Get().SetObjectProperties(ObjectRef, DeserializerBackend))
+						{
+							Response->Code = EHttpServerResponseCodes::Ok;
+						}
 					}
 				}
 					break;
 				default:
-					Response->Code = EHttpServerResponseCodes::BadRequest;
+					// Bad request
 					break;
 				}
 
