Update the RemoteControl.upluing description. Add an icon for the plugin.

diff --git a/Engine/Plugins/VirtualProduction/RemoteControl/RemoteControl.uplugin b/Engine/Plugins/VirtualProduction/RemoteControl/RemoteControl.uplugin
--- a/Engine/Plugins/VirtualProduction/RemoteControl/RemoteControl.uplugin
+++ b/Engine/Plugins/VirtualProduction/RemoteControl/RemoteControl.uplugin
@@ -1,20 +1,19 @@
 {
 	"FileVersion": 3,
 	"Version": 1,
-	"VersionName": "0.2",
-	"FriendlyName": "Remote Control",
-	"Description": "A suite of tools to remote control the editor and/or runtime.",
-	"Category": "Virtual Production",
+	"VersionName": "1.0",
+	"FriendlyName": "Remote Control API",
+	"Description": "A suite of tools for controlling the Unreal Engine, both in Editor or at Runtime via a webserver. This allows users to control Unreal Engine remotely thru HTTP requests.  This functionality allows developers to control Unreal thru 3rd party applications and web services.",
+	"Category": "Messaging",
 	"CreatedBy": "Epic Games Inc.",
-	"CreatedByURL": "",
+	"CreatedByURL": "http://epicgames.com",
 	"DocsURL": "",
 	"MarketplaceURL": "",
 	"SupportURL": "",
 	"CanContainContent": false,
 	"IsBetaVersion": true,
 	"Installed": false,
-	"Plugins":
-	[
+	"Plugins": [
 	],
 	"Modules": [
 		{
