Change a few more plugins to UncookedOnly.

#rnx

diff --git a/Engine/Plugins/NetcodeUnitTest/NUTUnrealEngine4/NUTUnrealEngine4.uplugin b/Engine/Plugins/NetcodeUnitTest/NUTUnrealEngine4/NUTUnrealEngine4.uplugin
--- a/Engine/Plugins/NetcodeUnitTest/NUTUnrealEngine4/NUTUnrealEngine4.uplugin
+++ b/Engine/Plugins/NetcodeUnitTest/NUTUnrealEngine4/NUTUnrealEngine4.uplugin
@@ -18,7 +18,7 @@
 	[
 		{
 			"Name" : "NUTUnrealEngine4",
-			"Type" : "Developer",
+			"Type" : "UncookedOnly",
 			"LoadingPhase" : "Default",
 			"WhitelistPlatforms" :
 			[
diff --git a/Engine/Plugins/VirtualProduction/RemoteControl/RemoteControl.uplugin b/Engine/Plugins/VirtualProduction/RemoteControl/RemoteControl.uplugin
--- a/Engine/Plugins/VirtualProduction/RemoteControl/RemoteControl.uplugin
+++ b/Engine/Plugins/VirtualProduction/RemoteControl/RemoteControl.uplugin
@@ -18,12 +18,12 @@
 	"Modules": [
 		{
 			"Name": "RemoteControl",
-			"Type": "Developer",
+			"Type": "UncookedOnly",
 			"LoadingPhase": "Default"
 		},
 		{
 			"Name": "WebRemoteControl",
-			"Type": "Developer",
+			"Type": "UncookedOnly",
 			"LoadingPhase": "Default"
 		}
 
