Fix uninitialized class pointer in ustruct

diff --git a/Engine/Plugins/VirtualProduction/RemoteControl/Source/WebRemoteControl/Private/RemoteControlResponse.h b/Engine/Plugins/VirtualProduction/RemoteControl/Source/WebRemoteControl/Private/RemoteControlResponse.h
--- a/Engine/Plugins/VirtualProduction/RemoteControl/Source/WebRemoteControl/Private/RemoteControlResponse.h
+++ b/Engine/Plugins/VirtualProduction/RemoteControl/Source/WebRemoteControl/Private/RemoteControlResponse.h
@@ -72,7 +72,7 @@ struct FDescribeObjectResponse
 	FString Name;
 
 	UPROPERTY()
-	UClass* Class;
+	UClass* Class = nullptr;
 
 	UPROPERTY()
 	TArray<FRCPropertyDescription> Properties;
