Merge cl from Release-Engine-Staging (15741694) to fix CIS issue

diff --git a/Engine/Plugins/VirtualProduction/RemoteControl/Source/RemoteControlUI/Private/UI/Customizations/RemoteControlEntityCustomization.cpp b/Engine/Plugins/VirtualProduction/RemoteControl/Source/RemoteControlUI/Private/UI/Customizations/RemoteControlEntityCustomization.cpp
--- a/Engine/Plugins/VirtualProduction/RemoteControl/Source/RemoteControlUI/Private/UI/Customizations/RemoteControlEntityCustomization.cpp
+++ b/Engine/Plugins/VirtualProduction/RemoteControl/Source/RemoteControlUI/Private/UI/Customizations/RemoteControlEntityCustomization.cpp
@@ -10,6 +10,7 @@
 #include "RemoteControlEntity.h"
 #include "RemoteControlField.h"
 #include "RemoteControlPreset.h"
+#include "ScopedTransaction.h"
 #include "Widgets/Input/SEditableTextBox.h"
 #include "Widgets/Text/STextBlock.h"
 #include "UObject/UnrealType.h"
