CIS no unity fixes

diff --git a/Engine/Plugins/FastBuildController/Source/Private/FastBuildControllerModule.cpp b/Engine/Plugins/FastBuildController/Source/Private/FastBuildControllerModule.cpp
--- a/Engine/Plugins/FastBuildController/Source/Private/FastBuildControllerModule.cpp
+++ b/Engine/Plugins/FastBuildController/Source/Private/FastBuildControllerModule.cpp
@@ -5,6 +5,7 @@
 #include "HAL/FileManager.h"
 #include "HAL/PlatformFileManager.h"
 #include "HAL/PlatformFile.h"
+#include "Misc/CommandLine.h"
 #include "Misc/ScopeLock.h"
 #include "Misc/Paths.h"
 #include "Misc/Guid.h"
diff --git a/Engine/Plugins/VirtualProduction/RemoteControl/Source/WebRemoteControl/Private/WebRemoteControlEditorRoutes.cpp b/Engine/Plugins/VirtualProduction/RemoteControl/Source/WebRemoteControl/Private/WebRemoteControlEditorRoutes.cpp
--- a/Engine/Plugins/VirtualProduction/RemoteControl/Source/WebRemoteControl/Private/WebRemoteControlEditorRoutes.cpp
+++ b/Engine/Plugins/VirtualProduction/RemoteControl/Source/WebRemoteControl/Private/WebRemoteControlEditorRoutes.cpp
@@ -31,6 +31,8 @@
 #include "Misc/FileHelper.h"
 #include "ObjectTools.h"
 
+#include "Styling/SlateBrush.h"
+
 void FWebRemoteControlEditorRoutes::RegisterRoutes(FWebRemoteControlModule* WebRemoteControl)
 {
 	if (FConsoleManager::Get().FindConsoleVariable(TEXT("WebControl.EnableExperimentalRoutes"))->GetBool())
diff --git a/Engine/Source/Editor/CurveTableEditor/Private/CurveTableEditor.cpp b/Engine/Source/Editor/CurveTableEditor/Private/CurveTableEditor.cpp
--- a/Engine/Source/Editor/CurveTableEditor/Private/CurveTableEditor.cpp
+++ b/Engine/Source/Editor/CurveTableEditor/Private/CurveTableEditor.cpp
@@ -7,6 +7,7 @@
 #include "Framework/Application/SlateApplication.h"
 #include "Widgets/Layout/SScrollBar.h"
 #include "Framework/Layout/Overscroll.h"
+#include "Widgets/Input/SButton.h"
 #include "Widgets/Input/SNumericEntryBox.h"
 #include "Widgets/Input/SSegmentedControl.h"
 #include "Widgets/Views/SListView.h"
diff --git a/Engine/Source/Editor/CurveTableEditor/Private/RealCurveModel.cpp b/Engine/Source/Editor/CurveTableEditor/Private/RealCurveModel.cpp
--- a/Engine/Source/Editor/CurveTableEditor/Private/RealCurveModel.cpp
+++ b/Engine/Source/Editor/CurveTableEditor/Private/RealCurveModel.cpp
@@ -4,7 +4,7 @@
 #include "CurveDrawInfo.h"
 
 #include "Curves/RealCurve.h"
-
+#include "Styling/AppStyle.h"
 
 FRealCurveModel::FRealCurveModel(FRealCurve* InRealCurve, UObject* InOwner)
 	: WeakOwner(InOwner)
diff --git a/Engine/Source/Editor/MergeActors/Private/MeshMergingTool/MeshMergingTool.cpp b/Engine/Source/Editor/MergeActors/Private/MeshMergingTool/MeshMergingTool.cpp
--- a/Engine/Source/Editor/MergeActors/Private/MeshMergingTool/MeshMergingTool.cpp
+++ b/Engine/Source/Editor/MergeActors/Private/MeshMergingTool/MeshMergingTool.cpp
@@ -20,7 +20,7 @@
 #include "AssetRegistryModule.h"
 #include "ScopedTransaction.h"
 #include "MeshMergeModule.h"
-
+#include "ComponentReregisterContext.h"
 
 #define LOCTEXT_NAMESPACE "MeshMergingTool"
 
diff --git a/Engine/Source/Editor/UnrealEd/Private/Commandlets/CommandletSourceControlUtils.cpp b/Engine/Source/Editor/UnrealEd/Private/Commandlets/CommandletSourceControlUtils.cpp
--- a/Engine/Source/Editor/UnrealEd/Private/Commandlets/CommandletSourceControlUtils.cpp
+++ b/Engine/Source/Editor/UnrealEd/Private/Commandlets/CommandletSourceControlUtils.cpp
@@ -8,6 +8,7 @@
 #include "Misc/PackageName.h"
 #include "PackageTools.h"
 #include "SourceControlOperations.h"
+#include "UObject/Linker.h"
 #include "UObject/Package.h"
 
 DEFINE_LOG_CATEGORY(LogSourceControlUtils);
diff --git a/Engine/Source/Runtime/AVEncoder/Private/Decoders/Windows/VideoDecoderH264_Windows.cpp b/Engine/Source/Runtime/AVEncoder/Private/Decoders/Windows/VideoDecoderH264_Windows.cpp
--- a/Engine/Source/Runtime/AVEncoder/Private/Decoders/Windows/VideoDecoderH264_Windows.cpp
+++ b/Engine/Source/Runtime/AVEncoder/Private/Decoders/Windows/VideoDecoderH264_Windows.cpp
@@ -11,6 +11,11 @@
 #include "VideoDecoderAllocationTypes.h"
 #include "VideoDecoderUtilities.h"
 
+#include "Containers/Queue.h"
+#include "Misc/ScopeLock.h"
+#include "Templates/RefCounting.h"
+
+
 #define VERIFY_HR(FNcall,...)						\
 Result = FNcall;									\
 if (FAILED(Result))									\
