FPrivateToken trick used to allow emplacement-style functions (like MakeShared) to compile without giving up constructor access rights.

This replaces classes' use of friendship with internal TSharedPtr functions.

diff --git a/Engine/Plugins/Experimental/CommonConversation/Source/CommonConversationRuntime/Private/ConversationRegistry.cpp b/Engine/Plugins/Experimental/CommonConversation/Source/CommonConversationRuntime/Private/ConversationRegistry.cpp
--- a/Engine/Plugins/Experimental/CommonConversation/Source/CommonConversationRuntime/Private/ConversationRegistry.cpp
+++ b/Engine/Plugins/Experimental/CommonConversation/Source/CommonConversationRuntime/Private/ConversationRegistry.cpp
@@ -17,13 +17,13 @@
 
 TSharedPtr<FConversationsHandle> FConversationsHandle::Create(const UConversationRegistry* InOwningRegistry, const TSharedPtr<FStreamableHandle>& InStreamableHandle, const TArray<FGameplayTag>& InEntryTags)
 {
-	TSharedPtr<FConversationsHandle> ConversationsHandle = MakeShared<FConversationsHandle>(const_cast<UConversationRegistry*>(InOwningRegistry), InStreamableHandle, InEntryTags);
+	TSharedPtr<FConversationsHandle> ConversationsHandle = MakeShared<FConversationsHandle>(FPrivateToken{}, const_cast<UConversationRegistry*>(InOwningRegistry), InStreamableHandle, InEntryTags);
 	ConversationsHandle->Initialize();
 
 	return ConversationsHandle;
 }
 
-FConversationsHandle::FConversationsHandle(UConversationRegistry* InOwningRegistry, const TSharedPtr<FStreamableHandle>& InStreamableHandle, const TArray<FGameplayTag>& InEntryTags)
+FConversationsHandle::FConversationsHandle(FPrivateToken, UConversationRegistry* InOwningRegistry, const TSharedPtr<FStreamableHandle>& InStreamableHandle, const TArray<FGameplayTag>& InEntryTags)
 	: StreamableHandle(InStreamableHandle)
 	, ConversationEntryTags(InEntryTags)
 	, OwningRegistryPtr(InOwningRegistry)
diff --git a/Engine/Plugins/Experimental/CommonConversation/Source/CommonConversationRuntime/Public/ConversationRegistry.h b/Engine/Plugins/Experimental/CommonConversation/Source/CommonConversationRuntime/Public/ConversationRegistry.h
--- a/Engine/Plugins/Experimental/CommonConversation/Source/CommonConversationRuntime/Public/ConversationRegistry.h
+++ b/Engine/Plugins/Experimental/CommonConversation/Source/CommonConversationRuntime/Public/ConversationRegistry.h
@@ -44,8 +44,13 @@ struct COMMONCONVERSATIONRUNTIME_API FConversationsHandle : public TSharedFromTh
 	static TSharedPtr<FConversationsHandle> Create(const UConversationRegistry* InOwningRegistry, const TSharedPtr<FStreamableHandle>& InStreamableHandle, const TArray<FGameplayTag>& InEntryTags);
 
 private:
-	FConversationsHandle(UConversationRegistry* InOwningRegistry, const TSharedPtr<FStreamableHandle>& InStreamableHandle, const TArray<FGameplayTag>& InEntryTags);
+	// Private token only allows members or friends to call MakeShared
+	struct FPrivateToken { explicit FPrivateToken() = default; };
 
+public:
+	FConversationsHandle(FPrivateToken, UConversationRegistry* InOwningRegistry, const TSharedPtr<FStreamableHandle>& InStreamableHandle, const TArray<FGameplayTag>& InEntryTags);
+
+private:
 	void Initialize();
 	void HandleAvailableConversationsChanged();
 
@@ -53,8 +58,6 @@ private:
 	TSharedPtr<FStreamableHandle> StreamableHandle;
 	TArray<FGameplayTag> ConversationEntryTags;
 	TWeakObjectPtr<UConversationRegistry> OwningRegistryPtr;
-
-	friend class SharedPointerInternals::TIntrusiveReferenceController<FConversationsHandle>;
 };
 
 
diff --git a/Engine/Plugins/VirtualProduction/RemoteControl/Source/RemoteControlProtocolWidgets/Private/ViewModels/ProtocolBindingViewModel.cpp b/Engine/Plugins/VirtualProduction/RemoteControl/Source/RemoteControlProtocolWidgets/Private/ViewModels/ProtocolBindingViewModel.cpp
--- a/Engine/Plugins/VirtualProduction/RemoteControl/Source/RemoteControlProtocolWidgets/Private/ViewModels/ProtocolBindingViewModel.cpp
+++ b/Engine/Plugins/VirtualProduction/RemoteControl/Source/RemoteControlProtocolWidgets/Private/ViewModels/ProtocolBindingViewModel.cpp
@@ -1,4 +1,4 @@
-﻿// Copyright Epic Games, Inc. All Rights Reserved.
+// Copyright Epic Games, Inc. All Rights Reserved.
 
 #include "ViewModels/ProtocolBindingViewModel.h"
 
@@ -56,13 +56,13 @@ namespace Commands
 
 TSharedRef<FProtocolBindingViewModel> FProtocolBindingViewModel::Create(const TSharedRef<FProtocolEntityViewModel>& InParentViewModel, const TSharedRef<FRemoteControlProtocolBinding>& InBinding)
 {
-	TSharedRef<FProtocolBindingViewModel> ViewModel = MakeShared<FProtocolBindingViewModel>(InParentViewModel, InBinding);
+	TSharedRef<FProtocolBindingViewModel> ViewModel = MakeShared<FProtocolBindingViewModel>(FPrivateToken{}, InParentViewModel, InBinding);
 	ViewModel->Initialize();
 
 	return ViewModel;
 }
 
-FProtocolBindingViewModel::FProtocolBindingViewModel(const TSharedRef<FProtocolEntityViewModel>& InParentViewModel, const TSharedRef<FRemoteControlProtocolBinding>& InBinding)
+FProtocolBindingViewModel::FProtocolBindingViewModel(FPrivateToken, const TSharedRef<FProtocolEntityViewModel>& InParentViewModel, const TSharedRef<FRemoteControlProtocolBinding>& InBinding)
 	: Preset(InParentViewModel->Preset),
 	ParentViewModel(InParentViewModel),
 	Property(InParentViewModel->Property),
@@ -124,7 +124,7 @@ FGuid FProtocolBindingViewModel::AddRangeMapping()
 TSharedPtr<FProtocolRangeViewModel>& FProtocolBindingViewModel::AddRangeMappingInternal()
 {
 	FGuid RangeId = Commands::AddRangeMappingInternal(Preset.Get(), {ParentViewModel.Pin()->GetId(), GetId(), FGuid()});
-	TSharedPtr<FProtocolRangeViewModel>& NewRangeViewModel = Ranges.Add_GetRef(MakeShared<FProtocolRangeViewModel>(SharedThis(this), RangeId));
+	TSharedPtr<FProtocolRangeViewModel>& NewRangeViewModel = Ranges.Add_GetRef(MakeShared<FProtocolRangeViewModel>(FProtocolRangeViewModel::FPrivateToken{}, SharedThis(this), RangeId));
 	return NewRangeViewModel;
 }
 
diff --git a/Engine/Plugins/VirtualProduction/RemoteControl/Source/RemoteControlProtocolWidgets/Private/ViewModels/ProtocolBindingViewModel.h b/Engine/Plugins/VirtualProduction/RemoteControl/Source/RemoteControlProtocolWidgets/Private/ViewModels/ProtocolBindingViewModel.h
--- a/Engine/Plugins/VirtualProduction/RemoteControl/Source/RemoteControlProtocolWidgets/Private/ViewModels/ProtocolBindingViewModel.h
+++ b/Engine/Plugins/VirtualProduction/RemoteControl/Source/RemoteControlProtocolWidgets/Private/ViewModels/ProtocolBindingViewModel.h
@@ -1,4 +1,4 @@
-﻿// Copyright Epic Games, Inc. All Rights Reserved.
+// Copyright Epic Games, Inc. All Rights Reserved.
 
 #pragma once
 
@@ -69,12 +69,16 @@ public:
 	FOnChanged& OnChanged() { return OnChangedDelegate; }
 
 private:
-	template <typename FProtocolBindingViewModel>
-    friend class SharedPointerInternals::TIntrusiveReferenceController;
-	
-	FProtocolBindingViewModel() = default;
-	FProtocolBindingViewModel(const TSharedRef<FProtocolEntityViewModel>& InParentViewModel, const TSharedRef<FRemoteControlProtocolBinding>& InBinding);
+	friend class FProtocolEntityViewModel;
+
+	// Private token only allows members or friends to call MakeShared
+	struct FPrivateToken { explicit FPrivateToken() = default; };
 
+public:
+	FProtocolBindingViewModel(FPrivateToken) {}
+	FProtocolBindingViewModel(FPrivateToken, const TSharedRef<FProtocolEntityViewModel>& InParentViewModel, const TSharedRef<FRemoteControlProtocolBinding>& InBinding);
+
+private:
 	void Initialize();
 
 	TSharedPtr<FProtocolRangeViewModel>& AddRangeMappingInternal();
diff --git a/Engine/Plugins/VirtualProduction/RemoteControl/Source/RemoteControlProtocolWidgets/Private/ViewModels/ProtocolEntityViewModel.cpp b/Engine/Plugins/VirtualProduction/RemoteControl/Source/RemoteControlProtocolWidgets/Private/ViewModels/ProtocolEntityViewModel.cpp
--- a/Engine/Plugins/VirtualProduction/RemoteControl/Source/RemoteControlProtocolWidgets/Private/ViewModels/ProtocolEntityViewModel.cpp
+++ b/Engine/Plugins/VirtualProduction/RemoteControl/Source/RemoteControlProtocolWidgets/Private/ViewModels/ProtocolEntityViewModel.cpp
@@ -1,4 +1,4 @@
-﻿// Copyright Epic Games, Inc. All Rights Reserved.
+// Copyright Epic Games, Inc. All Rights Reserved.
 
 #include "ProtocolEntityViewModel.h"
 
@@ -65,13 +65,13 @@ namespace Commands
 
 TSharedRef<FProtocolEntityViewModel> FProtocolEntityViewModel::Create(URemoteControlPreset* InPreset, const FGuid& InEntityId)
 {
-	TSharedRef<FProtocolEntityViewModel> ViewModel = MakeShared<FProtocolEntityViewModel>(InPreset, InEntityId);
+	TSharedRef<FProtocolEntityViewModel> ViewModel = MakeShared<FProtocolEntityViewModel>(FPrivateToken{}, InPreset, InEntityId);
 	ViewModel->Initialize();
 
 	return ViewModel;
 }
 
-FProtocolEntityViewModel::FProtocolEntityViewModel(URemoteControlPreset* InPreset, const FGuid& InEntityId)
+FProtocolEntityViewModel::FProtocolEntityViewModel(FPrivateToken, URemoteControlPreset* InPreset, const FGuid& InEntityId)
     : Preset(InPreset),
     PropertyId(InEntityId)
 {
@@ -165,7 +165,7 @@ TSharedPtr<FProtocolBindingViewModel> FProtocolEntityViewModel::AddBinding(const
             );
 		}
 
-		TSharedPtr<FProtocolBindingViewModel>& NewBindingViewModel = Bindings.Add_GetRef(MakeShared<FProtocolBindingViewModel>(SharedThis(this), ProtocolBinding.ToSharedRef()));
+		TSharedPtr<FProtocolBindingViewModel>& NewBindingViewModel = Bindings.Add_GetRef(MakeShared<FProtocolBindingViewModel>(FProtocolBindingViewModel::FPrivateToken{}, SharedThis(this), ProtocolBinding.ToSharedRef()));
 		NewBindingViewModel->AddDefaultRangeMappings();
 		
 		OnBindingAddedDelegate.Broadcast(NewBindingViewModel.ToSharedRef());
diff --git a/Engine/Plugins/VirtualProduction/RemoteControl/Source/RemoteControlProtocolWidgets/Private/ViewModels/ProtocolEntityViewModel.h b/Engine/Plugins/VirtualProduction/RemoteControl/Source/RemoteControlProtocolWidgets/Private/ViewModels/ProtocolEntityViewModel.h
--- a/Engine/Plugins/VirtualProduction/RemoteControl/Source/RemoteControlProtocolWidgets/Private/ViewModels/ProtocolEntityViewModel.h
+++ b/Engine/Plugins/VirtualProduction/RemoteControl/Source/RemoteControlProtocolWidgets/Private/ViewModels/ProtocolEntityViewModel.h
@@ -1,4 +1,4 @@
-﻿// Copyright Epic Games, Inc. All Rights Reserved.
+// Copyright Epic Games, Inc. All Rights Reserved.
 
 #pragma once
 
@@ -52,12 +52,14 @@ public:
 	FOnChanged& OnChanged() { return OnChangedDelegate; }
 
 private:
-	template <typename FProtocolEntityViewModel>
-    friend class SharedPointerInternals::TIntrusiveReferenceController;
-	
-	FProtocolEntityViewModel() = default;
-	FProtocolEntityViewModel(URemoteControlPreset* InPreset, const FGuid& InEntityId);
+	// Private token only allows members or friends to call MakeShared
+	struct FPrivateToken { explicit FPrivateToken() = default; };
+
+public:
+	FProtocolEntityViewModel(FPrivateToken) {}
+	FProtocolEntityViewModel(FPrivateToken, URemoteControlPreset* InPreset, const FGuid& InEntityId);
 
+private:
 	void Initialize();
 
 	//~ Begin FEditorUndoClient Interface
diff --git a/Engine/Plugins/VirtualProduction/RemoteControl/Source/RemoteControlProtocolWidgets/Private/ViewModels/ProtocolRangeViewModel.cpp b/Engine/Plugins/VirtualProduction/RemoteControl/Source/RemoteControlProtocolWidgets/Private/ViewModels/ProtocolRangeViewModel.cpp
--- a/Engine/Plugins/VirtualProduction/RemoteControl/Source/RemoteControlProtocolWidgets/Private/ViewModels/ProtocolRangeViewModel.cpp
+++ b/Engine/Plugins/VirtualProduction/RemoteControl/Source/RemoteControlProtocolWidgets/Private/ViewModels/ProtocolRangeViewModel.cpp
@@ -1,4 +1,4 @@
-﻿// Copyright Epic Games, Inc. All Rights Reserved.
+// Copyright Epic Games, Inc. All Rights Reserved.
 
 #include "ProtocolRangeViewModel.h"
 
@@ -10,13 +10,13 @@
 
 TSharedRef<FProtocolRangeViewModel> FProtocolRangeViewModel::Create(const TSharedRef<FProtocolBindingViewModel>& InParentViewModel, const FGuid& InRangeId)
 {
-	TSharedRef<FProtocolRangeViewModel> ViewModel = MakeShared<FProtocolRangeViewModel>(InParentViewModel, InRangeId);
+	TSharedRef<FProtocolRangeViewModel> ViewModel = MakeShared<FProtocolRangeViewModel>(FPrivateToken{}, InParentViewModel, InRangeId);
 	ViewModel->Initialize();
 
 	return ViewModel;
 }
 
-FProtocolRangeViewModel::FProtocolRangeViewModel(const TSharedRef<FProtocolBindingViewModel>& InParentViewModel, const FGuid& InRangeId)
+FProtocolRangeViewModel::FProtocolRangeViewModel(FPrivateToken, const TSharedRef<FProtocolBindingViewModel>& InParentViewModel, const FGuid& InRangeId)
 	: Preset(InParentViewModel->Preset),
 	ParentViewModel(InParentViewModel),
 	RangeId(InRangeId)
diff --git a/Engine/Plugins/VirtualProduction/RemoteControl/Source/RemoteControlProtocolWidgets/Private/ViewModels/ProtocolRangeViewModel.h b/Engine/Plugins/VirtualProduction/RemoteControl/Source/RemoteControlProtocolWidgets/Private/ViewModels/ProtocolRangeViewModel.h
--- a/Engine/Plugins/VirtualProduction/RemoteControl/Source/RemoteControlProtocolWidgets/Private/ViewModels/ProtocolRangeViewModel.h
+++ b/Engine/Plugins/VirtualProduction/RemoteControl/Source/RemoteControlProtocolWidgets/Private/ViewModels/ProtocolRangeViewModel.h
@@ -1,4 +1,4 @@
-﻿// Copyright Epic Games, Inc. All Rights Reserved.
+// Copyright Epic Games, Inc. All Rights Reserved.
 
 #pragma once
 
@@ -91,13 +91,17 @@ public:
 	/** Something has changed within the ViewModel */
 	FOnChanged& OnChanged() { return OnChangedDelegate; }
 
-protected:
-	template <typename FProtocolRangeViewModel>
-	friend class SharedPointerInternals::TIntrusiveReferenceController;
+private:
+	// Private token only allows members or friends to call MakeShared
+	struct FPrivateToken { explicit FPrivateToken() = default; };
 
-	FProtocolRangeViewModel() = default;
-	FProtocolRangeViewModel(const TSharedRef<FProtocolBindingViewModel>& InParentViewModel, const FGuid& InRangeId);
-	
+	friend class FProtocolBindingViewModel;
+
+public:
+	FProtocolRangeViewModel(FPrivateToken) {}
+	FProtocolRangeViewModel(FPrivateToken, const TSharedRef<FProtocolBindingViewModel>& InParentViewModel, const FGuid& InRangeId);
+
+protected:
 	void Initialize();;
 
 	FRemoteControlProtocolMapping* GetRangesData() const { return ParentViewModel.Pin()->GetRangesMapping(RangeId); }
diff --git a/Engine/Source/Developer/BlueprintCompilerCppBackend/Private/BlueprintCompilerCppBackendGatherDependencies.cpp b/Engine/Source/Developer/BlueprintCompilerCppBackend/Private/BlueprintCompilerCppBackendGatherDependencies.cpp
--- a/Engine/Source/Developer/BlueprintCompilerCppBackend/Private/BlueprintCompilerCppBackendGatherDependencies.cpp
+++ b/Engine/Source/Developer/BlueprintCompilerCppBackend/Private/BlueprintCompilerCppBackendGatherDependencies.cpp
@@ -456,7 +456,7 @@ bool FGatherConvertedClassDependencies::IsFieldFromExcludedPackage(const UField*
 	return false;
 }
 
-FGatherConvertedClassDependencies::FGatherConvertedClassDependencies(UStruct* InStruct, const FCompilerNativizationOptions& InNativizationOptions)
+FGatherConvertedClassDependencies::FGatherConvertedClassDependencies(FPrivateToken, UStruct* InStruct, const FCompilerNativizationOptions& InNativizationOptions)
 	: OriginalStruct(InStruct)
 	, NativizationOptions(InNativizationOptions)
 {
@@ -849,7 +849,7 @@ TSharedPtr<FGatherConvertedClassDependencies> FGatherConvertedClassDependencies:
 		ConvertedClassDependenciesPtr = CachedConvertedClassDependencies.FindOrAdd(InStruct);
 		if (!ConvertedClassDependenciesPtr.IsValid())
 		{
-			ConvertedClassDependenciesPtr = MakeShared<FGatherConvertedClassDependencies>(InStruct, InNativizationOptions);
+			ConvertedClassDependenciesPtr = MakeShared<FGatherConvertedClassDependencies>(FPrivateToken{}, InStruct, InNativizationOptions);
 			check(ConvertedClassDependenciesPtr.IsValid());
 		}
 	}
diff --git a/Engine/Source/Developer/BlueprintCompilerCppBackend/Public/BlueprintCompilerCppBackendGatherDependencies.h b/Engine/Source/Developer/BlueprintCompilerCppBackend/Public/BlueprintCompilerCppBackendGatherDependencies.h
--- a/Engine/Source/Developer/BlueprintCompilerCppBackend/Public/BlueprintCompilerCppBackendGatherDependencies.h
+++ b/Engine/Source/Developer/BlueprintCompilerCppBackend/Public/BlueprintCompilerCppBackendGatherDependencies.h
@@ -59,11 +59,13 @@ public:
 
 	static bool IsFieldFromExcludedPackage(const UField* Field, const TSet<FName>& InExcludedModules);
 
-protected:
-	FGatherConvertedClassDependencies(UStruct* InStruct, const FCompilerNativizationOptions& InNativizationOptions);
+private:
+	// Private token only allows members or friends to call MakeShared
+	struct FPrivateToken { explicit FPrivateToken() = default; };
 
-	/** Friend for access to constructor via MakeShared */
-	friend class SharedPointerInternals::TIntrusiveReferenceController<FGatherConvertedClassDependencies>;
+public:
+	FGatherConvertedClassDependencies(FPrivateToken, UStruct* InStruct, const FCompilerNativizationOptions& InNativizationOptions);
 
+protected:
 	void DependenciesForHeader();
 };
diff --git a/Engine/Source/Runtime/Online/WebSockets/Private/Lws/LwsWebSocket.cpp b/Engine/Source/Runtime/Online/WebSockets/Private/Lws/LwsWebSocket.cpp
--- a/Engine/Source/Runtime/Online/WebSockets/Private/Lws/LwsWebSocket.cpp
+++ b/Engine/Source/Runtime/Online/WebSockets/Private/Lws/LwsWebSocket.cpp
@@ -72,7 +72,7 @@ const TCHAR* FLwsWebSocket::ToString(const EState InState)
 	return TEXT("Unknown");
 }
 
-FLwsWebSocket::FLwsWebSocket(const FString& InUrl, const TArray<FString>& InProtocols, const FString& InUpgradeHeader)
+FLwsWebSocket::FLwsWebSocket(FPrivateToken, const FString& InUrl, const TArray<FString>& InProtocols, const FString& InUpgradeHeader)
 	: State(EState::None)
 	, LastGameThreadState(EState::None)
 	, bWasSendQueueEmpty(true)
diff --git a/Engine/Source/Runtime/Online/WebSockets/Private/Lws/LwsWebSocket.h b/Engine/Source/Runtime/Online/WebSockets/Private/Lws/LwsWebSocket.h
--- a/Engine/Source/Runtime/Online/WebSockets/Private/Lws/LwsWebSocket.h
+++ b/Engine/Source/Runtime/Online/WebSockets/Private/Lws/LwsWebSocket.h
@@ -97,6 +97,8 @@ struct FLwsReceiveBufferText
 
 typedef TUniquePtr<FLwsReceiveBufferText> FLwsReceiveBufferTextPtr;
 
+class FLwsWebSocketsManager;
+
 class FLwsWebSocket
 	: public IWebSocket
 	, public TSharedFromThis<FLwsWebSocket>
@@ -185,12 +187,16 @@ public:
 	void LwsThreadTick();
 
 private:
-	/** Constructor */
-	FLwsWebSocket(const FString& Url, const TArray<FString>& Protocols, const FString& UpgradeHeader);
+	// Private token only allows members or friends to call MakeShared
+	struct FPrivateToken { explicit FPrivateToken() = default; };
+
+	friend class FLwsWebSocketsManager;
 
-	/** Friend for access to constructor via MakeShared */
-	friend class SharedPointerInternals::TIntrusiveReferenceController<FLwsWebSocket>;
+public:
+	/** Constructor */
+	FLwsWebSocket(FPrivateToken, const FString& Url, const TArray<FString>& Protocols, const FString& UpgradeHeader);
 
+private:
 	/**
 	 * Start connecting
 	 * @param LwsContext libwebsockets context
diff --git a/Engine/Source/Runtime/Online/WebSockets/Private/Lws/LwsWebSocketsManager.cpp b/Engine/Source/Runtime/Online/WebSockets/Private/Lws/LwsWebSocketsManager.cpp
--- a/Engine/Source/Runtime/Online/WebSockets/Private/Lws/LwsWebSocketsManager.cpp
+++ b/Engine/Source/Runtime/Online/WebSockets/Private/Lws/LwsWebSocketsManager.cpp
@@ -398,7 +398,7 @@ TSharedRef<IWebSocket> FLwsWebSocketsManager::CreateWebSocket(const FString& Url
 		UpgradeHeaderString += FString::Printf(TEXT("%s: %s\r\n"), *OneHeader.Key, *OneHeader.Value);
 	}
 
-	FLwsWebSocketRef Socket = MakeShared<FLwsWebSocket>(Url, Protocols, UpgradeHeaderString);
+	FLwsWebSocketRef Socket = MakeShared<FLwsWebSocket>(FLwsWebSocket::FPrivateToken{}, Url, Protocols, UpgradeHeaderString);
 	return Socket;
 }
 
diff --git a/Engine/Source/Runtime/Slate/Public/Framework/Text/ShapedTextCache.h b/Engine/Source/Runtime/Slate/Public/Framework/Text/ShapedTextCache.h
--- a/Engine/Source/Runtime/Slate/Public/Framework/Text/ShapedTextCache.h
+++ b/Engine/Source/Runtime/Slate/Public/Framework/Text/ShapedTextCache.h
@@ -57,7 +57,7 @@ public:
 	/** Create a new shaped text cache */
 	static FShapedTextCacheRef Create(const TSharedRef<FSlateFontCache>& InFontCache)
 	{
-		return MakeShared<FShapedTextCache>(InFontCache);
+		return MakeShared<FShapedTextCache>(FPrivateToken{}, InFontCache);
 	}
 
 	/**
@@ -109,14 +109,17 @@ public:
 	}
 
 private:
-	friend class SharedPointerInternals::TIntrusiveReferenceController<FShapedTextCache>;
+	// Private token only allows members or friends to call MakeShared
+	struct FPrivateToken { explicit FPrivateToken() = default; };
 
+public:
 	/** Constructor */
-	FShapedTextCache(const TSharedRef<FSlateFontCache>& InFontCache)
+	FShapedTextCache(FPrivateToken, const TSharedRef<FSlateFontCache>& InFontCache)
 		: FontCachePtr(InFontCache)
 	{
 	}
 
+private:
 	/** Font cache to use when shaping text */
 	TWeakPtr<FSlateFontCache> FontCachePtr;
 
