Fix not being able to expose properties

diff --git a/Engine/Plugins/VirtualProduction/RemoteControl/Source/RemoteControlUI/Private/RemoteControlUIModule.cpp b/Engine/Plugins/VirtualProduction/RemoteControl/Source/RemoteControlUI/Private/RemoteControlUIModule.cpp
--- a/Engine/Plugins/VirtualProduction/RemoteControl/Source/RemoteControlUI/Private/RemoteControlUIModule.cpp
+++ b/Engine/Plugins/VirtualProduction/RemoteControl/Source/RemoteControlUI/Private/RemoteControlUIModule.cpp
@@ -221,13 +221,7 @@ bool FRemoteControlUIModule::CanToggleExposeProperty(TSharedPtr<IPropertyHandle>
 {
 	if (TSharedPtr<SRemoteControlPanel> Panel = WeakActivePanel.Pin())
 	{
-		if (Handle)
-		{
-			if (!ShouldDisplayExposeIcon(Handle.ToSharedRef()))
-			{
-				return true;
-			}
-		}
+		return Handle && ShouldDisplayExposeIcon(Handle.ToSharedRef());
 	}
 
 	return false;
diff --git a/Engine/Plugins/VirtualProduction/RemoteControl/Source/RemoteControlUI/Private/UI/SRemoteControlPanel.cpp b/Engine/Plugins/VirtualProduction/RemoteControl/Source/RemoteControlUI/Private/UI/SRemoteControlPanel.cpp
--- a/Engine/Plugins/VirtualProduction/RemoteControl/Source/RemoteControlUI/Private/UI/SRemoteControlPanel.cpp
+++ b/Engine/Plugins/VirtualProduction/RemoteControl/Source/RemoteControlUI/Private/UI/SRemoteControlPanel.cpp
@@ -586,20 +586,21 @@ TSharedRef<SWidget> SRemoteControlPanel::CreateExposeButton()
 	MenuBuilder.EndSection();
 	
 	return SAssignNew(ExposeComboButton, SComboButton)
-	.Visibility_Lambda([this]() { return this->bIsInEditMode ? EVisibility::Visible : EVisibility::Collapsed; })
-	.ButtonStyle(FEditorStyle::Get(), "PropertyEditor.AssetComboStyle")
-	.ForegroundColor(FSlateColor::UseForeground())
-	.CollapseMenuOnParentFocus(true)
-	.ButtonContent()
-	[
-		SNew(STextBlock)
-		.TextStyle(FEditorStyle::Get(), "ContentBrowser.TopBar.Font")
-		.Text(LOCTEXT("ExposeButtonLabel", "Expose"))
-	]
-	.MenuContent()
-	[
-		MenuBuilder.MakeWidget()
-	];
+		.Visibility_Lambda([this]() { return this->bIsInEditMode ? EVisibility::Visible : EVisibility::Collapsed; })
+		.ButtonStyle(FEditorStyle::Get(), "PropertyEditor.AssetComboStyle")
+		.ForegroundColor(FSlateColor::UseForeground())
+		.CollapseMenuOnParentFocus(true)
+		.ContentPadding(FMargin(10.0f, 0.f))
+		.ButtonContent()
+		[
+			SNew(STextBlock)
+			.TextStyle(FEditorStyle::Get(), "ContentBrowser.TopBar.Font")
+			.Text(LOCTEXT("ExposeButtonLabel", "Expose"))
+		]
+		.MenuContent()
+		[
+			MenuBuilder.MakeWidget()
+		];
 }
 
 TSharedRef<SWidget> SRemoteControlPanel::CreateExposeByClassWidget()
diff --git a/Engine/Source/Editor/PropertyEditor/Private/SDetailSingleItemRow.cpp b/Engine/Source/Editor/PropertyEditor/Private/SDetailSingleItemRow.cpp
--- a/Engine/Source/Editor/PropertyEditor/Private/SDetailSingleItemRow.cpp
+++ b/Engine/Source/Editor/PropertyEditor/Private/SDetailSingleItemRow.cpp
@@ -1107,10 +1107,15 @@ void SDetailSingleItemRow::CreateGlobalExtensionWidgets(TArray<FPropertyRowExten
 	FOnGenerateGlobalRowExtensionArgs Args;
 	Args.OwnerTreeNode = OwnerTreeNode;
 
+	
 	if (Customization->HasPropertyNode())
 	{
 		Args.PropertyHandle = PropertyEditorHelpers::GetPropertyHandle(Customization->GetPropertyNode().ToSharedRef(), nullptr, nullptr);
 	}
+	else if (WidgetRow.GetPropertyHandles().Num() && WidgetRow.GetPropertyHandles()[0] && WidgetRow.GetPropertyHandles()[0]->IsValidHandle())
+	{
+		Args.PropertyHandle = WidgetRow.GetPropertyHandles()[0];
+	}
 
 	PropertyEditorModule.GetGlobalRowExtensionDelegate().Broadcast(Args, OutExtensions);
 }
