- Fix crash when shutting down remote control

diff --git a/Engine/Plugins/VirtualProduction/RemoteControl/Source/WebRemoteControl/Private/RemoteControlWebSocketServer.cpp b/Engine/Plugins/VirtualProduction/RemoteControl/Source/WebRemoteControl/Private/RemoteControlWebSocketServer.cpp
--- a/Engine/Plugins/VirtualProduction/RemoteControl/Source/WebRemoteControl/Private/RemoteControlWebSocketServer.cpp
+++ b/Engine/Plugins/VirtualProduction/RemoteControl/Source/WebRemoteControl/Private/RemoteControlWebSocketServer.cpp
@@ -123,6 +123,7 @@ void FRCWebSocketServer::Stop()
 	if (TickerHandle.IsValid())
 	{
 		FTSTicker::GetCoreTicker().RemoveTicker(TickerHandle);
+		TickerHandle.Reset();
 	}
 	Router.Reset();
 	Server.Reset();
