Split remote control module into header and implementation

diff --git a/Engine/Plugins/VirtualProduction/RemoteControl/Source/RemoteControl/Private/RemoteControlModule.cpp b/Engine/Plugins/VirtualProduction/RemoteControl/Source/RemoteControl/Private/RemoteControlModule.cpp
--- a/Engine/Plugins/VirtualProduction/RemoteControl/Source/RemoteControl/Private/RemoteControlModule.cpp
+++ b/Engine/Plugins/VirtualProduction/RemoteControl/Source/RemoteControl/Private/RemoteControlModule.cpp
@@ -1,5 +1,7 @@
 // Copyright Epic Games, Inc. All Rights Reserved.
 
+#include "RemoteControlModule.h"
+
 #include "IRemoteControlInterceptionFeature.h"
 #include "IRemoteControlModule.h"
 #include "IStructDeserializerBackend.h"
@@ -319,1029 +321,970 @@ namespace RemoteControlSetterUtils
 	}
 }
 
-/**
- * Implementation of the RemoteControl interface
- */
-PRAGMA_DISABLE_DEPRECATION_WARNINGS
 
-class FRemoteControlModule : public IRemoteControlModule
+PRAGMA_DISABLE_DEPRECATION_WARNINGS
+void FRemoteControlModule::StartupModule()
 {
-public:
-	virtual void StartupModule() override
-	{
-		IAssetRegistry& AssetRegistry = FModuleManager::LoadModuleChecked<FAssetRegistryModule>(AssetRegistryConstants::ModuleName).Get();
+	IAssetRegistry& AssetRegistry = FModuleManager::LoadModuleChecked<FAssetRegistryModule>(AssetRegistryConstants::ModuleName).Get();
 
-		AssetRegistry.OnAssetAdded().AddRaw(this, &FRemoteControlModule::OnAssetAdded);
-		AssetRegistry.OnAssetRemoved().AddRaw(this, &FRemoteControlModule::OnAssetRemoved);
-		AssetRegistry.OnAssetRenamed().AddRaw(this, &FRemoteControlModule::OnAssetRenamed);
+	AssetRegistry.OnAssetAdded().AddRaw(this, &FRemoteControlModule::OnAssetAdded);
+	AssetRegistry.OnAssetRemoved().AddRaw(this, &FRemoteControlModule::OnAssetRemoved);
+	AssetRegistry.OnAssetRenamed().AddRaw(this, &FRemoteControlModule::OnAssetRenamed);
 
-		// Instantiate the RCI processor feature on module start
-		RCIProcessor = MakeUnique<FRemoteControlInterceptionProcessor>();
-		// Register the interceptor feature
-		IModularFeatures::Get().RegisterModularFeature(IRemoteControlInterceptionFeatureProcessor::GetName(), RCIProcessor.Get());
+	// Instantiate the RCI processor feature on module start
+	RCIProcessor = MakeUnique<FRemoteControlInterceptionProcessor>();
+	// Register the interceptor feature
+	IModularFeatures::Get().RegisterModularFeature(IRemoteControlInterceptionFeatureProcessor::GetName(), RCIProcessor.Get());
 
-		if (AssetRegistry.IsLoadingAssets())
-		{
-			AssetRegistry.OnFilesLoaded().AddRaw(this, &FRemoteControlModule::CachePresets);
-		}
+	if (AssetRegistry.IsLoadingAssets())
+	{
+		AssetRegistry.OnFilesLoaded().AddRaw(this, &FRemoteControlModule::CachePresets);
+	}
 
 #if WITH_EDITOR
-		FCoreDelegates::OnPostEngineInit.AddRaw(this, &FRemoteControlModule::HandleEnginePostInit);
+	FCoreDelegates::OnPostEngineInit.AddRaw(this, &FRemoteControlModule::HandleEnginePostInit);
 #endif
-	}
+}
 
-	virtual void ShutdownModule() override
-	{
+void FRemoteControlModule::ShutdownModule()
+{
 #if WITH_EDITOR
-		UnregisterEditorDelegates();
-		FCoreDelegates::OnPostEngineInit.RemoveAll(this);
+	UnregisterEditorDelegates();
+	FCoreDelegates::OnPostEngineInit.RemoveAll(this);
 #endif
-		if (FModuleManager::Get().IsModuleLoaded(AssetRegistryConstants::ModuleName))
-		{
-			IAssetRegistry& AssetRegistry = FModuleManager::GetModuleChecked<FAssetRegistryModule>(AssetRegistryConstants::ModuleName).Get();
-			AssetRegistry.OnFilesLoaded().RemoveAll(this);
-			AssetRegistry.OnAssetRenamed().RemoveAll(this);
-			AssetRegistry.OnAssetAdded().RemoveAll(this);
-			AssetRegistry.OnAssetRemoved().RemoveAll(this);
-		}
-
-		// Unregister the interceptor feature on module shutdown
-		IModularFeatures::Get().UnregisterModularFeature(IRemoteControlInterceptionFeatureProcessor::GetName(), RCIProcessor.Get());
-	}
-
-	virtual FOnPresetRegistered& OnPresetRegistered() override
+	if (FModuleManager::Get().IsModuleLoaded(AssetRegistryConstants::ModuleName))
 	{
-		return OnPresetRegisteredDelegate;
+		IAssetRegistry& AssetRegistry = FModuleManager::GetModuleChecked<FAssetRegistryModule>(AssetRegistryConstants::ModuleName).Get();
+		AssetRegistry.OnFilesLoaded().RemoveAll(this);
+		AssetRegistry.OnAssetRenamed().RemoveAll(this);
+		AssetRegistry.OnAssetAdded().RemoveAll(this);
+		AssetRegistry.OnAssetRemoved().RemoveAll(this);
 	}
 
-	virtual FOnPresetUnregistered& OnPresetUnregistered() override
-	{
-		return OnPresetUnregisteredDelegate;
-	}
+	// Unregister the interceptor feature on module shutdown
+	IModularFeatures::Get().UnregisterModularFeature(IRemoteControlInterceptionFeatureProcessor::GetName(), RCIProcessor.Get());
+}
 
-	/** Register the preset with the module, enabling using the preset remotely using its name. */
-	virtual bool RegisterPreset(FName Name, URemoteControlPreset* Preset) override
-	{
-		return false;
-	}
+IRemoteControlModule::FOnPresetRegistered& FRemoteControlModule::OnPresetRegistered()
+{
+	return OnPresetRegisteredDelegate;
+}
+
+IRemoteControlModule::FOnPresetUnregistered& FRemoteControlModule::OnPresetUnregistered()
+{
+	return OnPresetUnregisteredDelegate;
+}
 
-	/** Unregister the preset */
-	virtual void UnregisterPreset(FName Name) override
-	{ }
+/** Register the preset with the module, enabling using the preset remotely using its name. */
+bool FRemoteControlModule::RegisterPreset(FName Name, URemoteControlPreset* Preset)
+{
+	return false;
+}
+
+/** Unregister the preset */
+void FRemoteControlModule::UnregisterPreset(FName Name)
+{
+}
 
-	virtual bool ResolveCall(const FString& ObjectPath, const FString& FunctionName, FRCCallReference& OutCallRef, FString* OutErrorText) override
+bool FRemoteControlModule::ResolveCall(const FString& ObjectPath, const FString& FunctionName, FRCCallReference& OutCallRef, FString* OutErrorText)
+{
+	bool bSuccess = true;
+	FString ErrorText;
+	if (!GIsSavingPackage && !IsGarbageCollecting())
 	{
-		bool bSuccess = true;
-		FString ErrorText;
-		if (!GIsSavingPackage && !IsGarbageCollecting())
+		// Resolve the object
+		UObject* Object = StaticFindObject(UObject::StaticClass(), nullptr, *ObjectPath);
+		if (Object)
 		{
-			// Resolve the object
-			UObject* Object = StaticFindObject(UObject::StaticClass(), nullptr, *ObjectPath);
-			if (Object)
-			{
-				// Find the function to call
-				UFunction* Function = RemoteControlUtil::FindFunctionByNameOrMetaDataName(Object, FunctionName);
+			// Find the function to call
+			UFunction* Function = RemoteControlUtil::FindFunctionByNameOrMetaDataName(Object, FunctionName);
 
-				if (!Function)
-				{
-					ErrorText = FString::Printf(TEXT("Function: %s does not exist on object: %s"), *FunctionName, *ObjectPath);
-					bSuccess = false;
-				}
-				else if (!Function->HasAllFunctionFlags(FUNC_BlueprintCallable | FUNC_Public)
+			if (!Function)
+			{
+				ErrorText = FString::Printf(TEXT("Function: %s does not exist on object: %s"), *FunctionName, *ObjectPath);
+				bSuccess = false;
+			}
+			else if (!Function->HasAllFunctionFlags(FUNC_BlueprintCallable | FUNC_Public)
 #if WITH_EDITOR
-						|| Function->HasMetaData(RemoteControlUtil::NAME_DeprecatedFunction)
-						|| Function->HasMetaData(RemoteControlUtil::NAME_ScriptNoExport)
+					|| Function->HasMetaData(RemoteControlUtil::NAME_DeprecatedFunction)
+					|| Function->HasMetaData(RemoteControlUtil::NAME_ScriptNoExport)
 #endif
-				)
-				{
-					ErrorText = FString::Printf(TEXT("Function: %s is deprecated or unavailable remotely on object: %s"), *FunctionName, *ObjectPath);
-					bSuccess = false;
-				}
-				else
-				{
-					OutCallRef.Object = Object;
-					OutCallRef.Function = Function;
-				}
+			)
+			{
+				ErrorText = FString::Printf(TEXT("Function: %s is deprecated or unavailable remotely on object: %s"), *FunctionName, *ObjectPath);
+				bSuccess = false;
 			}
 			else
 			{
-				ErrorText = FString::Printf(TEXT("Object: %s does not exist."), *ObjectPath);
-				bSuccess = false;
+				OutCallRef.Object = Object;
+				OutCallRef.Function = Function;
 			}
 		}
 		else
 		{
-			ErrorText = FString::Printf(TEXT("Can't resolve object: %s while saving or garbage collecting."), *ObjectPath);
+			ErrorText = FString::Printf(TEXT("Object: %s does not exist."), *ObjectPath);
 			bSuccess = false;
 		}
+	}
+	else
+	{
+		ErrorText = FString::Printf(TEXT("Can't resolve object: %s while saving or garbage collecting."), *ObjectPath);
+		bSuccess = false;
+	}
 
-		if (OutErrorText && !ErrorText.IsEmpty())
-		{
-			*OutErrorText = MoveTemp(ErrorText);
-		}
-		return bSuccess;
+	if (OutErrorText && !ErrorText.IsEmpty())
+	{
+		*OutErrorText = MoveTemp(ErrorText);
 	}
+	return bSuccess;
+}
 
-	virtual bool InvokeCall(FRCCall& InCall, ERCPayloadType InPayloadType = ERCPayloadType::Json, const TArray<uint8>& InInterceptPayload = TArray<uint8>()) override
+bool FRemoteControlModule::InvokeCall(FRCCall& InCall, ERCPayloadType InPayloadType, const TArray<uint8>& InInterceptPayload)
+{
+	UE_LOG(LogRemoteControl, VeryVerbose, TEXT("Invoke function"));
+	if (InCall.IsValid())
 	{
-		UE_LOG(LogRemoteControl, VeryVerbose, TEXT("Invoke function"));
-		if (InCall.IsValid())
+		// Check the replication path before apply property values
+		if (InInterceptPayload.Num() != 0)
 		{
-			// Check the replication path before apply property values
-			if (InInterceptPayload.Num() != 0)
-			{
-				FRCIFunctionMetadata FunctionMetadata(InCall.CallRef.Object->GetPathName(), InCall.CallRef.Function->GetPathName(), InCall.bGenerateTransaction, ToExternal(InPayloadType), InInterceptPayload);
+			FRCIFunctionMetadata FunctionMetadata(InCall.CallRef.Object->GetPathName(), InCall.CallRef.Function->GetPathName(), InCall.bGenerateTransaction, ToExternal(InPayloadType), InInterceptPayload);
 
-				// Initialization
-				IModularFeatures& ModularFeatures = IModularFeatures::Get();
-				const FName InterceptorFeatureName = IRemoteControlInterceptionFeatureInterceptor::GetName();
-				const int32 InterceptorsAmount = ModularFeatures.GetModularFeatureImplementationCount(IRemoteControlInterceptionFeatureInterceptor::GetName());
+			// Initialization
+			IModularFeatures& ModularFeatures = IModularFeatures::Get();
+			const FName InterceptorFeatureName = IRemoteControlInterceptionFeatureInterceptor::GetName();
+			const int32 InterceptorsAmount = ModularFeatures.GetModularFeatureImplementationCount(IRemoteControlInterceptionFeatureInterceptor::GetName());
 
-				// Pass interception command data to all available interceptors
-				bool bShouldIntercept = false;
-				for (int32 InterceptorIdx = 0; InterceptorIdx < InterceptorsAmount; ++InterceptorIdx)
+			// Pass interception command data to all available interceptors
+			bool bShouldIntercept = false;
+			for (int32 InterceptorIdx = 0; InterceptorIdx < InterceptorsAmount; ++InterceptorIdx)
+			{
+				IRemoteControlInterceptionFeatureInterceptor* const Interceptor = static_cast<IRemoteControlInterceptionFeatureInterceptor*>(ModularFeatures.GetModularFeatureImplementation(InterceptorFeatureName, InterceptorIdx));
+				if (Interceptor)
 				{
-					IRemoteControlInterceptionFeatureInterceptor* const Interceptor = static_cast<IRemoteControlInterceptionFeatureInterceptor*>(ModularFeatures.GetModularFeatureImplementation(InterceptorFeatureName, InterceptorIdx));
-					if (Interceptor)
-					{
-						// Update response flag
-						UE_LOG(LogRemoteControl, VeryVerbose, TEXT("Invoke function - Intercepted"));
-						bShouldIntercept |= (Interceptor->InvokeCall(FunctionMetadata) == ERCIResponse::Intercept);
-					}
+					// Update response flag
+					UE_LOG(LogRemoteControl, VeryVerbose, TEXT("Invoke function - Intercepted"));
+					bShouldIntercept |= (Interceptor->InvokeCall(FunctionMetadata) == ERCIResponse::Intercept);
 				}
+			}
 
-				// Don't process the RC message if any of interceptors returned ERCIResponse::Intercept
-				if (bShouldIntercept)
-				{
-					return true;
-				}
+			// Don't process the RC message if any of interceptors returned ERCIResponse::Intercept
+			if (bShouldIntercept)
+			{
+				return true;
 			}
+		}
 
 #if WITH_EDITOR
-			if (CVarRemoteControlEnableOngoingChangeOptimization.GetValueOnAnyThread() == 1)
+		if (CVarRemoteControlEnableOngoingChangeOptimization.GetValueOnAnyThread() == 1)
+		{
+			// If we have a different ongoing change that hasn't been finalized, do that before handling the next function call.
+			if (OngoingModification && GetTypeHash(*OngoingModification) != GetTypeHash(InCall.CallRef))
 			{
-				// If we have a different ongoing change that hasn't been finalized, do that before handling the next function call.
-				if (OngoingModification && GetTypeHash(*OngoingModification) != GetTypeHash(InCall.CallRef))
-				{
-					constexpr bool bForceFinalizeChange = true;
-					TestOrFinalizeOngoingChange(true);
-				}
+				constexpr bool bForceFinalizeChange = true;
+				TestOrFinalizeOngoingChange(true);
+			}
 				
-				if (!OngoingModification)
+			if (!OngoingModification)
+			{
+				if (InCall.bGenerateTransaction)
 				{
-					if (InCall.bGenerateTransaction)
+					if (GEditor)
 					{
-						if (GEditor)
-						{
-							GEditor->BeginTransaction(LOCTEXT("RemoteCallTransaction", "Remote Call Transaction Wrap"));
-						}
-
-						if (ensureAlways(InCall.CallRef.Object.IsValid()))
-						{
-						InCall.CallRef.Object->Modify();
+						GEditor->BeginTransaction(LOCTEXT("RemoteCallTransaction", "Remote Call Transaction Wrap"));
 					}
+
+					if (ensureAlways(InCall.CallRef.Object.IsValid()))
+					{
+					InCall.CallRef.Object->Modify();
 				}
 			}
-			}
-			else if (GEditor && InCall.bGenerateTransaction)
-			{
-				GEditor->BeginTransaction(LOCTEXT("RemoteCallTransaction", "Remote Call Transaction Wrap"));
-			}
+		}
+		}
+		else if (GEditor && InCall.bGenerateTransaction)
+		{
+			GEditor->BeginTransaction(LOCTEXT("RemoteCallTransaction", "Remote Call Transaction Wrap"));
+		}
 			
 #endif
-			FEditorScriptExecutionGuard ScriptGuard;
-			if (ensureAlways(InCall.CallRef.Object.IsValid()))
-			{
-			InCall.CallRef.Object->ProcessEvent(InCall.CallRef.Function.Get(), InCall.ParamStruct.GetStructMemory());
-			}
+		FEditorScriptExecutionGuard ScriptGuard;
+		if (ensureAlways(InCall.CallRef.Object.IsValid()))
+		{
+		InCall.CallRef.Object->ProcessEvent(InCall.CallRef.Function.Get(), InCall.ParamStruct.GetStructMemory());
+		}
 
 #if WITH_EDITOR
-			if (CVarRemoteControlEnableOngoingChangeOptimization.GetValueOnAnyThread() == 1)
+		if (CVarRemoteControlEnableOngoingChangeOptimization.GetValueOnAnyThread() == 1)
+		{
+			// If we've called the same function recently, refresh the triggered flag and snapshot the object to the transaction buffer
+			if (OngoingModification && GetTypeHash(*OngoingModification) == GetTypeHash(InCall.CallRef))
 			{
-				// If we've called the same function recently, refresh the triggered flag and snapshot the object to the transaction buffer
-				if (OngoingModification && GetTypeHash(*OngoingModification) == GetTypeHash(InCall.CallRef))
+				OngoingModification->bWasTriggeredSinceLastPass = true;
+				if (OngoingModification->bHasStartedTransaction)
 				{
-					OngoingModification->bWasTriggeredSinceLastPass = true;
-					if (OngoingModification->bHasStartedTransaction)
-					{
-						SnapshotTransactionBuffer(InCall.CallRef.Object.Get());
-					}
-				}
-				else
-				{
-					OngoingModification = InCall.CallRef;
-					OngoingModification->bHasStartedTransaction = InCall.bGenerateTransaction;
+					SnapshotTransactionBuffer(InCall.CallRef.Object.Get());
 				}
 			}
-			else if (GEditor && InCall.bGenerateTransaction)
+			else
 			{
-				GEditor->EndTransaction();
+				OngoingModification = InCall.CallRef;
+				OngoingModification->bHasStartedTransaction = InCall.bGenerateTransaction;
 			}
-#endif
-			return true;
 		}
-		return false;
+		else if (GEditor && InCall.bGenerateTransaction)
+		{
+			GEditor->EndTransaction();
+		}
+#endif
+		return true;
 	}
+	return false;
+}
 
-	virtual bool ResolveObject(ERCAccess AccessType, const FString& ObjectPath, const FString& PropertyName, FRCObjectReference& OutObjectRef, FString* OutErrorText = nullptr) override
+bool FRemoteControlModule::ResolveObject(ERCAccess AccessType, const FString& ObjectPath, const FString& PropertyName, FRCObjectReference& OutObjectRef, FString* OutErrorText)
+{
+	bool bSuccess = true;
+	FString ErrorText;
+	if (!GIsSavingPackage && !IsGarbageCollecting())
 	{
-		bool bSuccess = true;
-		FString ErrorText;
-		if (!GIsSavingPackage && !IsGarbageCollecting())
+		// Resolve the object
+		UObject* Object = StaticFindObject(UObject::StaticClass(), nullptr, *ObjectPath);
+		if (Object)
 		{
-			// Resolve the object
-			UObject* Object = StaticFindObject(UObject::StaticClass(), nullptr, *ObjectPath);
-			if (Object)
-			{
-				bSuccess = ResolveObjectProperty(AccessType, Object, PropertyName, OutObjectRef, OutErrorText);
-			}
-			else
-			{
-				ErrorText = FString::Printf(TEXT("Object: %s does not exist when trying to resolve property: %s"), *ObjectPath, *PropertyName);
-				bSuccess = false;
-			}
+			bSuccess = ResolveObjectProperty(AccessType, Object, PropertyName, OutObjectRef, OutErrorText);
 		}
 		else
 		{
-			ErrorText = FString::Printf(TEXT("Can't resolve object: %s while saving or garbage collecting."), *ObjectPath);
+			ErrorText = FString::Printf(TEXT("Object: %s does not exist when trying to resolve property: %s"), *ObjectPath, *PropertyName);
 			bSuccess = false;
 		}
+	}
+	else
+	{
+		ErrorText = FString::Printf(TEXT("Can't resolve object: %s while saving or garbage collecting."), *ObjectPath);
+		bSuccess = false;
+	}
 
-		if (OutErrorText && !ErrorText.IsEmpty())
-		{
-			*OutErrorText = MoveTemp(ErrorText);
-		}
-
-		return bSuccess;
+	if (OutErrorText && !ErrorText.IsEmpty())
+	{
+		*OutErrorText = MoveTemp(ErrorText);
 	}
 
-	virtual bool ResolveObjectProperty(ERCAccess AccessType, UObject* Object, FRCFieldPathInfo PropertyPath, FRCObjectReference& OutObjectRef, FString* OutErrorText = nullptr) override
+	return bSuccess;
+}
+
+bool FRemoteControlModule::ResolveObjectProperty(ERCAccess AccessType, UObject* Object, FRCFieldPathInfo PropertyPath, FRCObjectReference& OutObjectRef, FString* OutErrorText)
+{
+	TRACE_CPUPROFILER_EVENT_SCOPE(FRemoteControlModule::ResolveObjectProperty);
+	bool bSuccess = true;
+	FString ErrorText;
+	if (!GIsSavingPackage && !IsGarbageCollecting())
 	{
-		TRACE_CPUPROFILER_EVENT_SCOPE(FRemoteControlModule::ResolveObjectProperty);
-		bool bSuccess = true;
-		FString ErrorText;
-		if (!GIsSavingPackage && !IsGarbageCollecting())
+		if (Object)
 		{
-			if (Object)
-			{
-				const bool bObjectInGame = !GIsEditor || IsRunningGame();
+			const bool bObjectInGame = !GIsEditor || IsRunningGame();
 
-				if (PropertyPath.GetSegmentCount() != 0)
+			if (PropertyPath.GetSegmentCount() != 0)
+			{
+				//Build a FieldPathInfo using property name to facilitate resolving
+				if (PropertyPath.Resolve(Object))
 				{
-					//Build a FieldPathInfo using property name to facilitate resolving
-					if (PropertyPath.Resolve(Object))
-					{
-						FProperty* ResolvedProperty = PropertyPath.GetResolvedData().Field;
+					FProperty* ResolvedProperty = PropertyPath.GetResolvedData().Field;
 
-						// When resolving a property for writing, resolve successfully if it should use a setter since it will end up using it. 
-						if ((RemoteControlUtil::IsWriteAccess(AccessType) && PropertyModificationShouldUseSetter(Object, ResolvedProperty))
-							|| RemoteControlUtil::IsPropertyAllowed(ResolvedProperty, AccessType, bObjectInGame))
-						{
-							OutObjectRef = FRCObjectReference{AccessType, Object, MoveTemp(PropertyPath)};
-						}
-						else
-						{
-							ErrorText = FString::Printf(TEXT("Object property: %s is unavailable remotely on object: %s"), *PropertyPath.GetFieldName().ToString(), *Object->GetPathName());
-							bSuccess = false;
-						}
+					// When resolving a property for writing, resolve successfully if it should use a setter since it will end up using it. 
+					if ((RemoteControlUtil::IsWriteAccess(AccessType) && PropertyModificationShouldUseSetter(Object, ResolvedProperty))
+						|| RemoteControlUtil::IsPropertyAllowed(ResolvedProperty, AccessType, bObjectInGame))
+					{
+						OutObjectRef = FRCObjectReference{AccessType, Object, MoveTemp(PropertyPath)};
 					}
 					else
 					{
-						ErrorText = FString::Printf(TEXT("Object property: %s could not be resolved on object: %s"), *PropertyPath.GetFieldName().ToString(), *Object->GetPathName());
+						ErrorText = FString::Printf(TEXT("Object property: %s is unavailable remotely on object: %s"), *PropertyPath.GetFieldName().ToString(), *Object->GetPathName());
 						bSuccess = false;
 					}
 				}
 				else
 				{
-					OutObjectRef = FRCObjectReference{AccessType, Object};
+					ErrorText = FString::Printf(TEXT("Object property: %s could not be resolved on object: %s"), *PropertyPath.GetFieldName().ToString(), *Object->GetPathName());
+					bSuccess = false;
 				}
 			}
 			else
 			{
-				ErrorText = FString::Printf(TEXT("Invalid object to resolve property '%s'"), *PropertyPath.GetFieldName().ToString());
-				bSuccess = false;
+				OutObjectRef = FRCObjectReference{AccessType, Object};
 			}
 		}
 		else
 		{
-			ErrorText = FString::Printf(TEXT("Can't resolve object '%s' properties '%s' : %s while saving or garbage collecting."), *Object->GetPathName(), *PropertyPath.GetFieldName().ToString());
+			ErrorText = FString::Printf(TEXT("Invalid object to resolve property '%s'"), *PropertyPath.GetFieldName().ToString());
 			bSuccess = false;
 		}
+	}
+	else
+	{
+		ErrorText = FString::Printf(TEXT("Can't resolve object '%s' properties '%s' : %s while saving or garbage collecting."), *Object->GetPathName(), *PropertyPath.GetFieldName().ToString());
+		bSuccess = false;
+	}
 
-		if (OutErrorText && !ErrorText.IsEmpty())
-		{
-			*OutErrorText = MoveTemp(ErrorText);
-		}
-
-		return bSuccess;
+	if (OutErrorText && !ErrorText.IsEmpty())
+	{
+		*OutErrorText = MoveTemp(ErrorText);
 	}
 
-	virtual bool GetObjectProperties(const FRCObjectReference& ObjectAccess, IStructSerializerBackend& Backend) override
+	return bSuccess;
+}
+
+bool FRemoteControlModule::GetObjectProperties(const FRCObjectReference& ObjectAccess, IStructSerializerBackend& Backend)
+{
+	if (ObjectAccess.IsValid() && ObjectAccess.Access == ERCAccess::READ_ACCESS)
 	{
-		if (ObjectAccess.IsValid() && ObjectAccess.Access == ERCAccess::READ_ACCESS)
-		{
-			bool bCanSerialize = true;
-			UObject* Object = ObjectAccess.Object.Get();
-			UStruct* ContainerType = ObjectAccess.ContainerType.Get();
+		bool bCanSerialize = true;
+		UObject* Object = ObjectAccess.Object.Get();
+		UStruct* ContainerType = ObjectAccess.ContainerType.Get();
 
-			FStructSerializerPolicies Policies;
-			if (ObjectAccess.Property.IsValid())
+		FStructSerializerPolicies Policies;
+		if (ObjectAccess.Property.IsValid())
+		{
+			if (ObjectAccess.PropertyPathInfo.IsResolved())
 			{
-				if (ObjectAccess.PropertyPathInfo.IsResolved())
-				{
-					Policies.MapSerialization = EStructSerializerMapPolicies::Array;
-					Policies.PropertyFilter = [&ObjectAccess](const FProperty* CurrentProp, const FProperty* ParentProp)
-					{
-						return CurrentProp == ObjectAccess.Property || ParentProp != nullptr;
-					};
-				}
-				else
+				Policies.MapSerialization = EStructSerializerMapPolicies::Array;
+				Policies.PropertyFilter = [&ObjectAccess](const FProperty* CurrentProp, const FProperty* ParentProp)
 				{
-					bCanSerialize = false;
-				}
+					return CurrentProp == ObjectAccess.Property || ParentProp != nullptr;
+				};
 			}
 			else
 			{
-				bool bObjectInGame = !GIsEditor || Object->GetOutermost()->HasAnyPackageFlags(PKG_PlayInEditor);
-				Policies.PropertyFilter = [&ObjectAccess, bObjectInGame](const FProperty* CurrentProp, const FProperty* ParentProp)
-				{
-					return RemoteControlUtil::IsPropertyAllowed(CurrentProp, ObjectAccess.Access, bObjectInGame) || ParentProp != nullptr;
-				};
+				bCanSerialize = false;
 			}
+		}
+		else
+		{
+			bool bObjectInGame = !GIsEditor || Object->GetOutermost()->HasAnyPackageFlags(PKG_PlayInEditor);
+			Policies.PropertyFilter = [&ObjectAccess, bObjectInGame](const FProperty* CurrentProp, const FProperty* ParentProp)
+			{
+				return RemoteControlUtil::IsPropertyAllowed(CurrentProp, ObjectAccess.Access, bObjectInGame) || ParentProp != nullptr;
+			};
+		}
 
-			if (bCanSerialize)
+		if (bCanSerialize)
+		{
+			//Serialize the element if we're looking for a member or serialize the full object if not
+			if (ObjectAccess.PropertyPathInfo.IsResolved())
 			{
-				//Serialize the element if we're looking for a member or serialize the full object if not
-				if (ObjectAccess.PropertyPathInfo.IsResolved())
-				{
-					const FRCFieldPathSegment& LastSegment = ObjectAccess.PropertyPathInfo.GetFieldSegment(ObjectAccess.PropertyPathInfo.GetSegmentCount() - 1);
-					int32 Index = LastSegment.ArrayIndex != INDEX_NONE ? LastSegment.ArrayIndex : LastSegment.ResolvedData.MapIndex;
+				const FRCFieldPathSegment& LastSegment = ObjectAccess.PropertyPathInfo.GetFieldSegment(ObjectAccess.PropertyPathInfo.GetSegmentCount() - 1);
+				int32 Index = LastSegment.ArrayIndex != INDEX_NONE ? LastSegment.ArrayIndex : LastSegment.ResolvedData.MapIndex;
 
-					FStructSerializer::SerializeElement(ObjectAccess.ContainerAdress, LastSegment.ResolvedData.Field, Index, Backend, Policies);
-				}
-				else
-				{
-					FStructSerializer::Serialize(ObjectAccess.ContainerAdress, *ObjectAccess.ContainerType, Backend, Policies);
-				}
-				return true;
+				FStructSerializer::SerializeElement(ObjectAccess.ContainerAdress, LastSegment.ResolvedData.Field, Index, Backend, Policies);
+			}
+			else
+			{
+				FStructSerializer::Serialize(ObjectAccess.ContainerAdress, *ObjectAccess.ContainerType, Backend, Policies);
 			}
+			return true;
 		}
-		return false;
 	}
+	return false;
+}
 
-	virtual bool SetObjectProperties(const FRCObjectReference& ObjectAccess, IStructDeserializerBackend& Backend, ERCPayloadType InPayloadType, const TArray<uint8>& InPayload) override
+bool FRemoteControlModule::SetObjectProperties(const FRCObjectReference& ObjectAccess, IStructDeserializerBackend& Backend, ERCPayloadType InPayloadType, const TArray<uint8>& InPayload)
+{
+	TRACE_CPUPROFILER_EVENT_SCOPE(FRemoteControlModule::SetObjectProperties);
+	UE_LOG(LogRemoteControl, VeryVerbose, TEXT("Set Object Properties"));
+	// Check the replication path before applying property values
+	if (InPayload.Num() != 0 && ObjectAccess.Object.IsValid())
 	{
-		TRACE_CPUPROFILER_EVENT_SCOPE(FRemoteControlModule::SetObjectProperties);
-		UE_LOG(LogRemoteControl, VeryVerbose, TEXT("Set Object Properties"));
-		// Check the replication path before applying property values
-		if (InPayload.Num() != 0 && ObjectAccess.Object.IsValid())
+		// Convert raw property modifications to setter function calls if necessary.
+		if (PropertyModificationShouldUseSetter(ObjectAccess.Object.Get(), ObjectAccess.Property.Get()))
 		{
-			// Convert raw property modifications to setter function calls if necessary.
-			if (PropertyModificationShouldUseSetter(ObjectAccess.Object.Get(), ObjectAccess.Property.Get()))
-			{
-				FRCCall Call;
-				FRCInterceptionPayload InterceptionPayload;
-				constexpr bool bCreateInterceptionPayload = true;
-				RemoteControlSetterUtils::FConvertToFunctionCallArgs Args(ObjectAccess, Backend, Call);
-				if (RemoteControlSetterUtils::ConvertModificationToFunctionCall(Args, InterceptionPayload))
-				{
-					const bool bResult = InvokeCall(Call, InterceptionPayload.Type, InterceptionPayload.Payload);
-					return bResult;
-				}
-			}
-
-			// If a setter wasn't used, verify if the property should be allowed.
-			bool bObjectInGame = !GIsEditor || (ObjectAccess.Object.IsValid() && ObjectAccess.Object->GetOutermost()->HasAnyPackageFlags(PKG_PlayInEditor));
-			if (!RemoteControlUtil::IsPropertyAllowed(ObjectAccess.Property.Get(), ObjectAccess.Access, bObjectInGame))
+			FRCCall Call;
+			FRCInterceptionPayload InterceptionPayload;
+			constexpr bool bCreateInterceptionPayload = true;
+			RemoteControlSetterUtils::FConvertToFunctionCallArgs Args(ObjectAccess, Backend, Call);
+			if (RemoteControlSetterUtils::ConvertModificationToFunctionCall(Args, InterceptionPayload))
 			{
-				return false;
+				const bool bResult = InvokeCall(Call, InterceptionPayload.Type, InterceptionPayload.Payload);
+				return bResult;
 			}
+		}
 
-			// Build interception command
-			FString PropertyPathString = TFieldPath<FProperty>(ObjectAccess.Property.Get()).ToString();
-			FRCIPropertiesMetadata PropsMetadata(ObjectAccess.Object->GetPathName(), PropertyPathString, ObjectAccess.PropertyPathInfo.ToString(), ToExternal(ObjectAccess.Access), ToExternal(InPayloadType), InPayload);
+		// If a setter wasn't used, verify if the property should be allowed.
+		bool bObjectInGame = !GIsEditor || (ObjectAccess.Object.IsValid() && ObjectAccess.Object->GetOutermost()->HasAnyPackageFlags(PKG_PlayInEditor));
+		if (!RemoteControlUtil::IsPropertyAllowed(ObjectAccess.Property.Get(), ObjectAccess.Access, bObjectInGame))
+		{
+			return false;
+		}
 
-			// Initialization
-			IModularFeatures& ModularFeatures = IModularFeatures::Get();
-			const FName InterceptorFeatureName = IRemoteControlInterceptionFeatureInterceptor::GetName();
-			const int32 InterceptorsAmount = ModularFeatures.GetModularFeatureImplementationCount(IRemoteControlInterceptionFeatureInterceptor::GetName());
+		// Build interception command
+		FString PropertyPathString = TFieldPath<FProperty>(ObjectAccess.Property.Get()).ToString();
+		FRCIPropertiesMetadata PropsMetadata(ObjectAccess.Object->GetPathName(), PropertyPathString, ObjectAccess.PropertyPathInfo.ToString(), ToExternal(ObjectAccess.Access), ToExternal(InPayloadType), InPayload);
 
-			// Pass interception command data to all available interceptors
-			bool bShouldIntercept = false;
-			for (int32 InterceptorIdx = 0; InterceptorIdx < InterceptorsAmount; ++InterceptorIdx)
-			{
-				IRemoteControlInterceptionFeatureInterceptor* const Interceptor = static_cast<IRemoteControlInterceptionFeatureInterceptor*>(ModularFeatures.GetModularFeatureImplementation(InterceptorFeatureName, InterceptorIdx));
-				if (Interceptor)
-				{
-					// Update response flag
-					UE_LOG(LogRemoteControl, VeryVerbose, TEXT("Set Object Properties - Intercepted"));
-					bShouldIntercept |= (Interceptor->SetObjectProperties(PropsMetadata) == ERCIResponse::Intercept);
-				}
-			}
+		// Initialization
+		IModularFeatures& ModularFeatures = IModularFeatures::Get();
+		const FName InterceptorFeatureName = IRemoteControlInterceptionFeatureInterceptor::GetName();
+		const int32 InterceptorsAmount = ModularFeatures.GetModularFeatureImplementationCount(IRemoteControlInterceptionFeatureInterceptor::GetName());
 
-			// Don't process the RC message if any of interceptors returned ERCIResponse::Intercept
-			if (bShouldIntercept)
+		// Pass interception command data to all available interceptors
+		bool bShouldIntercept = false;
+		for (int32 InterceptorIdx = 0; InterceptorIdx < InterceptorsAmount; ++InterceptorIdx)
+		{
+			IRemoteControlInterceptionFeatureInterceptor* const Interceptor = static_cast<IRemoteControlInterceptionFeatureInterceptor*>(ModularFeatures.GetModularFeatureImplementation(InterceptorFeatureName, InterceptorIdx));
+			if (Interceptor)
 			{
-				return true;
+				// Update response flag
+				UE_LOG(LogRemoteControl, VeryVerbose, TEXT("Set Object Properties - Intercepted"));
+				bShouldIntercept |= (Interceptor->SetObjectProperties(PropsMetadata) == ERCIResponse::Intercept);
 			}
 		}
 
-		// Convert raw property modifications to setter function calls if necessary.
-		if (PropertyModificationShouldUseSetter(ObjectAccess.Object.Get(), ObjectAccess.Property.Get()))
+		// Don't process the RC message if any of interceptors returned ERCIResponse::Intercept
+		if (bShouldIntercept)
 		{
-			FRCCall Call;
-			RemoteControlSetterUtils::FConvertToFunctionCallArgs Args(ObjectAccess, Backend, Call);
-			if (RemoteControlSetterUtils::ConvertModificationToFunctionCall(Args))
-			{
-				const bool bResult = InvokeCall(Call);
-				return bResult;
-			}
+			return true;
 		}
+	}
 
-		// Setting object properties require a property and can't be done at the object level. Property must be valid to move forward
-		if (ObjectAccess.IsValid()
-			&& (RemoteControlUtil::IsWriteAccess(ObjectAccess.Access))
-			&& ObjectAccess.Property.IsValid()
-			&& ObjectAccess.PropertyPathInfo.IsResolved())
+	// Convert raw property modifications to setter function calls if necessary.
+	if (PropertyModificationShouldUseSetter(ObjectAccess.Object.Get(), ObjectAccess.Property.Get()))
+	{
+		FRCCall Call;
+		RemoteControlSetterUtils::FConvertToFunctionCallArgs Args(ObjectAccess, Backend, Call);
+		if (RemoteControlSetterUtils::ConvertModificationToFunctionCall(Args))
 		{
-			UStruct* ContainerType = ObjectAccess.ContainerType.Get();
-			FRCObjectReference MutableObjectReference = ObjectAccess;
+			const bool bResult = InvokeCall(Call);
+			return bResult;
+		}
+	}
+
+	// Setting object properties require a property and can't be done at the object level. Property must be valid to move forward
+	if (ObjectAccess.IsValid()
+		&& (RemoteControlUtil::IsWriteAccess(ObjectAccess.Access))
+		&& ObjectAccess.Property.IsValid()
+		&& ObjectAccess.PropertyPathInfo.IsResolved())
+	{
+		UStruct* ContainerType = ObjectAccess.ContainerType.Get();
+		FRCObjectReference MutableObjectReference = ObjectAccess;
 
 #if WITH_EDITOR
-			const bool bGenerateTransaction = MutableObjectReference.Access == ERCAccess::WRITE_TRANSACTION_ACCESS;
-			if (CVarRemoteControlEnableOngoingChangeOptimization.GetValueOnAnyThread() == 1)
-			{
-				FString ObjectPath = MutableObjectReference.Object->GetPathName();
+		const bool bGenerateTransaction = MutableObjectReference.Access == ERCAccess::WRITE_TRANSACTION_ACCESS;
+		if (CVarRemoteControlEnableOngoingChangeOptimization.GetValueOnAnyThread() == 1)
+		{
+			FString ObjectPath = MutableObjectReference.Object->GetPathName();
 
-				// If we have a change that hasn't yet generated a post edit change property, do that before handling the next change.
-				if (OngoingModification && GetTypeHash(*OngoingModification) != GetTypeHash(MutableObjectReference))
-				{
-					constexpr bool bForcePostEditChange = true;
-					TestOrFinalizeOngoingChange(true);
-				}
+			// If we have a change that hasn't yet generated a post edit change property, do that before handling the next change.
+			if (OngoingModification && GetTypeHash(*OngoingModification) != GetTypeHash(MutableObjectReference))
+			{
+				constexpr bool bForcePostEditChange = true;
+				TestOrFinalizeOngoingChange(true);
+			}
 
-				// This step is necessary because the object might get recreated by a PostEditChange called in TestOrFinalizeOngoingChange.
-				if (!MutableObjectReference.Object.IsValid())
+			// This step is necessary because the object might get recreated by a PostEditChange called in TestOrFinalizeOngoingChange.
+			if (!MutableObjectReference.Object.IsValid())
+			{
+				MutableObjectReference.Object = StaticFindObject(UObject::StaticClass(), nullptr, *ObjectPath);
+				if (MutableObjectReference.Object.IsValid() && MutableObjectReference.PropertyPathInfo.IsResolved())
 				{
-					MutableObjectReference.Object = StaticFindObject(UObject::StaticClass(), nullptr, *ObjectPath);
-					if (MutableObjectReference.Object.IsValid() && MutableObjectReference.PropertyPathInfo.IsResolved())
-					{
-						// Update ContainerAddress as well if the path was resolved.
-						if (MutableObjectReference.PropertyPathInfo.Resolve(MutableObjectReference.Object.Get()))
-						{
-							MutableObjectReference.ContainerAdress = MutableObjectReference.PropertyPathInfo.GetResolvedData().ContainerAddress;
-						}
-					}
-					else
+					// Update ContainerAddress as well if the path was resolved.
+					if (MutableObjectReference.PropertyPathInfo.Resolve(MutableObjectReference.Object.Get()))
 					{
-						return false;
+						MutableObjectReference.ContainerAdress = MutableObjectReference.PropertyPathInfo.GetResolvedData().ContainerAddress;
 					}
 				}
-
-				// Only create the transaction if we have no ongoing change.
-				if (!OngoingModification)
+				else
 				{
-					if (GEditor && bGenerateTransaction)
-					{
-						GEditor->BeginTransaction(LOCTEXT("RemoteSetPropertyTransaction", "Remote Set Object Property"));
-
-						// Call modify since it's not called by PreEditChange until the end of the ongoing change.
-						MutableObjectReference.Object->Modify();
-					}
+					return false;
 				}
-
 			}
-			else 
+
+			// Only create the transaction if we have no ongoing change.
+			if (!OngoingModification)
 			{
 				if (GEditor && bGenerateTransaction)
 				{
 					GEditor->BeginTransaction(LOCTEXT("RemoteSetPropertyTransaction", "Remote Set Object Property"));
-				}
 
-				FEditPropertyChain PreEditChain;
-				MutableObjectReference.PropertyPathInfo.ToEditPropertyChain(PreEditChain);
-				MutableObjectReference.Object->PreEditChange(PreEditChain);
-				
+					// Call modify since it's not called by PreEditChange until the end of the ongoing change.
+					MutableObjectReference.Object->Modify();
+				}
 			}
-#endif
 
-			FStructDeserializerPolicies Policies;
-			if (MutableObjectReference.Property.IsValid())
-			{
-				Policies.PropertyFilter = [&MutableObjectReference](const FProperty* CurrentProp, const FProperty* ParentProp)
-				{
-					return CurrentProp == MutableObjectReference.Property || ParentProp != nullptr;
-				};
-			}
-			else
+		}
+		else 
+		{
+			if (GEditor && bGenerateTransaction)
 			{
-				bool bObjectInGame = !GIsEditor || MutableObjectReference.Object->GetOutermost()->HasAnyPackageFlags(PKG_PlayInEditor);
-				Policies.PropertyFilter = [&MutableObjectReference, bObjectInGame](const FProperty* CurrentProp, const FProperty* ParentProp)
-				{
-					return RemoteControlUtil::IsPropertyAllowed(CurrentProp, MutableObjectReference.Access, bObjectInGame) || ParentProp != nullptr;
-				};
+				GEditor->BeginTransaction(LOCTEXT("RemoteSetPropertyTransaction", "Remote Set Object Property"));
 			}
 
-			//Serialize the element if we're looking for a member or serialize the full object if not
-			bool bSuccess = false;
-			if (MutableObjectReference.PropertyPathInfo.IsResolved())
+			FEditPropertyChain PreEditChain;
+			MutableObjectReference.PropertyPathInfo.ToEditPropertyChain(PreEditChain);
+			MutableObjectReference.Object->PreEditChange(PreEditChain);
+				
+		}
+#endif
+
+		FStructDeserializerPolicies Policies;
+		if (MutableObjectReference.Property.IsValid())
+		{
+			Policies.PropertyFilter = [&MutableObjectReference](const FProperty* CurrentProp, const FProperty* ParentProp)
 			{
-				const FRCFieldPathSegment& LastSegment = MutableObjectReference.PropertyPathInfo.GetFieldSegment(MutableObjectReference.PropertyPathInfo.GetSegmentCount() - 1);
-				int32 Index = LastSegment.ArrayIndex != INDEX_NONE ? LastSegment.ArrayIndex : LastSegment.ResolvedData.MapIndex;
-				bSuccess = FStructDeserializer::DeserializeElement(MutableObjectReference.ContainerAdress, *LastSegment.ResolvedData.Struct, Index, Backend, Policies);
-			}
-			else
+				return CurrentProp == MutableObjectReference.Property || ParentProp != nullptr;
+			};
+		}
+		else
+		{
+			bool bObjectInGame = !GIsEditor || MutableObjectReference.Object->GetOutermost()->HasAnyPackageFlags(PKG_PlayInEditor);
+			Policies.PropertyFilter = [&MutableObjectReference, bObjectInGame](const FProperty* CurrentProp, const FProperty* ParentProp)
 			{
-				bSuccess = FStructDeserializer::Deserialize(MutableObjectReference.ContainerAdress, *ContainerType, Backend, Policies);
-			}
+				return RemoteControlUtil::IsPropertyAllowed(CurrentProp, MutableObjectReference.Access, bObjectInGame) || ParentProp != nullptr;
+			};
+		}
+
+		//Serialize the element if we're looking for a member or serialize the full object if not
+		bool bSuccess = false;
+		if (MutableObjectReference.PropertyPathInfo.IsResolved())
+		{
+			const FRCFieldPathSegment& LastSegment = MutableObjectReference.PropertyPathInfo.GetFieldSegment(MutableObjectReference.PropertyPathInfo.GetSegmentCount() - 1);
+			int32 Index = LastSegment.ArrayIndex != INDEX_NONE ? LastSegment.ArrayIndex : LastSegment.ResolvedData.MapIndex;
+			bSuccess = FStructDeserializer::DeserializeElement(MutableObjectReference.ContainerAdress, *LastSegment.ResolvedData.Struct, Index, Backend, Policies);
+		}
+		else
+		{
+			bSuccess = FStructDeserializer::Deserialize(MutableObjectReference.ContainerAdress, *ContainerType, Backend, Policies);
+		}
 
 #if WITH_EDITOR
-			if (CVarRemoteControlEnableOngoingChangeOptimization.GetValueOnAnyThread() == 1)
+		if (CVarRemoteControlEnableOngoingChangeOptimization.GetValueOnAnyThread() == 1)
+		{
+			// If we have modified the same object and property in the last frames,
+			// update the triggered flag and snapshot the object to the transaction buffer.
+			if (OngoingModification && GetTypeHash(*OngoingModification) == GetTypeHash(MutableObjectReference))
 			{
-				// If we have modified the same object and property in the last frames,
-				// update the triggered flag and snapshot the object to the transaction buffer.
-				if (OngoingModification && GetTypeHash(*OngoingModification) == GetTypeHash(MutableObjectReference))
+				OngoingModification->bWasTriggeredSinceLastPass = true;
+				if (OngoingModification->bHasStartedTransaction)
 				{
-					OngoingModification->bWasTriggeredSinceLastPass = true;
-					if (OngoingModification->bHasStartedTransaction)
-					{
-						SnapshotTransactionBuffer(MutableObjectReference.Object.Get());
-					}
+					SnapshotTransactionBuffer(MutableObjectReference.Object.Get());
+				}
 
-					// Update the world lighting if we're modifying a color.
-					if (MutableObjectReference.IsValid() && MutableObjectReference.Property->GetOwnerClass()->IsChildOf(ULightComponentBase::StaticClass()))
+				// Update the world lighting if we're modifying a color.
+				if (MutableObjectReference.IsValid() && MutableObjectReference.Property->GetOwnerClass()->IsChildOf(ULightComponentBase::StaticClass()))
+				{
+					UWorld* World = MutableObjectReference.Object->GetWorld();
+					if (World && World->Scene)
 					{
-						UWorld* World = MutableObjectReference.Object->GetWorld();
-						if (World && World->Scene)
+						if (MutableObjectReference.Object->IsA<ULightComponent>())
 						{
-							if (MutableObjectReference.Object->IsA<ULightComponent>())
-							{
-								World->Scene->UpdateLightColorAndBrightness(Cast<ULightComponent>(MutableObjectReference.Object.Get()));
-							}
-					}
-				}
-				}
-				else
-				{
-					OngoingModification = MutableObjectReference;
-					OngoingModification->bHasStartedTransaction = bGenerateTransaction;
+							World->Scene->UpdateLightColorAndBrightness(Cast<ULightComponent>(MutableObjectReference.Object.Get()));
+						}
 				}
 			}
+			}
 			else
 			{
-				if (GEditor && bGenerateTransaction)
-				{
-					GEditor->EndTransaction();
-				}
-
-				FPropertyChangedEvent PropertyEvent(MutableObjectReference.PropertyPathInfo.ToPropertyChangedEvent());
-				MutableObjectReference.Object->PostEditChangeProperty(PropertyEvent);
+				OngoingModification = MutableObjectReference;
+				OngoingModification->bHasStartedTransaction = bGenerateTransaction;
 			}
-#endif
-			return bSuccess;
 		}
-		return false;
-	}
-
-	virtual bool ResetObjectProperties(const FRCObjectReference& ObjectAccess, const bool bAllowIntercept) override
-	{
-		// Check the replication path before reset the property on this instance
-		if (bAllowIntercept && ObjectAccess.Object.IsValid())
+		else
 		{
-			// Build interception command
-			FString PropertyPathString = TFieldPath<FProperty>(ObjectAccess.Property.Get()).ToString();
-			FRCIObjectMetadata ObjectMetadata(ObjectAccess.Object->GetPathName(), PropertyPathString, ObjectAccess.PropertyPathInfo.ToString(), ToExternal(ObjectAccess.Access));
-
-			// Initialization
-			IModularFeatures& ModularFeatures = IModularFeatures::Get();
-			const FName InterceptorFeatureName = IRemoteControlInterceptionFeatureInterceptor::GetName();
-			const int32 InterceptorsAmount = ModularFeatures.GetModularFeatureImplementationCount(IRemoteControlInterceptionFeatureInterceptor::GetName());
-
-			// Pass interception command data to all available interceptors
-			bool bShouldIntercept = false;
-			for (int32 InterceptorIdx = 0; InterceptorIdx < InterceptorsAmount; ++InterceptorIdx)
+			if (GEditor && bGenerateTransaction)
 			{
-				IRemoteControlInterceptionFeatureInterceptor* const Interceptor = static_cast<IRemoteControlInterceptionFeatureInterceptor*>(ModularFeatures.GetModularFeatureImplementation(InterceptorFeatureName, InterceptorIdx));
-				if (Interceptor)
-				{
-					// Update response flag
-					bShouldIntercept |= (Interceptor->ResetObjectProperties(ObjectMetadata) == ERCIResponse::Intercept);
-				}
+				GEditor->EndTransaction();
 			}
 
-			// Don't process the RC message if any of interceptors returned ERCIResponse::Intercept
-			if (bShouldIntercept)
-			{
-				return true;
-			}
+			FPropertyChangedEvent PropertyEvent(MutableObjectReference.PropertyPathInfo.ToPropertyChangedEvent());
+			MutableObjectReference.Object->PostEditChangeProperty(PropertyEvent);
 		}
+#endif
+		return bSuccess;
+	}
+	return false;
+}
 
-		if (ObjectAccess.IsValid() && RemoteControlUtil::IsWriteAccess(ObjectAccess.Access))
+bool FRemoteControlModule::ResetObjectProperties(const FRCObjectReference& ObjectAccess, const bool bAllowIntercept)
+{
+	// Check the replication path before reset the property on this instance
+	if (bAllowIntercept && ObjectAccess.Object.IsValid())
+	{
+		// Build interception command
+		FString PropertyPathString = TFieldPath<FProperty>(ObjectAccess.Property.Get()).ToString();
+		FRCIObjectMetadata ObjectMetadata(ObjectAccess.Object->GetPathName(), PropertyPathString, ObjectAccess.PropertyPathInfo.ToString(), ToExternal(ObjectAccess.Access));
+
+		// Initialization
+		IModularFeatures& ModularFeatures = IModularFeatures::Get();
+		const FName InterceptorFeatureName = IRemoteControlInterceptionFeatureInterceptor::GetName();
+		const int32 InterceptorsAmount = ModularFeatures.GetModularFeatureImplementationCount(IRemoteControlInterceptionFeatureInterceptor::GetName());
+
+		// Pass interception command data to all available interceptors
+		bool bShouldIntercept = false;
+		for (int32 InterceptorIdx = 0; InterceptorIdx < InterceptorsAmount; ++InterceptorIdx)
 		{
-			UObject* Object = ObjectAccess.Object.Get();
-			UStruct* ContainerType = ObjectAccess.ContainerType.Get();
-
-#if WITH_EDITOR
-			bool bGenerateTransaction = ObjectAccess.Access == ERCAccess::WRITE_TRANSACTION_ACCESS;
-			constexpr bool bForceEndChange = true;
-			TestOrFinalizeOngoingChange(true);
-			FScopedTransaction Transaction(LOCTEXT("RemoteResetPropertyTransaction", "Remote Reset Object Property"), bGenerateTransaction);
-			if (bGenerateTransaction)
+			IRemoteControlInterceptionFeatureInterceptor* const Interceptor = static_cast<IRemoteControlInterceptionFeatureInterceptor*>(ModularFeatures.GetModularFeatureImplementation(InterceptorFeatureName, InterceptorIdx));
+			if (Interceptor)
 			{
-				FEditPropertyChain PreEditChain;
-				ObjectAccess.PropertyPathInfo.ToEditPropertyChain(PreEditChain);
-				Object->PreEditChange(PreEditChain);
+				// Update response flag
+				bShouldIntercept |= (Interceptor->ResetObjectProperties(ObjectMetadata) == ERCIResponse::Intercept);
 			}
-#endif
-
-			// Copy the value from the field on the CDO.
-			FRCFieldPathInfo FieldPathInfo = ObjectAccess.PropertyPathInfo;
-			void* TargetAddress = FieldPathInfo.GetResolvedData().ContainerAddress;
-			UObject* DefaultObject = Object->GetClass()->GetDefaultObject();
-			FieldPathInfo.Resolve(DefaultObject);
-			FRCFieldResolvedData DefaultObjectResolvedData = FieldPathInfo.GetResolvedData();
-			ObjectAccess.Property->CopyCompleteValue_InContainer(TargetAddress, DefaultObjectResolvedData.ContainerAddress);
+		}
 
-			// if we are generating a transaction, also generate post edit property event, event if the change ended up unsuccessful
-			// this is to match the pre edit change call that can unregister components for example
-#if WITH_EDITOR
-			if (bGenerateTransaction)
-			{
-				FPropertyChangedEvent PropertyEvent(ObjectAccess.Property.Get());
-				Object->PostEditChangeProperty(PropertyEvent);
-			}
-#endif
+		// Don't process the RC message if any of interceptors returned ERCIResponse::Intercept
+		if (bShouldIntercept)
+		{
 			return true;
 		}
-		return false;
 	}
 
-	virtual TOptional<FExposedFunction> ResolvePresetFunction(const FResolvePresetFieldArgs& Args) const override
+	if (ObjectAccess.IsValid() && RemoteControlUtil::IsWriteAccess(ObjectAccess.Access))
 	{
-		TOptional<FExposedFunction> ExposedFunction;
+		UObject* Object = ObjectAccess.Object.Get();
+		UStruct* ContainerType = ObjectAccess.ContainerType.Get();
 
-		if (URemoteControlPreset* Preset = ResolvePreset(FName(*Args.PresetName)))
+#if WITH_EDITOR
+		bool bGenerateTransaction = ObjectAccess.Access == ERCAccess::WRITE_TRANSACTION_ACCESS;
+		constexpr bool bForceEndChange = true;
+		TestOrFinalizeOngoingChange(true);
+		FScopedTransaction Transaction(LOCTEXT("RemoteResetPropertyTransaction", "Remote Reset Object Property"), bGenerateTransaction);
+		if (bGenerateTransaction)
 		{
-			ExposedFunction = Preset->ResolveExposedFunction(FName(*Args.FieldLabel));
+			FEditPropertyChain PreEditChain;
+			ObjectAccess.PropertyPathInfo.ToEditPropertyChain(PreEditChain);
+			Object->PreEditChange(PreEditChain);
 		}
+#endif
 
-		return ExposedFunction;
+		// Copy the value from the field on the CDO.
+		FRCFieldPathInfo FieldPathInfo = ObjectAccess.PropertyPathInfo;
+		void* TargetAddress = FieldPathInfo.GetResolvedData().ContainerAddress;
+		UObject* DefaultObject = Object->GetClass()->GetDefaultObject();
+		FieldPathInfo.Resolve(DefaultObject);
+		FRCFieldResolvedData DefaultObjectResolvedData = FieldPathInfo.GetResolvedData();
+		ObjectAccess.Property->CopyCompleteValue_InContainer(TargetAddress, DefaultObjectResolvedData.ContainerAddress);
+
+		// if we are generating a transaction, also generate post edit property event, event if the change ended up unsuccessful
+		// this is to match the pre edit change call that can unregister components for example
+#if WITH_EDITOR
+		if (bGenerateTransaction)
+		{
+			FPropertyChangedEvent PropertyEvent(ObjectAccess.Property.Get());
+			Object->PostEditChangeProperty(PropertyEvent);
+		}
+#endif
+		return true;
 	}
+	return false;
+}
 
-	virtual TOptional<FExposedProperty> ResolvePresetProperty(const FResolvePresetFieldArgs& Args) const override
+TOptional<FExposedFunction> FRemoteControlModule::ResolvePresetFunction(const FResolvePresetFieldArgs& Args) const
+{
+	TOptional<FExposedFunction> ExposedFunction;
+
+	if (URemoteControlPreset* Preset = ResolvePreset(FName(*Args.PresetName)))
 	{
-		TOptional<FExposedProperty> ExposedProperty;
+		ExposedFunction = Preset->ResolveExposedFunction(FName(*Args.FieldLabel));
+	}
 
-		if (URemoteControlPreset* Preset = ResolvePreset(FName(*Args.PresetName)))
-		{
-			ExposedProperty = Preset->ResolveExposedProperty(FName(*Args.FieldLabel));
-		}
+	return ExposedFunction;
+}
+
+TOptional<FExposedProperty> FRemoteControlModule::ResolvePresetProperty(const FResolvePresetFieldArgs& Args) const
+{
+	TOptional<FExposedProperty> ExposedProperty;
 
-		return ExposedProperty;
+	if (URemoteControlPreset* Preset = ResolvePreset(FName(*Args.PresetName)))
+	{
+		ExposedProperty = Preset->ResolveExposedProperty(FName(*Args.FieldLabel));
 	}
 
-	virtual URemoteControlPreset* ResolvePreset(FName PresetName) const override
+	return ExposedProperty;
+}
+
+URemoteControlPreset* FRemoteControlModule::ResolvePreset(FName PresetName) const
+{
+	if (const TArray<FAssetData>* Assets = CachedPresetsByName.Find(PresetName))
 	{
-		if (const TArray<FAssetData>* Assets = CachedPresetsByName.Find(PresetName))
+		for (const FAssetData& Asset : *Assets)
 		{
-			for (const FAssetData& Asset : *Assets)
+			if (Asset.AssetName == PresetName)
 			{
-				if (Asset.AssetName == PresetName)
-				{
-					return Cast<URemoteControlPreset>(Asset.GetAsset());
-				}
+				return Cast<URemoteControlPreset>(Asset.GetAsset());
 			}
 		}
-
-		if (URemoteControlPreset* FoundPreset = RemoteControlUtil::GetPresetByName(PresetName))
-		{
-			CachedPresetsByName.FindOrAdd(PresetName).AddUnique(FoundPreset);
-			return FoundPreset;
-		}
-
-		return nullptr;
 	}
 
-	virtual URemoteControlPreset* ResolvePreset(const FGuid& PresetId) const override
+	if (URemoteControlPreset* FoundPreset = RemoteControlUtil::GetPresetByName(PresetName))
 	{
-		TRACE_CPUPROFILER_EVENT_SCOPE(FRemoteControlModule::ResolvePreset);
+		CachedPresetsByName.FindOrAdd(PresetName).AddUnique(FoundPreset);
+		return FoundPreset;
+	}
 
-		if (const FName* AssetName = CachedPresetNamesById.Find(PresetId))
+	return nullptr;
+}
+
+URemoteControlPreset* FRemoteControlModule::ResolvePreset(const FGuid& PresetId) const
+{
+	TRACE_CPUPROFILER_EVENT_SCOPE(FRemoteControlModule::ResolvePreset);
+
+	if (const FName* AssetName = CachedPresetNamesById.Find(PresetId))
+	{
+		if (const TArray<FAssetData>* Assets = CachedPresetsByName.Find(*AssetName))
 		{
-			if (const TArray<FAssetData>* Assets = CachedPresetsByName.Find(*AssetName))
+			for (const FAssetData& Asset : *Assets)
 			{
-				for (const FAssetData& Asset : *Assets)
+				if (RemoteControlUtil::GetPresetId(Asset) == PresetId)
 				{
-					if (RemoteControlUtil::GetPresetId(Asset) == PresetId)
-					{
-						return Cast<URemoteControlPreset>(Asset.GetAsset());
-					}
+					return Cast<URemoteControlPreset>(Asset.GetAsset());
 				}
 			}
-			else
-			{
-				ensureMsgf(false, TEXT("Preset id should be cached if the asset name already is."));
-			}
 		}
-
-		if (URemoteControlPreset* FoundPreset = RemoteControlUtil::GetPresetById(PresetId))
+		else
 		{
-			CachedPresetNamesById.Emplace(PresetId, FoundPreset->GetName());
-			return FoundPreset;
+			ensureMsgf(false, TEXT("Preset id should be cached if the asset name already is."));
 		}
-
-		return nullptr;
 	}
 
-	virtual void GetPresets(TArray<TSoftObjectPtr<URemoteControlPreset>>& OutPresets) const override
+	if (URemoteControlPreset* FoundPreset = RemoteControlUtil::GetPresetById(PresetId))
 	{
-		OutPresets.Reserve(CachedPresetsByName.Num());
-		for (const TPair<FName, TArray<FAssetData>>& Entry : CachedPresetsByName)
-		{
-			Algo::Transform(Entry.Value, OutPresets, [this](const FAssetData& AssetData) { return Cast<URemoteControlPreset>(AssetData.GetAsset()); });
-		}
+		CachedPresetNamesById.Emplace(PresetId, FoundPreset->GetName());
+		return FoundPreset;
 	}
 
-	virtual void GetPresetAssets(TArray<FAssetData>& OutPresetAssets) const override
-	{
-		OutPresetAssets.Reserve(CachedPresetsByName.Num());
-		for (const TPair<FName, TArray<FAssetData>>& Entry : CachedPresetsByName)
-		{
-			OutPresetAssets.Append(Entry.Value);
-		}
-	}
+	return nullptr;
+}
 
-	virtual const TMap<FName, FEntityMetadataInitializer>& GetDefaultMetadataInitializers() const override
+void FRemoteControlModule::GetPresets(TArray<TSoftObjectPtr<URemoteControlPreset>>& OutPresets) const
+{
+	OutPresets.Reserve(CachedPresetsByName.Num());
+	for (const TPair<FName, TArray<FAssetData>>& Entry : CachedPresetsByName)
 	{
-		return DefaultMetadataInitializers;
+		Algo::Transform(Entry.Value, OutPresets, [this](const FAssetData& AssetData) { return Cast<URemoteControlPreset>(AssetData.GetAsset()); });
 	}
+}
 
-	virtual bool RegisterDefaultEntityMetadata(FName MetadataKey, FEntityMetadataInitializer MetadataInitializer) override
+void FRemoteControlModule::GetPresetAssets(TArray<FAssetData>& OutPresetAssets) const
+{
+	OutPresetAssets.Reserve(CachedPresetsByName.Num());
+	for (const TPair<FName, TArray<FAssetData>>& Entry : CachedPresetsByName)
 	{
-		if (!DefaultMetadataInitializers.Contains(MetadataKey))
-		{
-			DefaultMetadataInitializers.FindOrAdd(MetadataKey) = MoveTemp(MetadataInitializer);
-			return true;
-		}
-		return false;
+		OutPresetAssets.Append(Entry.Value);
 	}
+}
 
-	virtual void UnregisterDefaultEntityMetadata(FName MetadataKey) override
-	{
-		DefaultMetadataInitializers.Remove(MetadataKey);
-	}
+const TMap<FName, FEntityMetadataInitializer>& FRemoteControlModule::GetDefaultMetadataInitializers() const
+{
+	return DefaultMetadataInitializers;
+}
 
-	virtual bool PropertySupportsRawModificationWithoutEditor(FProperty* Property, UClass* OwnerClass = nullptr) const override
+bool FRemoteControlModule::RegisterDefaultEntityMetadata(FName MetadataKey, FEntityMetadataInitializer MetadataInitializer)
+{
+	if (!DefaultMetadataInitializers.Contains(MetadataKey))
 	{
-		constexpr bool bInGameOrPackage = true;
-		return Property && (RemoteControlUtil::IsPropertyAllowed(Property, ERCAccess::WRITE_ACCESS, bInGameOrPackage) || !!RemoteControlPropertyUtilities::FindSetterFunction(Property, OwnerClass));
+		DefaultMetadataInitializers.FindOrAdd(MetadataKey) = MoveTemp(MetadataInitializer);
+		return true;
 	}
+	return false;
+}
 
-private:
-	void CachePresets()
-	{
-		TArray<FAssetData> Assets;
-		RemoteControlUtil::GetAllPresetAssets(Assets);
+void FRemoteControlModule::UnregisterDefaultEntityMetadata(FName MetadataKey)
+{
+	DefaultMetadataInitializers.Remove(MetadataKey);
+}
 
-		for (const FAssetData& AssetData : Assets)
-		{
-			CachedPresetsByName.FindOrAdd(AssetData.AssetName).AddUnique(AssetData);
+bool FRemoteControlModule::PropertySupportsRawModificationWithoutEditor(FProperty* Property, UClass* OwnerClass) const
+{
+	constexpr bool bInGameOrPackage = true;
+	return Property && (RemoteControlUtil::IsPropertyAllowed(Property, ERCAccess::WRITE_ACCESS, bInGameOrPackage) || !!RemoteControlPropertyUtilities::FindSetterFunction(Property, OwnerClass));
+}
 
-			const FGuid PresetAssetId = RemoteControlUtil::GetPresetId(AssetData);
-			if (PresetAssetId.IsValid())
-			{
-				CachedPresetNamesById.Add(PresetAssetId, AssetData.AssetName);
-			}
-			else if (URemoteControlPreset* Preset = Cast<URemoteControlPreset>(AssetData.GetAsset()))
-			{
-				// Handle the case where the preset asset data does not contain the ID yet.
-				// This can happen with old assets that haven't been resaved yet.
-				const FGuid PresetId = Preset->GetPresetId();
-				if (PresetId.IsValid())
-				{
-					CachedPresetNamesById.Add(PresetId, AssetData.AssetName);
-					CachedPresetsByName.FindOrAdd(AssetData.AssetName).AddUnique(AssetData);
-				}
-			}
-		}
-	}
+void FRemoteControlModule::CachePresets()
+{
+	TArray<FAssetData> Assets;
+	RemoteControlUtil::GetAllPresetAssets(Assets);
 
-	void OnAssetAdded(const FAssetData& AssetData)
+	for (const FAssetData& AssetData : Assets)
 	{
-		if (AssetData.AssetClass != URemoteControlPreset::StaticClass()->GetFName())
-		{
-			return;
-		}
+		CachedPresetsByName.FindOrAdd(AssetData.AssetName).AddUnique(AssetData);
 
 		const FGuid PresetAssetId = RemoteControlUtil::GetPresetId(AssetData);
 		if (PresetAssetId.IsValid())
 		{
 			CachedPresetNamesById.Add(PresetAssetId, AssetData.AssetName);
+		}
+		else if (URemoteControlPreset* Preset = Cast<URemoteControlPreset>(AssetData.GetAsset()))
+		{
+			// Handle the case where the preset asset data does not contain the ID yet.
+			// This can happen with old assets that haven't been resaved yet.
+			const FGuid PresetId = Preset->GetPresetId();
+			if (PresetId.IsValid())
+			{
+				CachedPresetNamesById.Add(PresetId, AssetData.AssetName);
 				CachedPresetsByName.FindOrAdd(AssetData.AssetName).AddUnique(AssetData);
 			}
 		}
+	}
+}
 
-	void OnAssetRemoved(const FAssetData& AssetData)
+void FRemoteControlModule::OnAssetAdded(const FAssetData& AssetData)
+{
+	if (AssetData.AssetClass != URemoteControlPreset::StaticClass()->GetFName())
 	{
-		if (AssetData.AssetClass != URemoteControlPreset::StaticClass()->GetFName())
-		{
-			return;
+		return;
+	}
+
+	const FGuid PresetAssetId = RemoteControlUtil::GetPresetId(AssetData);
+	if (PresetAssetId.IsValid())
+	{
+		CachedPresetNamesById.Add(PresetAssetId, AssetData.AssetName);
+			CachedPresetsByName.FindOrAdd(AssetData.AssetName).AddUnique(AssetData);
 		}
+	}
+
+void FRemoteControlModule::OnAssetRemoved(const FAssetData& AssetData)
+{
+	if (AssetData.AssetClass != URemoteControlPreset::StaticClass()->GetFName())
+	{
+		return;
+	}
 
-		const FGuid PresetId = RemoteControlUtil::GetPresetId(AssetData);
-		if (FName* PresetName = CachedPresetNamesById.Find(PresetId))
+	const FGuid PresetId = RemoteControlUtil::GetPresetId(AssetData);
+	if (FName* PresetName = CachedPresetNamesById.Find(PresetId))
+	{
+		if (TArray<FAssetData>* Assets = CachedPresetsByName.Find(*PresetName))
 		{
-			if (TArray<FAssetData>* Assets = CachedPresetsByName.Find(*PresetName))
+			for (auto It = Assets->CreateIterator(); It; ++It)
 			{
-				for (auto It = Assets->CreateIterator(); It; ++It)
+				if (RemoteControlUtil::GetPresetId(*It) == PresetId)
 				{
-					if (RemoteControlUtil::GetPresetId(*It) == PresetId)
-					{
-						It.RemoveCurrent();
-						break;
-					}
+					It.RemoveCurrent();
+					break;
 				}
 			}
 		}
 	}
+}
 
-	void OnAssetRenamed(const FAssetData& AssetData, const FString&)
+void FRemoteControlModule::OnAssetRenamed(const FAssetData& AssetData, const FString&)
+{
+	if (AssetData.AssetClass != URemoteControlPreset::StaticClass()->GetFName())
 	{
-		if (AssetData.AssetClass != URemoteControlPreset::StaticClass()->GetFName())
-		{
-			return;
-		}
+		return;
+	}
 
-		const FGuid PresetId = RemoteControlUtil::GetPresetId(AssetData);
-		if (FName* OldPresetName = CachedPresetNamesById.Find(PresetId))
+	const FGuid PresetId = RemoteControlUtil::GetPresetId(AssetData);
+	if (FName* OldPresetName = CachedPresetNamesById.Find(PresetId))
+	{
+		if (TArray<FAssetData>* Assets = CachedPresetsByName.Find(*OldPresetName))
 		{
-			if (TArray<FAssetData>* Assets = CachedPresetsByName.Find(*OldPresetName))
+			for (auto It = Assets->CreateIterator(); It; ++It)
 			{
-				for (auto It = Assets->CreateIterator(); It; ++It)
+				if (RemoteControlUtil::GetPresetId(*It) == PresetId)
 				{
-					if (RemoteControlUtil::GetPresetId(*It) == PresetId)
-					{
-						It.RemoveCurrent();
-						break;
-					}
+					It.RemoveCurrent();
+					break;
 				}
+			}
 
-				if (Assets->Num() == 0)
-				{
-					CachedPresetsByName.Remove(*OldPresetName);
-				}
+			if (Assets->Num() == 0)
+			{
+				CachedPresetsByName.Remove(*OldPresetName);
 			}
 		}
+	}
 
-		CachedPresetNamesById.Remove(PresetId);
-		CachedPresetNamesById.Add(PresetId, AssetData.AssetName);
+	CachedPresetNamesById.Remove(PresetId);
+	CachedPresetNamesById.Add(PresetId, AssetData.AssetName);
 
-		CachedPresetsByName.FindOrAdd(AssetData.AssetName).AddUnique(AssetData);
-	}
+	CachedPresetsByName.FindOrAdd(AssetData.AssetName).AddUnique(AssetData);
+}
 
-	bool PropertyModificationShouldUseSetter(UObject* Object, FProperty* Property)
+bool FRemoteControlModule::PropertyModificationShouldUseSetter(UObject* Object, FProperty* Property)
+{
+	if (!Property || !Object)
 	{
-		if (!Property || !Object)
-		{
-			return false;
-		}
-
-		return !!RemoteControlPropertyUtilities::FindSetterFunction(Property, Object->GetClass());
+		return false;
 	}
 
+	return !!RemoteControlPropertyUtilities::FindSetterFunction(Property, Object->GetClass());
+}
+
 #if WITH_EDITOR
-	void TestOrFinalizeOngoingChange(bool bForceEndChange = false)
+void FRemoteControlModule::TestOrFinalizeOngoingChange(bool bForceEndChange)
+{
+	if (OngoingModification)
 	{
-		if (OngoingModification)
+		if (bForceEndChange || !OngoingModification->bWasTriggeredSinceLastPass)
 		{
-			if (bForceEndChange || !OngoingModification->bWasTriggeredSinceLastPass)
+			// Call PreEditChange for the last call
+			if (OngoingModification->Reference.IsType<FRCObjectReference>())
 			{
-				// Call PreEditChange for the last call
-				if (OngoingModification->Reference.IsType<FRCObjectReference>())
+				FRCObjectReference& Reference = OngoingModification->Reference.Get<FRCObjectReference>();
+				if (Reference.IsValid())
 				{
-					FRCObjectReference& Reference = OngoingModification->Reference.Get<FRCObjectReference>();
-					if (Reference.IsValid())
-					{
-						FEditPropertyChain PreEditChain;
-						Reference.PropertyPathInfo.ToEditPropertyChain(PreEditChain);
-						Reference.Object->PreEditChange(PreEditChain);
-					}
+					FEditPropertyChain PreEditChain;
+					Reference.PropertyPathInfo.ToEditPropertyChain(PreEditChain);
+					Reference.Object->PreEditChange(PreEditChain);
 				}
+			}
 			
-				// If not, trigger the post edit change (and end the transaction if one was started)
-				if (GEditor && OngoingModification->bHasStartedTransaction)
-				{
-					GEditor->EndTransaction();
-				}
+			// If not, trigger the post edit change (and end the transaction if one was started)
+			if (GEditor && OngoingModification->bHasStartedTransaction)
+			{
+				GEditor->EndTransaction();
+			}
 
-				if (OngoingModification->Reference.IsType<FRCObjectReference>())
+			if (OngoingModification->Reference.IsType<FRCObjectReference>())
+			{
+				FPropertyChangedEvent PropertyEvent = OngoingModification->Reference.Get<FRCObjectReference>().PropertyPathInfo.ToPropertyChangedEvent();
+				if (UObject* Object = OngoingModification->Reference.Get<FRCObjectReference>().Object.Get())
 				{
-					FPropertyChangedEvent PropertyEvent = OngoingModification->Reference.Get<FRCObjectReference>().PropertyPathInfo.ToPropertyChangedEvent();
-					if (UObject* Object = OngoingModification->Reference.Get<FRCObjectReference>().Object.Get())
-					{
-						Object->PostEditChangeProperty(PropertyEvent);
-					}
+					Object->PostEditChangeProperty(PropertyEvent);
 				}
-
-				OngoingModification.Reset();
-			}
-			else
-			{
-				// If the change has occured for this change, effectively reset the counter for it.
-				OngoingModification->bWasTriggeredSinceLastPass = false;
 			}
+
+			OngoingModification.Reset();
+		}
+		else
+		{
+			// If the change has occured for this change, effectively reset the counter for it.
+			OngoingModification->bWasTriggeredSinceLastPass = false;
 		}
 	}
+}
 
-	void HandleEnginePostInit()
-	{
-		CachePresets();
-		RegisterEditorDelegates();
-	}
+void FRemoteControlModule::HandleEnginePostInit()
+{
+	CachePresets();
+	RegisterEditorDelegates();
+}
 
-	void RegisterEditorDelegates()
+void FRemoteControlModule::RegisterEditorDelegates()
+{
+	if (GEditor)
 	{
-		if (GEditor)
-		{
-			GEditor->GetTimerManager()->SetTimer(OngoingChangeTimer, FTimerDelegate::CreateRaw(this, &FRemoteControlModule::TestOrFinalizeOngoingChange, false), SecondsBetweenOngoingChangeCheck, true);
-		}
+		GEditor->GetTimerManager()->SetTimer(OngoingChangeTimer, FTimerDelegate::CreateRaw(this, &FRemoteControlModule::TestOrFinalizeOngoingChange, false), SecondsBetweenOngoingChangeCheck, true);
 	}
+}
 	
-	void UnregisterEditorDelegates()
-	{
-		if (GEditor)
-		{ 
-			GEditor->GetTimerManager()->ClearTimer(OngoingChangeTimer);
-		}
+void FRemoteControlModule::UnregisterEditorDelegates()
+{
+	if (GEditor)
+	{ 
+		GEditor->GetTimerManager()->ClearTimer(OngoingChangeTimer);
 	}
+}
 #endif
 
-private:
-	/** Cache of preset names to preset assets */
-	mutable TMap<FName, TArray<FAssetData>> CachedPresetsByName;
-
-	/** Cache of ids to preset names. */
-	mutable TMap<FGuid, FName> CachedPresetNamesById;
-
-	/** Map of registered default metadata initializers. */
-	TMap<FName, FEntityMetadataInitializer> DefaultMetadataInitializers;
-
-	/** Delegate for preset registration */
-	FOnPresetRegistered OnPresetRegisteredDelegate;
-
-	/** Delegate for preset unregistration */
-	FOnPresetUnregistered OnPresetUnregisteredDelegate;
-
-	/** RC Processor feature instance */
-	TUniquePtr<IRemoteControlInterceptionFeatureProcessor> RCIProcessor;
-
 #if WITH_EDITOR
-	/** Flags for a given RC change. */
-	struct FOngoingChange
-	{
-		FOngoingChange(FRCObjectReference InReference)
-		{
-			Reference.Set<FRCObjectReference>(MoveTemp(InReference));
-		}
+FRemoteControlModule::FOngoingChange::FOngoingChange(FRCObjectReference InReference)
+{
+	Reference.Set<FRCObjectReference>(MoveTemp(InReference));
+}
 		
-		FOngoingChange(FRCCallReference InReference)
-		{
-			Reference.Set<FRCCallReference>(MoveTemp(InReference));
-		}
-
-		friend uint32 GetTypeHash(const FOngoingChange& Modification)
-		{
-			if (Modification.Reference.IsType<FRCObjectReference>())
-			{
-				return GetTypeHash(Modification.Reference.Get<FRCObjectReference>());
-			}
-			else
-			{
-				return GetTypeHash(Modification.Reference.Get<FRCCallReference>());
-			}
-		}
-
-		/** Reference to either the property we're modifying or the function we're calling. */
-		TVariant<FRCObjectReference, FRCCallReference> Reference;
-		/** Whether this change was triggered with a transaction or not. */
-		bool bHasStartedTransaction = false;
-		/** Whether this change was triggered since the last tick of OngoingChangeTimer. */
-		bool bWasTriggeredSinceLastPass = true;
-	};
+FRemoteControlModule::FOngoingChange::FOngoingChange(FRCCallReference InReference)
+{
+	Reference.Set<FRCCallReference>(MoveTemp(InReference));
+}
 	
-	/** Ongoing change that needs to end its transaction and call PostEditChange in the case of a property modification. */
-	TOptional<FOngoingChange> OngoingModification;
-
-	/** Handle to the timer that that ends the ongoing change in regards to PostEditChange and transactions. */
-	FTimerHandle OngoingChangeTimer;
-
-	/** Delay before we check if a modification is no longer ongoing. */
-	static constexpr float SecondsBetweenOngoingChangeCheck = 0.2f;
 #endif
-};
 
 PRAGMA_ENABLE_DEPRECATION_WARNINGS
 
 #undef LOCTEXT_NAMESPACE
-
-IMPLEMENT_MODULE(FRemoteControlModule, RemoteControl);
diff --git a/Engine/Plugins/VirtualProduction/RemoteControl/Source/RemoteControl/Private/RemoteControlModule.h b/Engine/Plugins/VirtualProduction/RemoteControl/Source/RemoteControl/Private/RemoteControlModule.h
new file mode 100644
--- /dev/null
+++ b/Engine/Plugins/VirtualProduction/RemoteControl/Source/RemoteControl/Private/RemoteControlModule.h
@@ -0,0 +1,133 @@
+// Copyright Epic Games, Inc. All Rights Reserved.
+#pragma once
+
+#include "IRemoteControlModule.h"
+#include "CoreMinimal.h"
+
+class IRemoteControlInterceptionFeatureProcessor;
+struct FAssetData;
+
+/**
+ * Implementation of the RemoteControl interface
+ */
+PRAGMA_DISABLE_DEPRECATION_WARNINGS
+
+class FRemoteControlModule : public IRemoteControlModule
+{
+public:
+	//~ Begin IModuleInterface
+	virtual void StartupModule() override;
+	virtual void ShutdownModule() override;
+	//~ End IModuleInterface
+
+	//~ Begin IRemoteControlModule
+	virtual FOnPresetRegistered& OnPresetRegistered() override;
+	virtual FOnPresetUnregistered& OnPresetUnregistered() override;
+	virtual bool RegisterPreset(FName Name, URemoteControlPreset* Preset) override;
+	virtual void UnregisterPreset(FName Name) override;
+	virtual bool ResolveCall(const FString& ObjectPath, const FString& FunctionName, FRCCallReference& OutCallRef, FString* OutErrorText) override;
+	virtual bool InvokeCall(FRCCall& InCall, ERCPayloadType InPayloadType = ERCPayloadType::Json, const TArray<uint8>& InInterceptPayload = TArray<uint8>()) override;
+	virtual bool ResolveObject(ERCAccess AccessType, const FString& ObjectPath, const FString& PropertyName, FRCObjectReference& OutObjectRef, FString* OutErrorText = nullptr) override;
+	virtual bool ResolveObjectProperty(ERCAccess AccessType, UObject* Object, FRCFieldPathInfo PropertyPath, FRCObjectReference& OutObjectRef, FString* OutErrorText = nullptr) override;
+	virtual bool GetObjectProperties(const FRCObjectReference& ObjectAccess, IStructSerializerBackend& Backend) override;
+	virtual bool SetObjectProperties(const FRCObjectReference& ObjectAccess, IStructDeserializerBackend& Backend, ERCPayloadType InPayloadType, const TArray<uint8>& InPayload) override;
+	virtual bool ResetObjectProperties(const FRCObjectReference& ObjectAccess, const bool bAllowIntercept) override;
+	virtual TOptional<FExposedFunction> ResolvePresetFunction(const FResolvePresetFieldArgs& Args) const override;
+	virtual TOptional<FExposedProperty> ResolvePresetProperty(const FResolvePresetFieldArgs& Args) const override;
+	virtual URemoteControlPreset* ResolvePreset(FName PresetName) const override;
+	virtual URemoteControlPreset* ResolvePreset(const FGuid& PresetId) const override;
+	virtual void GetPresets(TArray<TSoftObjectPtr<URemoteControlPreset>>& OutPresets) const override;
+	virtual void GetPresetAssets(TArray<FAssetData>& OutPresetAssets) const override;
+	virtual const TMap<FName, FEntityMetadataInitializer>& GetDefaultMetadataInitializers() const override;
+	virtual bool RegisterDefaultEntityMetadata(FName MetadataKey, FEntityMetadataInitializer MetadataInitializer) override;
+	virtual void UnregisterDefaultEntityMetadata(FName MetadataKey) override;
+	virtual bool PropertySupportsRawModificationWithoutEditor(FProperty* Property, UClass* OwnerClass = nullptr) const override;
+	//~ End IRemoteControlModule
+
+private:
+	/** Cache all presets in the project for the ResolvePreset function. */
+	void CachePresets();
+	
+	//~ Asset registry callbacks
+	void OnAssetAdded(const FAssetData& AssetData);
+	void OnAssetRemoved(const FAssetData& AssetData);
+	void OnAssetRenamed(const FAssetData& AssetData, const FString&);
+
+	/** Determines if a property modification should use a setter or default to deserializing directly onto an object. */
+	static bool PropertyModificationShouldUseSetter(UObject* Object, FProperty* Property);
+
+#if WITH_EDITOR
+	/** Finalize an ongoing change, triggering post edit change on the tracked object. */
+	void TestOrFinalizeOngoingChange(bool bForceEndChange = false);
+
+	/** Callback to handle registering delegates once the engine has finished its startup. */
+	void HandleEnginePostInit();
+
+	//~ Register/Unregister editor delegates.
+	void RegisterEditorDelegates();
+	void UnregisterEditorDelegates();
+	
+#endif
+
+private:
+	/** Cache of preset names to preset assets */
+	mutable TMap<FName, TArray<FAssetData>> CachedPresetsByName;
+
+	/** Cache of ids to preset names. */
+	mutable TMap<FGuid, FName> CachedPresetNamesById;
+
+	/** Map of registered default metadata initializers. */
+	TMap<FName, FEntityMetadataInitializer> DefaultMetadataInitializers;
+
+	/** Delegate for preset registration */
+	FOnPresetRegistered OnPresetRegisteredDelegate;
+
+	/** Delegate for preset unregistration */
+	FOnPresetUnregistered OnPresetUnregisteredDelegate;
+
+	/** RC Processor feature instance */
+	TUniquePtr<IRemoteControlInterceptionFeatureProcessor> RCIProcessor;
+
+#if WITH_EDITOR
+	/** Flags for a given RC change. */
+	struct FOngoingChange
+	{
+		FOngoingChange(FRCObjectReference InReference);
+		FOngoingChange(FRCCallReference InReference);
+		
+		friend uint32 GetTypeHash(const FOngoingChange& Modification)
+		{
+			if (Modification.Reference.IsType<FRCObjectReference>())
+			{
+				return GetTypeHash(Modification.Reference.Get<FRCObjectReference>());
+			}
+			else
+			{
+				return GetTypeHash(Modification.Reference.Get<FRCCallReference>());
+			}
+		}
+
+		/** Reference to either the property we're modifying or the function we're calling. */
+		TVariant<FRCObjectReference, FRCCallReference> Reference;
+		/** Whether this change was triggered with a transaction or not. */
+		bool bHasStartedTransaction = false;
+		/** Whether this change was triggered since the last tick of OngoingChangeTimer. */
+		bool bWasTriggeredSinceLastPass = true;
+	};
+	
+	/** Ongoing change that needs to end its transaction and call PostEditChange in the case of a property modification. */
+	TOptional<FOngoingChange> OngoingModification;
+
+	/** Handle to the timer that that ends the ongoing change in regards to PostEditChange and transactions. */
+	FTimerHandle OngoingChangeTimer;
+
+	/** Delay before we check if a modification is no longer ongoing. */
+	static constexpr float SecondsBetweenOngoingChangeCheck = 0.2f;
+#endif
+};
+
+PRAGMA_ENABLE_DEPRECATION_WARNINGS
+
+#undef LOCTEXT_NAMESPACE
+
+IMPLEMENT_MODULE(FRemoteControlModule, RemoteControl);
