Refactor preset UI code to be more generic in regards to exposed entities

diff --git a/Engine/Plugins/VirtualProduction/RemoteControl/Source/RemoteControl/Private/RemoteControlExposeRegistry.cpp b/Engine/Plugins/VirtualProduction/RemoteControl/Source/RemoteControl/Private/RemoteControlExposeRegistry.cpp
--- a/Engine/Plugins/VirtualProduction/RemoteControl/Source/RemoteControl/Private/RemoteControlExposeRegistry.cpp
+++ b/Engine/Plugins/VirtualProduction/RemoteControl/Source/RemoteControl/Private/RemoteControlExposeRegistry.cpp
@@ -84,12 +84,18 @@ const UScriptStruct* URemoteControlExposeRegistry::GetExposedEntityType(const FG
 	return nullptr;
 }
 
+const TSet<UScriptStruct*>& URemoteControlExposeRegistry::GetExposedEntityTypes() const
+{
+	return ExposedTypes;
+}
+
 TSharedPtr<FRemoteControlEntity> URemoteControlExposeRegistry::AddExposedEntity(FRemoteControlEntity&& EntityToExpose, UScriptStruct* EntityType)
 {
 	LabelToIdCache.Add(EntityToExpose.GetLabel(), EntityToExpose.GetId());
 	FRCEntityWrapper Wrapper{ MoveTemp(EntityToExpose), EntityType};
 	TSharedPtr<FRemoteControlEntity> Entity = Wrapper.Get();
 	ExposedEntities.Add(MoveTemp(Wrapper));
+	ExposedTypes.Add(EntityType);
 	return Entity;
 }
 
diff --git a/Engine/Plugins/VirtualProduction/RemoteControl/Source/RemoteControl/Private/RemoteControlExposeRegistry.h b/Engine/Plugins/VirtualProduction/RemoteControl/Source/RemoteControl/Private/RemoteControlExposeRegistry.h
--- a/Engine/Plugins/VirtualProduction/RemoteControl/Source/RemoteControl/Private/RemoteControlExposeRegistry.h
+++ b/Engine/Plugins/VirtualProduction/RemoteControl/Source/RemoteControl/Private/RemoteControlExposeRegistry.h
@@ -5,8 +5,6 @@
 #include "CoreMinimal.h"
 #include "RemoteControlEntity.h"
 #include "UObject/Class.h"
-#include "UObject/StructOnScope.h"
-#include "Templates/EnableIf.h"
 #include "Templates/SharedPointer.h"
 #include "Templates/UnrealTypeTraits.h"
 
@@ -151,6 +149,11 @@ public:
 	 */
 	const UScriptStruct* GetExposedEntityType(const FGuid& ExposedEntityId) const;
 
+	/**
+	 * Get all the exposed types in the registry (ie. FRemoteControlActor)
+	 */
+	const TSet<UScriptStruct*>& GetExposedEntityTypes() const;
+
 	/**
 	 * Add an entity to the set of exposed entities of its type.
 	 * @param EntityToExpose the entity to expose.
@@ -206,4 +209,8 @@ private:
 	/** Cache of label to ids. */
 	UPROPERTY(Transient)
 	TMap<FName, FGuid> LabelToIdCache;
+
+	/** Holds the types of entities exposed in the registry. */
+	UPROPERTY()
+	TSet<UScriptStruct*> ExposedTypes;
 };
diff --git a/Engine/Plugins/VirtualProduction/RemoteControl/Source/RemoteControl/Private/RemoteControlPreset.cpp b/Engine/Plugins/VirtualProduction/RemoteControl/Source/RemoteControl/Private/RemoteControlPreset.cpp
--- a/Engine/Plugins/VirtualProduction/RemoteControl/Source/RemoteControl/Private/RemoteControlPreset.cpp
+++ b/Engine/Plugins/VirtualProduction/RemoteControl/Source/RemoteControl/Private/RemoteControlPreset.cpp
@@ -1444,6 +1444,11 @@ const UScriptStruct* URemoteControlPreset::GetExposedEntityType(const FGuid& Exp
 	return Registry->GetExposedEntityType(ExposedEntityId);
 }
 
+const TSet<UScriptStruct*>& URemoteControlPreset::GetExposedEntityTypes() const
+{
+	return Registry->GetExposedEntityTypes();
+}
+
 FName URemoteControlPreset::RenameExposedEntity(const FGuid& ExposedEntityId, FName NewLabel)
 {
 	FName AssignedLabel = Registry->RenameExposedEntity(ExposedEntityId, NewLabel);
diff --git a/Engine/Plugins/VirtualProduction/RemoteControl/Source/RemoteControl/Public/RemoteControlPreset.h b/Engine/Plugins/VirtualProduction/RemoteControl/Source/RemoteControl/Public/RemoteControlPreset.h
--- a/Engine/Plugins/VirtualProduction/RemoteControl/Source/RemoteControl/Public/RemoteControlPreset.h
+++ b/Engine/Plugins/VirtualProduction/RemoteControl/Source/RemoteControl/Public/RemoteControlPreset.h
@@ -528,7 +528,7 @@ public:
 	/**
 	 * Get the exposed entities of a certain type.
 	 */
-	template <typename ExposableEntityType>
+	template <typename ExposableEntityType = FRemoteControlEntity>
 	TArray<TWeakPtr<const ExposableEntityType>> GetExposedEntities() const
 	{
 		static_assert(TIsDerivedFrom<ExposableEntityType, FRemoteControlEntity>::Value, "ExposableEntityType must derive from FRemoteControlEntity.");
@@ -545,7 +545,7 @@ public:
 	/**
 	 * Get the exposed entities of a certain type.
 	 */
-	template <typename ExposableEntityType>
+	template <typename ExposableEntityType = FRemoteControlEntity>
 	TArray<TWeakPtr<ExposableEntityType>> GetExposedEntities()
 	{
 		static_assert(TIsDerivedFrom<ExposableEntityType, FRemoteControlEntity>::Value, "ExposableEntityType must derive from FRemoteControlEntity.");
@@ -560,6 +560,14 @@ public:
 		return ReturnedEntities;
 	}
 
+	TArray<TWeakPtr<FRemoteControlEntity>> GetExposedEntities(UScriptStruct* EntityType)
+	{
+		TArray<TWeakPtr<FRemoteControlEntity>> ReturnedEntities;
+		const TArray<TSharedPtr<FRemoteControlEntity>> Entities = GetEntities(EntityType);
+		ReturnedEntities.Append(Entities);
+		return ReturnedEntities;
+	}
+
 	/**
 	 * Get a copy of an exposed entity on the preset.
 	 * @param ExposedEntityId The id of the entity to get.
@@ -586,6 +594,9 @@ public:
 
 	/** Get the type of an exposed entity by querying with its id. (ie. FRemoteControlActor) */
 	const UScriptStruct* GetExposedEntityType(const FGuid& ExposedEntityId) const;
+	
+	/** Get all types of exposed entities currently exposed. (ie. FRemoteControlActor) */
+	const TSet<UScriptStruct*>& GetExposedEntityTypes() const;
 
 	/** Returns whether an entity is exposed on the preset. */
 	bool IsExposed(const FGuid& ExposedEntityId) const;
diff --git a/Engine/Plugins/VirtualProduction/RemoteControl/Source/RemoteControlProtocolWidgets/Private/RemoteControlProtocolWidgetsModule.cpp b/Engine/Plugins/VirtualProduction/RemoteControl/Source/RemoteControlProtocolWidgets/Private/RemoteControlProtocolWidgetsModule.cpp
--- a/Engine/Plugins/VirtualProduction/RemoteControl/Source/RemoteControlProtocolWidgets/Private/RemoteControlProtocolWidgetsModule.cpp
+++ b/Engine/Plugins/VirtualProduction/RemoteControl/Source/RemoteControlProtocolWidgets/Private/RemoteControlProtocolWidgetsModule.cpp
@@ -19,7 +19,7 @@ TSharedRef<SWidget> FRemoteControlProtocolWidgetsModule::GenerateDetailsForEntit
 	ResetProtocolBindingList();
 
 	// Currently only supports Properties
-	if (InFieldId.IsValid() && InFieldType == EExposedFieldType::Property)
+	if (InFieldId.IsValid() && InPreset->GetExposedEntityType(InFieldId) == FRemoteControlProperty::StaticStruct())
 	{
 		return SAssignNew(RCProtocolBindingList, SRCProtocolBindingList, FProtocolEntityViewModel::Create(InPreset, InFieldId));
 	}
diff --git a/Engine/Plugins/VirtualProduction/RemoteControl/Source/RemoteControlUI/Private/RemoteControlUIModule.cpp b/Engine/Plugins/VirtualProduction/RemoteControl/Source/RemoteControlUI/Private/RemoteControlUIModule.cpp
--- a/Engine/Plugins/VirtualProduction/RemoteControl/Source/RemoteControlUI/Private/RemoteControlUIModule.cpp
+++ b/Engine/Plugins/VirtualProduction/RemoteControl/Source/RemoteControlUI/Private/RemoteControlUIModule.cpp
@@ -4,25 +4,23 @@
 
 #include "AssetToolsModule.h"
 #include "AssetTools/RemoteControlPresetActions.h"
-#include "EditorStyleSet.h"
-#include "ISettingsModule.h"
-#include "ISettingsSection.h"
 #include "Framework/MultiBox/MultiBoxBuilder.h"
 #include "HAL/PlatformApplicationMisc.h"
+#include "MaterialEditor/DEditorParameterValue.h"
 #include "PropertyHandle.h"
 #include "RemoteControlActor.h"
 #include "RemoteControlField.h"
 #include "RemoteControlPreset.h"
 #include "RemoteControlSettings.h"
+#include "Styling/AppStyle.h"
 #include "Textures/SlateIcon.h"
 #include "UI/Customizations/RemoteControlEntityCustomization.h"
+#include "UI/SRCPanelExposedField.h"
+#include "UI/SRCPanelExposedActor.h"
 #include "UI/RemoteControlPanelStyle.h"
 #include "UI/SRCPanelExposedEntitiesList.h"
 #include "UI/SRemoteControlPanel.h"
 #include "Widgets/DeclarativeSyntaxSupport.h"
-#include "Widgets/Images/SImage.h"
-#include "Widgets/Input/SButton.h"
-#include "Widgets/SWidget.h"
 
 #define LOCTEXT_NAMESPACE "RemoteControlUI"
 
@@ -32,10 +30,11 @@ namespace RemoteControlUIModule
 	{
 		static TArray<FName> CustomizedStructNames =
 		{
-		FRemoteControlProperty::StaticStruct()->GetFName(),
-		FRemoteControlFunction::StaticStruct()->GetFName(),
-		FRemoteControlActor::StaticStruct()->GetFName()
-	};
+			FRemoteControlProperty::StaticStruct()->GetFName(),
+			FRemoteControlFunction::StaticStruct()->GetFName(),
+			FRemoteControlActor::StaticStruct()->GetFName()
+		};
+
 		return CustomizedStructNames;
 	};
 }
@@ -48,6 +47,7 @@ void FRemoteControlUIModule::StartupModule()
 	RegisterContextMenuExtender();
 	RegisterStructCustomizations();
 	RegisterSettings();
+	RegisterWidgetFactories();
 }
 
 void FRemoteControlUIModule::ShutdownModule()
@@ -214,7 +214,7 @@ FSlateIcon FRemoteControlUIModule::OnGetExposedIcon(TSharedPtr<IPropertyHandle>
 		}
 	}
 
-	return FSlateIcon(FEditorStyle::Get().GetStyleSetName(), BrushName);
+	return FSlateIcon(FAppStyle::Get().GetStyleSetName(), BrushName);
 }
 
 bool FRemoteControlUIModule::CanToggleExposeProperty(TSharedPtr<IPropertyHandle> Handle) const
@@ -391,6 +391,41 @@ void FRemoteControlUIModule::OnSettingsModified(UObject*, FPropertyChangedEvent&
 	}
 }
 
+void FRemoteControlUIModule::RegisterWidgetFactoryForType(UScriptStruct* RemoteControlEntityType, const FOnGenerateRCWidget& OnGenerateRCWidgetDelegate)
+{
+	if (!GenerateWidgetDelegates.Contains(RemoteControlEntityType))
+	{
+		GenerateWidgetDelegates.Add(RemoteControlEntityType, OnGenerateRCWidgetDelegate);
+	}
+}
+
+void FRemoteControlUIModule::UnregisterWidgetFactoryForType(UScriptStruct* RemoteControlEntityType)
+{
+	GenerateWidgetDelegates.Remove(RemoteControlEntityType);
+}
+
+void FRemoteControlUIModule::RegisterWidgetFactories()
+{
+	RegisterWidgetFactoryForType(FRemoteControlActor::StaticStruct(), FOnGenerateRCWidget::CreateStatic(&SRCPanelExposedActor::MakeInstance));
+	RegisterWidgetFactoryForType(FRemoteControlProperty::StaticStruct(), FOnGenerateRCWidget::CreateStatic(&SRCPanelExposedField::MakeInstance));
+	RegisterWidgetFactoryForType(FRemoteControlFunction::StaticStruct(), FOnGenerateRCWidget::CreateStatic(&SRCPanelExposedField::MakeInstance));
+}
+
+TSharedPtr<SRCPanelTreeNode> FRemoteControlUIModule::GenerateEntityWidget(const FGenerateWidgetArgs& Args)
+{
+	if (Args.Preset && Args.Entity)
+	{
+		const UScriptStruct* EntityType = Args.Preset->GetExposedEntityType(Args.Entity->GetId());
+		
+		if (FOnGenerateRCWidget* OnGenerateWidget = GenerateWidgetDelegates.Find(const_cast<UScriptStruct*>(EntityType)))
+		{
+			return OnGenerateWidget->Execute(Args);
+		}
+	}
+
+	return nullptr;
+}
+
 IMPLEMENT_MODULE(FRemoteControlUIModule, RemoteControlUI);
 
-#undef LOCTEXT_NAMESPACE /*RemoteControlUI*/
\ No newline at end of file
+#undef LOCTEXT_NAMESPACE /*RemoteControlUI*/
diff --git a/Engine/Plugins/VirtualProduction/RemoteControl/Source/RemoteControlUI/Private/RemoteControlUIModule.h b/Engine/Plugins/VirtualProduction/RemoteControl/Source/RemoteControlUI/Private/RemoteControlUIModule.h
--- a/Engine/Plugins/VirtualProduction/RemoteControl/Source/RemoteControlUI/Private/RemoteControlUIModule.h
+++ b/Engine/Plugins/VirtualProduction/RemoteControl/Source/RemoteControlUI/Private/RemoteControlUIModule.h
@@ -4,9 +4,7 @@
 #include "IRemoteControlUIModule.h"
 
 #include "CoreMinimal.h"	
-#include "Input/Reply.h"
 #include "LevelEditor.h"
-#include "Modules/ModuleInterface.h"
 #include "Modules/ModuleManager.h"
 #include "PropertyEditorDelegates.h"
 #include "PropertyHandle.h"
@@ -45,6 +43,9 @@ public:
 	virtual void RegisterMetadataCustomization(FName MetadataKey, FOnCustomizeMetadataEntry OnCustomizeCallback) override;
 	virtual void UnregisterMetadataCustomization(FName MetadataKey) override;
 	virtual URemoteControlPreset* GetActivePreset() const override;
+	virtual void RegisterWidgetFactoryForType(UScriptStruct* RemoteControlEntityType, const FOnGenerateRCWidget& OnGenerateRCWidgetDelegate) override;
+	virtual void UnregisterWidgetFactoryForType(UScriptStruct* RemoteControlEntityType) override;
+
 	//~ End IRemoteControlUIModule interface
 
 	/**
@@ -62,6 +63,8 @@ public:
 		return ExternalEntityMetadataCustomizations;
 	}
 
+	TSharedPtr<SRCPanelTreeNode> GenerateEntityWidget(const FGenerateWidgetArgs& Args);
+
 private:
 	/**
 	 * The status of a property.
@@ -121,6 +124,8 @@ private:
 	void RegisterSettings();
 	void UnregisterSettings();
 	void OnSettingsModified(UObject*, struct FPropertyChangedEvent&);
+
+	void RegisterWidgetFactories();
 private:
 	/** The custom actions added to the actor context menu. */
 	TSharedPtr<class FRemoteControlPresetActions> RemoteControlPresetActions;
@@ -142,4 +147,6 @@ private:
 
 	/** Map of metadata key to customization handler. */
 	TMap<FName, FOnCustomizeMetadataEntry> ExternalEntityMetadataCustomizations;
+
+	TMap<TWeakObjectPtr<UScriptStruct>, FOnGenerateRCWidget> GenerateWidgetDelegates;
 };
diff --git a/Engine/Plugins/VirtualProduction/RemoteControl/Source/RemoteControlUI/Private/UI/RemoteControlPanelStyle.cpp b/Engine/Plugins/VirtualProduction/RemoteControl/Source/RemoteControlUI/Private/UI/RemoteControlPanelStyle.cpp
--- a/Engine/Plugins/VirtualProduction/RemoteControl/Source/RemoteControlUI/Private/UI/RemoteControlPanelStyle.cpp
+++ b/Engine/Plugins/VirtualProduction/RemoteControl/Source/RemoteControlUI/Private/UI/RemoteControlPanelStyle.cpp
@@ -1,10 +1,10 @@
 // Copyright Epic Games, Inc. All Rights Reserved.
 
 #include "UI/RemoteControlPanelStyle.h"
+#include "Styling/AppStyle.h"
 #include "Styling/SlateStyleRegistry.h"
 #include "Styling/SlateTypes.h"
 #include "Styling/CoreStyle.h"
-#include "EditorStyleSet.h"
 #include "Interfaces/IPluginManager.h"
 #include "SlateOptMacros.h"
 #include "Widgets/Input/SButton.h"
@@ -14,7 +14,7 @@
 #define BOX_BRUSH(RelativePath, ...) FSlateBoxBrush(StyleSet->RootToContentDir(RelativePath, TEXT(".png")), __VA_ARGS__)
 #define CORE_BOX_BRUSH( RelativePath, ... ) FSlateBoxBrush(StyleSet->RootToCoreContentDir(RelativePath, TEXT(".png") ), __VA_ARGS__)
 #define BOX_PLUGIN_BRUSH( RelativePath, ... ) FSlateBoxBrush(FRemoteControlPanelStyle::InContent( RelativePath, ".png" ), __VA_ARGS__)
-#define DEFAULT_FONT(...) FCoreStyle::GetDefaultFontStyle(__VA_ARGS__)
+#define DEFAULT_FONT(...) FAppStyle::Get().GetDefaultFontStyle(__VA_ARGS__)
 
 TSharedPtr<FSlateStyleSet> FRemoteControlPanelStyle::StyleSet;
 
@@ -36,27 +36,29 @@ void FRemoteControlPanelStyle::Initialize()
 
 	StyleSet->Set("ClassThumbnail.RemoteControlPreset", new IMAGE_PLUGIN_BRUSH("Icons/RemoteControlAPI_64x", Icon64x64));
 
-	FButtonStyle ExposeFunctionButtonStyle = FEditorStyle::Get().GetWidgetStyle<FButtonStyle>("PropertyEditor.AssetComboStyle");
+	const ISlateStyle& AppStyle = FAppStyle::Get();
+
+	FButtonStyle ExposeFunctionButtonStyle = AppStyle.GetWidgetStyle<FButtonStyle>("PropertyEditor.AssetComboStyle");
 	ExposeFunctionButtonStyle.Normal = BOX_BRUSH("Common/GroupBorder", FMargin(4.0f / 16.0f));
 	ExposeFunctionButtonStyle.Normal.TintColor = FLinearColor(0, 0, 0, 0.1);
 	StyleSet->Set("RemoteControlPanel.ExposeFunctionButton", ExposeFunctionButtonStyle);
 
-	FHyperlinkStyle ObjectSectionNameStyle = FCoreStyle::Get().GetWidgetStyle<FHyperlinkStyle>("Hyperlink");
+	FHyperlinkStyle ObjectSectionNameStyle = AppStyle.GetWidgetStyle<FHyperlinkStyle>("Hyperlink");
 	FButtonStyle SectionNameButtonStyle = ObjectSectionNameStyle.UnderlineStyle;
 	SectionNameButtonStyle.Normal = ObjectSectionNameStyle.UnderlineStyle.Hovered;
 	StyleSet->Set("RemoteControlPanel.SectionNameButton", SectionNameButtonStyle);
 
 	FTextBlockStyle SectionNameTextStyle = ObjectSectionNameStyle.TextStyle;
-	SectionNameTextStyle.Font = FEditorStyle::GetFontStyle("DetailsView.CategoryFontStyle");
+	SectionNameTextStyle.Font = AppStyle.GetFontStyle("DetailsView.CategoryFontStyle");
 	StyleSet->Set("RemoteControlPanel.SectionName", SectionNameTextStyle);
 
-	FButtonStyle UnexposeButtonStyle = FEditorStyle::Get().GetWidgetStyle<FButtonStyle>("FlatButton");
+	FButtonStyle UnexposeButtonStyle = AppStyle.GetWidgetStyle<FButtonStyle>("FlatButton");
 	UnexposeButtonStyle.Normal = FSlateNoResource();
 	UnexposeButtonStyle.NormalPadding = FMargin(0, 1.5f);
 	UnexposeButtonStyle.PressedPadding = FMargin(0, 1.5f);
 	StyleSet->Set("RemoteControlPanel.UnexposeButton", UnexposeButtonStyle);
 
-	FTextBlockStyle ButtonTextStyle = FEditorStyle::Get().GetWidgetStyle<FTextBlockStyle>("ContentBrowser.TopBar.Font");
+	FTextBlockStyle ButtonTextStyle = AppStyle.GetWidgetStyle<FTextBlockStyle>("ContentBrowser.TopBar.Font");
 	FLinearColor ButtonTextColor = ButtonTextStyle.ColorAndOpacity.GetSpecifiedColor();
 	ButtonTextColor.A /= 2;
 	ButtonTextStyle.ColorAndOpacity = ButtonTextColor;
@@ -66,7 +68,7 @@ void FRemoteControlPanelStyle::Initialize()
 	// Default to transparent
 	StyleSet->Set("RemoteControlPanel.ExposedFieldBorder", new FSlateNoResource());
 
-	FEditableTextBoxStyle SectionNameTextBoxStyle = FCoreStyle::Get().GetWidgetStyle< FEditableTextBoxStyle >("NormalEditableTextBox");
+	FEditableTextBoxStyle SectionNameTextBoxStyle = AppStyle.GetWidgetStyle< FEditableTextBoxStyle >("NormalEditableTextBox");
 	SectionNameTextBoxStyle.BackgroundImageNormal = BOX_BRUSH("Common/GroupBorderLight", FMargin(4.0f / 16.0f));
 	StyleSet->Set("RemoteControlPanel.SectionNameTextBox", SectionNameTextBoxStyle);
 
@@ -76,14 +78,14 @@ void FRemoteControlPanelStyle::Initialize()
 	StyleSet->Set("RemoteControlPanel.HorizontalDash", new IMAGE_BRUSH("Common/HorizontalDottedLine_16x1px", FVector2D(16.0f, 1.0f), FLinearColor::White, ESlateBrushTileType::Horizontal));
 	StyleSet->Set("RemoteControlPanel.VerticalDash", new IMAGE_BRUSH("Common/VerticalDottedLine_1x16px", FVector2D(1.0f, 16.0f), FLinearColor::White, ESlateBrushTileType::Vertical));
 
-	FLinearColor NewSelectionColor = FCoreStyle::Get().GetWidgetStyle<FTableRowStyle>("TableView.Row").ActiveBrush.TintColor.GetSpecifiedColor();
+	FLinearColor NewSelectionColor = AppStyle.GetWidgetStyle<FTableRowStyle>("TableView.Row").ActiveBrush.TintColor.GetSpecifiedColor();
 	NewSelectionColor.R *= 1.8;
 	NewSelectionColor.G *= 1.8;
 	NewSelectionColor.B *= 1.8;
 	NewSelectionColor.A = 0.3;
 	StyleSet->Set("RemoteControlPanel.GroupRowSelected", new BOX_BRUSH("Common/GroupBorderLight", FMargin(4.0f / 16.0f), NewSelectionColor));
 
-	FTableRowStyle GroupRowStyle = FCoreStyle::Get().GetWidgetStyle<FTableRowStyle>("TableView.Row");
+	FTableRowStyle GroupRowStyle = AppStyle.GetWidgetStyle<FTableRowStyle>("TableView.Row");
 	GroupRowStyle.SetEvenRowBackgroundBrush(*FRemoteControlPanelStyle::Get()->GetBrush("RemoteControlPanel.GroupBorder"));
 	GroupRowStyle.SetOddRowBackgroundBrush(*FRemoteControlPanelStyle::Get()->GetBrush("RemoteControlPanel.GroupBorder"));
 	GroupRowStyle.SetEvenRowBackgroundHoveredBrush(*FRemoteControlPanelStyle::Get()->GetBrush("RemoteControlPanel.GroupBorder"));
diff --git a/Engine/Plugins/VirtualProduction/RemoteControl/Source/RemoteControlUI/Private/UI/SRCLogger.cpp b/Engine/Plugins/VirtualProduction/RemoteControl/Source/RemoteControlUI/Private/UI/SRCLogger.cpp
--- a/Engine/Plugins/VirtualProduction/RemoteControl/Source/RemoteControlUI/Private/UI/SRCLogger.cpp
+++ b/Engine/Plugins/VirtualProduction/RemoteControl/Source/RemoteControlUI/Private/UI/SRCLogger.cpp
@@ -2,10 +2,10 @@
 
 #include "SRCLogger.h"
 
-#include "EditorStyleSet.h"
 #include "MessageLogModule.h"
 #include "RemoteControlLogger.h"
 #include "Modules/ModuleManager.h"
+#include "Styling/AppStyle.h"
 #include "Widgets/Layout/SBorder.h"
 
 void SRCLogger::Construct(const FArguments& InArgs)
@@ -27,7 +27,7 @@ void SRCLogger::Construct(const FArguments& InArgs)
 	ChildSlot
 	[
 		SNew(SBorder)
-		.BorderImage(FEditorStyle::GetBrush("Menu.Background"))
+		.BorderImage(FAppStyle::Get().GetBrush("Menu.Background"))
 		.Padding(2.f)
 		[
 			LogListingWidget.ToSharedRef()
diff --git a/Engine/Plugins/VirtualProduction/RemoteControl/Source/RemoteControlUI/Private/UI/SRCPanelExposedActor.cpp b/Engine/Plugins/VirtualProduction/RemoteControl/Source/RemoteControlUI/Private/UI/SRCPanelExposedActor.cpp
--- a/Engine/Plugins/VirtualProduction/RemoteControl/Source/RemoteControlUI/Private/UI/SRCPanelExposedActor.cpp
+++ b/Engine/Plugins/VirtualProduction/RemoteControl/Source/RemoteControlUI/Private/UI/SRCPanelExposedActor.cpp
@@ -5,30 +5,31 @@
 #include "AssetRegistry/AssetData.h"
 #include "PropertyCustomizationHelpers.h"
 #include "RemoteControlActor.h"
-#include "RemoteControlPanelStyle.h"
 #include "RemoteControlPreset.h"
 #include "ScopedTransaction.h"
 #include "SRCPanelDragHandle.h"
-#include "UObject/SoftObjectPath.h"
-#include "Widgets/Input/SButton.h"
-#include "Widgets/Text/SInlineEditableTextBlock.h"
-#include "Widgets/Text/STextBlock.h"
 
 #define LOCTEXT_NAMESPACE "RemoteControlPanelActor"
 
+TSharedPtr<SRCPanelTreeNode> SRCPanelExposedActor::MakeInstance(const FGenerateWidgetArgs& Args)
+{
+	return SNew(SRCPanelExposedActor, StaticCastSharedPtr<FRemoteControlActor>(Args.Entity), Args.Preset, Args.ColumnSizeData).EditMode(Args.bIsInEditMode);
+}
+
 void SRCPanelExposedActor::Construct(const FArguments& InArgs, TWeakPtr<FRemoteControlActor> InWeakActor, URemoteControlPreset* InPreset, FRCColumnSizeData InColumnSizeData)
 {
 	FString ActorPath;
 
-	TSharedPtr<FRemoteControlActor> Actor = InWeakActor.Pin();
+	const TSharedPtr<FRemoteControlActor> Actor = InWeakActor.Pin();
 	if (ensure(Actor))
 	{
+		Initialize(Actor->GetId(), InPreset, InArgs._EditMode);
+		
 		ColumnSizeData = MoveTemp(InColumnSizeData);
-		ExposedActorId = Actor->GetId();
 		WeakPreset = InPreset;
 		bEditMode = InArgs._EditMode;
-		CachedLabel = Actor->GetLabel();
 		WeakActor = MoveTemp(InWeakActor);
+		
 		if (AActor* ResolvedActor = Actor->GetActor())
 		{
 			ActorPath = ResolvedActor->GetPathName();
@@ -41,24 +42,7 @@ void SRCPanelExposedActor::Construct(const FArguments& InArgs, TWeakPtr<FRemoteC
 	];
 }
 
-void SRCPanelExposedActor::Tick(const FGeometry&, const double, const float)
-{
-	if (bNeedsRename)
-	{
-		if (NameTextBox)
-		{
-			NameTextBox->EnterEditingMode();
-		}
-		bNeedsRename = false;
-	}
-}
-
-FGuid SRCPanelExposedActor::GetId() const
-{
-	return ExposedActorId;
-}
-
-SRCPanelTreeNode::ENodeType SRCPanelExposedActor::GetType() const
+SRCPanelTreeNode::ENodeType SRCPanelExposedActor::GetRCType() const
 {
 	return ENodeType::Actor;
 }
@@ -79,21 +63,6 @@ void SRCPanelExposedActor::Refresh()
 	ChildSlot.AttachWidget(RecreateWidget(FString()));
 }
 
-TSharedPtr<SRCPanelExposedActor> SRCPanelExposedActor::AsActor()
-{
-	return SharedThis(this);
-}
-
-TSharedPtr<FRemoteControlEntity> SRCPanelExposedActor::GetEntity() const
-{
-	return WeakActor.Pin();
-}
-
-TWeakPtr<FRemoteControlActor> SRCPanelExposedActor::GetRemoteControlActor() const
-{
-	return WeakActor;
-}
-
 TSharedRef<SWidget> SRCPanelExposedActor::RecreateWidget(const FString& Path)
 {
 	TSharedRef<SObjectPropertyEntryBox> EntryBox = SNew(SObjectPropertyEntryBox)
@@ -108,35 +77,7 @@ TSharedRef<SWidget> SRCPanelExposedActor::RecreateWidget(const FString& Path)
 	float MinWidth = 0.f, MaxWidth = 0.f;
 	EntryBox->GetDesiredWidth(MinWidth, MaxWidth);
 
-	FMakeNodeWidgetArgs Args;
-	Args.NameWidget = SAssignNew(NameTextBox, SInlineEditableTextBlock)
-		.Text(FText::FromName(CachedLabel))
-		.OnTextCommitted(this, &SRCPanelExposedActor::OnLabelCommitted)
-		.OnVerifyTextChanged(this, &SRCPanelExposedActor::OnVerifyItemLabelChanged)
-		.IsReadOnly_Lambda([this]() { return !bEditMode.Get(); });
-
-	Args.DragHandle = SNew(SBox)
-		.Visibility(this, &SRCPanelExposedActor::GetVisibilityAccordingToEditMode, EVisibility::Hidden)
-		[
-			SNew(SRCPanelDragHandle<FExposedEntityDragDrop>, ExposedActorId)
-			.Widget(AsShared())
-		];
-
-	Args.RenameButton = SNew(SButton)
-		.Visibility(this, &SRCPanelExposedActor::GetVisibilityAccordingToEditMode, EVisibility::Collapsed)
-		.ButtonStyle(FEditorStyle::Get(), "FlatButton")
-		.OnClicked_Lambda([this]() {
-			bNeedsRename = true;
-			return FReply::Handled();
-		})
-		[
-			SNew(STextBlock)
-			.TextStyle(FRemoteControlPanelStyle::Get(), "RemoteControlPanel.Button.TextStyle")
-			.Font(FEditorStyle::Get().GetFontStyle("FontAwesome.10"))
-			.Text(FText::FromString(FString(TEXT("\xf044"))) /*fa-edit*/)
-		];
-
-	Args.ValueWidget =
+	TSharedRef<SWidget> ValueWidget =
 		SNew(SBox)
 		.MinDesiredWidth(MinWidth)
 		.MaxDesiredWidth(MaxWidth)
@@ -144,19 +85,7 @@ TSharedRef<SWidget> SRCPanelExposedActor::RecreateWidget(const FString& Path)
 			MoveTemp(EntryBox)
 		];
 
-	Args.UnexposeButton = SNew(SButton)
-		.Visibility(this, &SRCPanelExposedActor::GetVisibilityAccordingToEditMode, EVisibility::Collapsed)
-		.OnPressed(this, &SRCPanelExposedActor::HandleUnexposeActor)
-		.ButtonStyle(FRemoteControlPanelStyle::Get(), "RemoteControlPanel.UnexposeButton")
-		[
-			SNew(STextBlock)
-			.TextStyle(FRemoteControlPanelStyle::Get(), "RemoteControlPanel.Button.TextStyle")
-			.Font(FEditorStyle::Get().GetFontStyle("FontAwesome.10"))
-			.Text(FText::FromString(FString(TEXT("\xf00d"))) /*fa-times*/)
-		];
-
-
-	return MakeNodeWidget(Args);
+	return CreateEntityWidget(MoveTemp(ValueWidget));
 }
 
 void SRCPanelExposedActor::OnChangeActor(const FAssetData& AssetData)
@@ -172,42 +101,4 @@ void SRCPanelExposedActor::OnChangeActor(const FAssetData& AssetData)
 	}
 }
 
-void SRCPanelExposedActor::HandleUnexposeActor()
-{
-	if (URemoteControlPreset* RCPreset = WeakPreset.Get())
-	{
-		RCPreset->Unexpose(ExposedActorId);
-	}
-}
-
-bool SRCPanelExposedActor::OnVerifyItemLabelChanged(const FText& InLabel, FText& OutErrorMessage)
-{
-	if (URemoteControlPreset* RCPreset = WeakPreset.Get())
-	{
-		if (InLabel.ToString() != CachedLabel.ToString() && RCPreset->GetExposedEntityId(FName(*InLabel.ToString())).IsValid())
-		{
-			OutErrorMessage = LOCTEXT("NameAlreadyExists", "This name already exists.");
-			return false;
-		}
-	}
-
-	return true;
-}
-
-void SRCPanelExposedActor::OnLabelCommitted(const FText& InLabel, ETextCommit::Type InCommitInfo)
-{
-	if (URemoteControlPreset* RCPreset = WeakPreset.Get())
-	{
-		FScopedTransaction Transaction(LOCTEXT("RenameExposedActor", "Modify exposed actor label"));
-		RCPreset->Modify();
-		CachedLabel =  RCPreset->RenameExposedEntity(ExposedActorId, *InLabel.ToString());
-		NameTextBox->SetText(FText::FromName(CachedLabel));
-	}
-}
-
-EVisibility SRCPanelExposedActor::GetVisibilityAccordingToEditMode(EVisibility NonEditModeVisibility) const
-{
-	return bEditMode.Get() ? EVisibility::Visible : NonEditModeVisibility;
-}
-
 #undef LOCTEXT_NAMESPACE /*RemoteControlPanelActor*/
\ No newline at end of file
diff --git a/Engine/Plugins/VirtualProduction/RemoteControl/Source/RemoteControlUI/Private/UI/SRCPanelExposedActor.h b/Engine/Plugins/VirtualProduction/RemoteControl/Source/RemoteControlUI/Private/UI/SRCPanelExposedActor.h
--- a/Engine/Plugins/VirtualProduction/RemoteControl/Source/RemoteControlUI/Private/UI/SRCPanelExposedActor.h
+++ b/Engine/Plugins/VirtualProduction/RemoteControl/Source/RemoteControlUI/Private/UI/SRCPanelExposedActor.h
@@ -1,12 +1,14 @@
 // Copyright Epic Games, Inc. All Rights Reserved.
 #pragma once
 
+#include "IRemoteControlUIModule.h"
 #include "RemoteControlEntity.h"
 #include "SRCPanelExposedEntity.h"
 #include "Misc/Guid.h"
 #include "UObject/WeakObjectPtrTemplates.h"
 
 struct FRemoteControlActor;
+struct FGenerateWidgetArgs;
 class SInlineEditableTextBlock;
 class SObjectPropertyEntryBox;
 class URemoteControlPreset;
@@ -15,59 +17,35 @@ struct FAssetData;
 DECLARE_DELEGATE_OneParam(FOnUnexposeActor, const FGuid& /**ActorId*/);
 
 /** Represents an actor exposed to remote control. */
-struct SRCPanelExposedActor : public SCompoundWidget, public SRCPanelExposedEntity	
+struct SRCPanelExposedActor : public SRCPanelExposedEntity
 {
+	using SCompoundWidget::AsShared;
+	
 	SLATE_BEGIN_ARGS(SRCPanelExposedActor)
 		: _EditMode(true)
 	{}
 		SLATE_ATTRIBUTE(bool, EditMode)
 	SLATE_END_ARGS()
 
+	static TSharedPtr<SRCPanelTreeNode> MakeInstance(const FGenerateWidgetArgs& Args);
+	
 	void Construct(const FArguments& InArgs, TWeakPtr<FRemoteControlActor> InWeakActor, URemoteControlPreset* InPreset, FRCColumnSizeData InColumnSizeData);
 
-	//~ Tick
-	virtual void Tick(const FGeometry&, const double, const float) override;
-
 	//~ Begin SRCPanelTreeNode interface
-	FGuid GetId() const override;
-	ENodeType GetType() const override;
+	virtual ENodeType GetRCType() const override;
 	virtual void Refresh() override;
-	TSharedPtr<SRCPanelExposedActor> AsActor() override;
 	//~ End SRCPanelTreeNode interface
-	
-	//~ SRCPanelExposedEntity Interface
-	virtual TSharedPtr<FRemoteControlEntity> GetEntity() const override;
-	//~ End SRCPanelExposedEntity Interface
-
-	/** Get a weak pointer to the underlying remote control actor. */
-	TWeakPtr<FRemoteControlActor> GetRemoteControlActor() const;
 
 private:
 	/** Regenerate this row's content. */
 	TSharedRef<SWidget> RecreateWidget(const FString& Path);
 	/** Handle the user selecting a different actor to expose. */
 	void OnChangeActor(const FAssetData& AssetData);
-	/** Handle unexposing the actor this row represents. */
-	void HandleUnexposeActor();
-	/** Verifies that the field's label doesn't already exist. */
-	bool OnVerifyItemLabelChanged(const FText& InLabel, FText& OutErrorMessage);
-	/** Handles committing a field label. */
-	void OnLabelCommitted(const FText& InLabel, ETextCommit::Type InCommitInfo);
-	/** Returns the visibility according to the edit mode. */
-	EVisibility GetVisibilityAccordingToEditMode(EVisibility NonEditModeVisibility) const;
 private:
-	/** Holds the ID of the exposed entity. */
-	FGuid ExposedActorId;
-	/** Cached label of the exposed entity. */
-	FName CachedLabel;
 	/** Weak reference to the preset that exposes the actor. */
 	TWeakObjectPtr<URemoteControlPreset> WeakPreset;
 	/** Weak ptr to the remote control actor structure. */
 	TWeakPtr<FRemoteControlActor> WeakActor;
 	/** Holds this row's panel edit mode. */
 	TAttribute<bool> bEditMode;
-	/** Whether to enter edit mode on the label text box. */
-	bool bNeedsRename = false;
-	/** Holds the editable label text box. */
-	TSharedPtr<SInlineEditableTextBlock> NameTextBox;
 };
diff --git a/Engine/Plugins/VirtualProduction/RemoteControl/Source/RemoteControlUI/Private/UI/SRCPanelExposedEntitiesList.cpp b/Engine/Plugins/VirtualProduction/RemoteControl/Source/RemoteControlUI/Private/UI/SRCPanelExposedEntitiesList.cpp
--- a/Engine/Plugins/VirtualProduction/RemoteControl/Source/RemoteControlUI/Private/UI/SRCPanelExposedEntitiesList.cpp
+++ b/Engine/Plugins/VirtualProduction/RemoteControl/Source/RemoteControlUI/Private/UI/SRCPanelExposedEntitiesList.cpp
@@ -12,17 +12,13 @@
 #include "Editor/EditorEngine.h"
 #include "Misc/Attribute.h"
 #include "Misc/Guid.h"
-#include "RemoteControlActor.h"
 #include "RemoteControlPanelStyle.h"
 #include "RemoteControlPreset.h"
-#include "RemoteControlSettings.h"
+#include "RemoteControlUIModule.h"
 #include "ScopedTransaction.h"
 #include "SRCPanelFieldGroup.h"
-#include "SRCPanelExposedActor.h"
 #include "SRCPanelExposedField.h"
-#include "Modules/ModuleManager.h"
 #include "UObject/Object.h"
-#include "Widgets/Text/STextBlock.h"
 #include "Widgets/Views/STableRow.h"
 #include "Widgets/Views/STreeView.h"
 
@@ -32,7 +28,6 @@ void SRCPanelExposedEntitiesList::Construct(const FArguments& InArgs, URemoteCon
 {
 	bIsInEditMode = InArgs._EditMode;
 	Preset = TStrongObjectPtr<URemoteControlPreset>(InPreset);
-	bDisplayValues = InArgs._DisplayValues;
 	OnEntityListUpdatedDelegate = InArgs._OnEntityListUpdated;
 	WidgetRegistry = MoveTemp(InWidgetRegistry);
 
@@ -77,13 +72,13 @@ void SRCPanelExposedEntitiesList::SetSelection(const TSharedPtr<SRCPanelTreeNode
 {
 	if (Node)
 	{
-		if (TSharedPtr<SRCPanelTreeNode>* FoundTreeNode = FieldWidgetMap.Find(Node->GetId()))
+		if (TSharedPtr<SRCPanelTreeNode>* FoundTreeNode = FieldWidgetMap.Find(Node->GetRCId()))
 		{
 			TreeView->SetSelection(*FoundTreeNode);
 			return;
 		}
 
-		if (TSharedPtr<SRCPanelGroup> SRCGroup = FindGroupById(Node->GetId()))
+		if (TSharedPtr<SRCPanelGroup> SRCGroup = FindGroupById(Node->GetRCId()))
 		{
 			TreeView->SetSelection(SRCGroup);
 		}	
@@ -160,26 +155,17 @@ void SRCPanelExposedEntitiesList::GenerateListWidgets()
 {
 	FieldWidgetMap.Reset();
 
-	for (TWeakPtr<FRemoteControlField> WeakField : Preset->GetExposedEntities<FRemoteControlField>())
+	for (TWeakPtr<FRemoteControlEntity> WeakEntity : Preset->GetExposedEntities())
 	{
-		if (TSharedPtr<FRemoteControlField> Field = WeakField.Pin())
+		if (const TSharedPtr<FRemoteControlEntity> Entity = WeakEntity.Pin())
 		{
-			FieldWidgetMap.Add(Field->GetId(),
-	            SNew(SRCPanelExposedField, MoveTemp(WeakField), ColumnSizeData, WidgetRegistry)
-	            .Preset(Preset.Get())
-	            .EditMode(bIsInEditMode)
-	            .DisplayValues(bDisplayValues)
-	        );
-		}
-	}
-
-	for (TWeakPtr<FRemoteControlActor> WeakActor : Preset->GetExposedEntities<FRemoteControlActor>())
-	{
-		if (TSharedPtr<FRemoteControlActor> Actor = WeakActor.Pin())
-		{
-			FieldWidgetMap.Add(Actor->GetId(),
-				SNew(SRCPanelExposedActor, MoveTemp(WeakActor), Preset.Get(), ColumnSizeData)
-				.EditMode(bIsInEditMode));
+			FGenerateWidgetArgs Args;
+			Args.Entity = Entity;
+			Args.Preset = Preset.Get();
+			Args.WidgetRegistry = WidgetRegistry;
+			Args.ColumnSizeData = ColumnSizeData;
+			
+			FieldWidgetMap.Add(Entity->GetId(), FRemoteControlUIModule::Get().GenerateEntityWidget(Args));
 		}
 	}
 }
@@ -215,62 +201,49 @@ void SRCPanelExposedEntitiesList::RefreshGroups()
 
 TSharedRef<ITableRow> SRCPanelExposedEntitiesList::OnGenerateRow(TSharedPtr<SRCPanelTreeNode> Node, const TSharedRef<STableViewBase>& OwnerTable)
 {
-	constexpr float LeftPadding = 3.f;
-	if (Node->GetType() == SRCPanelTreeNode::Group)
-	{
-		return SNew(STableRow<TSharedPtr<SWidget>>, OwnerTable)
-		.Padding(FMargin(0.5f, 0.5f, 0.5f, 0.5f))
-		.Style(&FRemoteControlPanelStyle::Get()->GetWidgetStyle<FTableRowStyle>("RemoteControlPanel.GroupRow"))
-		[
-			Node->AsGroup().ToSharedRef()
-		];
-	}
-	else if (Node->GetType() == SRCPanelTreeNode::Field)
+	const TSharedRef<SWidget> NodeWidget = Node->AsShared();
+	
+	auto OnDropLambda = [this, Node]
+	(const FDragDropEvent& Event)
 	{
-		auto OnDropLambda = [this, Field = Node->AsField()]
-		(const FDragDropEvent& Event)
+		if (const TSharedPtr<FExposedEntityDragDrop> DragDropOp = Event.GetOperationAs<FExposedEntityDragDrop>())
 		{
-			if (TSharedPtr<FExposedEntityDragDrop> DragDropOp = Event.GetOperationAs<FExposedEntityDragDrop>())
+			if (const TSharedPtr<SRCPanelGroup> Group = FindGroupById(GetGroupId(Node->GetRCId())))
 			{
-				if (TSharedPtr<SRCPanelGroup> Group = FindGroupById(GetGroupId(Field->GetId())))
+				if (DragDropOp->IsOfType<FExposedEntityDragDrop>())
 				{
-					if (DragDropOp->IsOfType<FExposedEntityDragDrop>())
-					{
-						return OnDropOnGroup(DragDropOp, Field, Group);
-					}
-					else if (DragDropOp->IsOfType<FFieldGroupDragDropOp>())
-					{
-						return OnDropOnGroup(DragDropOp, nullptr, Group);
-					}
+					return OnDropOnGroup(DragDropOp, Node, Group);
+				}
+				else if (DragDropOp->IsOfType<FFieldGroupDragDropOp>())
+				{
+					return OnDropOnGroup(DragDropOp, nullptr, Group);
 				}
 			}
+		}
 
-			return FReply::Unhandled();
-		};
-
-		return SNew(STableRow<TSharedPtr<FGuid>>, OwnerTable)
-			.OnDragEnter_Lambda([Field = Node->AsField()](const FDragDropEvent& Event) { Field->SetIsHovered(true); })
-			.OnDragLeave_Lambda([Field = Node->AsField()](const FDragDropEvent& Event) { Field->SetIsHovered(false); })
-			.OnDrop_Lambda(OnDropLambda)
-			.Padding(FMargin(LeftPadding, 0.f, 0.f, 0.f))
-			[
-				Node->AsField().ToSharedRef()
-			];
-	}
-	else if (Node->GetType() == SRCPanelTreeNode::FieldChild)
+		return FReply::Unhandled();
+	};
+	
+	if (Node->GetRCType() == SRCPanelTreeNode::Group)
 	{
 		return SNew(STableRow<TSharedPtr<SWidget>>, OwnerTable)
-			.Padding(FMargin(LeftPadding + 10.f, 0.f, 0.f, 0.f))
+			.Padding(FMargin(0.5f, 0.5f, 0.5f, 0.5f))
+			.Style(&FRemoteControlPanelStyle::Get()->GetWidgetStyle<FTableRowStyle>("RemoteControlPanel.GroupRow"))
 			[
-				Node->AsFieldChild().ToSharedRef()
+				NodeWidget
 			];
 	}
 	else
 	{
-		return SNew(STableRow<TSharedPtr<SWidget>>, OwnerTable)
-			.Padding(FMargin(LeftPadding, 0.f, 0.f, 0.f))
+		constexpr float LeftPadding = 3.f;
+		const FMargin Margin = Node->GetRCType() == SRCPanelTreeNode::FieldChild ? FMargin(LeftPadding + 10.f, 0.f, 0.f, 0.f) : FMargin(LeftPadding, 0.f, 0.f, 0.f);
+		return SNew(STableRow<TSharedPtr<FGuid>>, OwnerTable)
+			.OnDragEnter_Lambda([Node](const FDragDropEvent& Event) { if (Node && Node->GetRCType() == SRCPanelTreeNode::Field) StaticCastSharedPtr<SRCPanelExposedField>(Node)->SetIsHovered(true); })
+			.OnDragLeave_Lambda([Node](const FDragDropEvent& Event) { if (Node && Node->GetRCType() == SRCPanelTreeNode::Field) StaticCastSharedPtr<SRCPanelExposedField>(Node)->SetIsHovered(false); })
+			.OnDrop_Lambda(OnDropLambda)
+			.Padding(Margin)
 			[
-				Node->AsActor().ToSharedRef()
+				NodeWidget
 			];
 	}
 }
@@ -290,10 +263,10 @@ void SRCPanelExposedEntitiesList::OnSelectionChanged(TSharedPtr<SRCPanelTreeNode
 		return;
 	}
 
-	if (TSharedPtr<SRCPanelExposedField> Field = Node->AsField())
+	// todo call handler on node itself
+	if (TSharedPtr<FRemoteControlField> RCField = Preset->GetExposedEntity<FRemoteControlField>(Node->GetRCId()).Pin())
 	{
-		TSet<UObject*> Objects;
-		Field->GetBoundObjects(Objects);
+		TSet<UObject*> Objects = TSet<UObject*>{ RCField->GetBoundObjects() };
 
 		TArray<UObject*> OwnerActors;
 		for (UObject* Object : Objects)
@@ -308,16 +281,13 @@ void SRCPanelExposedEntitiesList::OnSelectionChanged(TSharedPtr<SRCPanelTreeNode
 			}
 		}
 
-		for (UObject* Object : Objects)
-		{
-			SelectActorsInlevel(OwnerActors);
-		}
+		SelectActorsInlevel(OwnerActors);
 	}
 
 	OnSelectionChangeDelegate.Broadcast(Node);
 }
 
-FReply SRCPanelExposedEntitiesList::OnDropOnGroup(const TSharedPtr<FDragDropOperation>& DragDropOperation, const TSharedPtr<SRCPanelTreeNode>& TargetEntity, const TSharedPtr<SRCPanelGroup>& DragTargetGroup)
+FReply SRCPanelExposedEntitiesList::OnDropOnGroup(const TSharedPtr<FDragDropOperation>& DragDropOperation, const TSharedPtr<SRCPanelTreeNode>& TargetEntity, const TSharedPtr<SRCPanelTreeNode>& DragTargetGroup)
 {
 	checkSlow(DragTargetGroup);
 
@@ -333,12 +303,12 @@ FReply SRCPanelExposedEntitiesList::OnDropOnGroup(const TSharedPtr<FDragDropOper
 
 			FRemoteControlPresetLayout::FFieldSwapArgs Args;
 			Args.OriginGroupId = DragOriginGroupId;
-			Args.TargetGroupId = DragTargetGroup->GetId();
+			Args.TargetGroupId = DragTargetGroup->GetRCId();
 			Args.DraggedFieldId = DragDropOp->GetId();
 
 			if (TargetEntity)
 			{
-				Args.TargetFieldId = TargetEntity->GetId();
+				Args.TargetFieldId = TargetEntity->GetRCId();
 			}
 			else
 			{
@@ -360,7 +330,7 @@ FReply SRCPanelExposedEntitiesList::OnDropOnGroup(const TSharedPtr<FDragDropOper
 		if (TSharedPtr<FFieldGroupDragDropOp> DragDropOp = StaticCastSharedPtr<FFieldGroupDragDropOp>(DragDropOperation))
 		{
 			FGuid DragOriginGroupId = DragDropOp->GetGroupId();
-			FGuid DragTargetGroupId = DragTargetGroup->GetId();
+			FGuid DragTargetGroupId = DragTargetGroup->GetRCId();
 
 			if (DragOriginGroupId == DragTargetGroupId)
 			{
@@ -389,11 +359,11 @@ FGuid SRCPanelExposedEntitiesList::GetGroupId(const FGuid& EntityId)
 	return GroupId;
 }
 
-void SRCPanelExposedEntitiesList::OnDeleteGroup(const TSharedPtr<SRCPanelGroup>& PanelGroup)
+void SRCPanelExposedEntitiesList::OnDeleteGroup(const FGuid& GroupId)
 {
 	FScopedTransaction Transaction(LOCTEXT("DeleteGroup", "Delete Group"));
 	Preset->Modify();
-	Preset->Layout.DeleteGroup(PanelGroup->GetId());
+	Preset->Layout.DeleteGroup(GroupId);
 }
 
 void SRCPanelExposedEntitiesList::SelectActorsInlevel(const TArray<UObject*>& Objects)
@@ -436,7 +406,7 @@ void SRCPanelExposedEntitiesList::UnregisterEvents()
 TSharedPtr<SRCPanelGroup> SRCPanelExposedEntitiesList::FindGroupById(const FGuid& Id)
 {
 	TSharedPtr<SRCPanelGroup> TargetGroup;
-	if (TSharedPtr<SRCPanelGroup>* FoundGroup = FieldGroups.FindByPredicate([Id](const TSharedPtr<SRCPanelGroup>& InGroup) {return InGroup->GetId() == Id; }))
+	if (TSharedPtr<SRCPanelGroup>* FoundGroup = FieldGroups.FindByPredicate([Id](const TSharedPtr<SRCPanelGroup>& InGroup) {return InGroup->GetRCId() == Id; }))
 	{
 		TargetGroup = *FoundGroup;
 	}
@@ -484,29 +454,29 @@ void SRCPanelExposedEntitiesList::UnregisterPresetDelegates()
 
 void SRCPanelExposedEntitiesList::OnEntityAdded(const FGuid& InEntityId)
 {
-	auto ExposeEntity = [this, InEntityId](TSharedRef<SRCPanelTreeNode>&& Node)
+	auto ExposeEntity = [this, InEntityId](TSharedPtr<SRCPanelTreeNode>&& Node)
 	{
-		FieldWidgetMap.Add(InEntityId, Node);
-
-		if (TSharedPtr<SRCPanelGroup> SRCGroup = FindGroupById(GetGroupId(InEntityId)))
+		if (Node)
 		{
-			SRCGroup->GetNodes().Add(MoveTemp(Node));
-			TreeView->SetItemExpansion(SRCGroup, true);
+			FieldWidgetMap.Add(InEntityId, Node);
+
+			if (TSharedPtr<SRCPanelGroup> SRCGroup = FindGroupById(GetGroupId(InEntityId)))
+			{
+				SRCGroup->GetNodes().Add(MoveTemp(Node));
+				TreeView->SetItemExpansion(SRCGroup, true);
+			}
 		}
 	};
-
-	if (TSharedPtr<FRemoteControlActor> Actor = Preset->GetExposedEntity<FRemoteControlActor>(InEntityId).Pin())
-	{
-		ExposeEntity(SNew(SRCPanelExposedActor, MoveTemp(Actor), Preset.Get(), ColumnSizeData)
-			.EditMode(bIsInEditMode));
-	}
-	else if (TSharedPtr<FRemoteControlField> Field = Preset->GetExposedEntity<FRemoteControlField>(InEntityId).Pin())
-	{
-		ExposeEntity(SNew(SRCPanelExposedField, MoveTemp(Field), ColumnSizeData, WidgetRegistry)
-			.Preset(Preset.Get())
-			.EditMode(bIsInEditMode)
-			.DisplayValues(bDisplayValues));
-	}
+	
+	FGenerateWidgetArgs Args;
+	Args.Preset = Preset.Get();
+	Args.WidgetRegistry = WidgetRegistry;
+	Args.ColumnSizeData = ColumnSizeData;
+	Args.bIsInEditMode = bIsInEditMode;
+	Args.Entity = Preset->GetExposedEntity(InEntityId).Pin();
+
+	ExposeEntity(FRemoteControlUIModule::Get().GenerateEntityWidget(Args));
+	
 	TreeView->RequestListRefresh();
 }
 
@@ -514,16 +484,16 @@ void SRCPanelExposedEntitiesList::OnEntityRemoved(const FGuid& InGroupId, const
 {
 	if (TSharedPtr<SRCPanelGroup> PanelGroup = FindGroupById(InGroupId))
 	{
-		int32 EntityIndex = PanelGroup->GetNodes().IndexOfByPredicate([InEntityId](const TSharedPtr<SRCPanelTreeNode>& Node) { return Node->GetId() == InEntityId; });
+		const int32 EntityIndex = PanelGroup->GetNodes().IndexOfByPredicate([InEntityId](const TSharedPtr<SRCPanelTreeNode>& Node) { return Node->GetRCId() == InEntityId; });
 		if (EntityIndex != INDEX_NONE)
 		{
 			PanelGroup->GetNodes().RemoveAt(EntityIndex);
 		}
 	}
 
-	if (TSharedPtr<SRCPanelTreeNode> Node = GetSelection())
+	if (const TSharedPtr<SRCPanelTreeNode> Node = GetSelection())
 	{
-		if (Node->GetId() == InEntityId)
+		if (Node->GetRCId() == InEntityId)
 		{
 			OnSelectionChangeDelegate.Broadcast(nullptr);
 		}
@@ -561,11 +531,11 @@ void SRCPanelExposedEntitiesList::OnGroupAdded(const FRemoteControlPresetGroup&
 
 void SRCPanelExposedEntitiesList::OnGroupDeleted(FRemoteControlPresetGroup DeletedGroup)
 {
-	int32 Index = FieldGroups.IndexOfByPredicate([&DeletedGroup](const TSharedPtr<SRCPanelGroup>& Group) { return Group->GetId() == DeletedGroup.Id; });
+	int32 Index = FieldGroups.IndexOfByPredicate([&DeletedGroup](const TSharedPtr<SRCPanelGroup>& Group) { return Group->GetRCId() == DeletedGroup.Id; });
 
 	if (TSharedPtr<SRCPanelTreeNode> Node = GetSelection())
 	{
-		if (DeletedGroup.GetFields().Contains(Node->GetId()))
+		if (DeletedGroup.GetFields().Contains(Node->GetRCId()))
 		{
 			OnSelectionChangeDelegate.Broadcast(nullptr);
 		}
@@ -590,7 +560,7 @@ void SRCPanelExposedEntitiesList::OnGroupOrderChanged(const TArray<FGuid>& Group
 	auto SortFunc = [&IndicesMap]
 	(const TSharedPtr<SRCPanelGroup>& A, const TSharedPtr<SRCPanelGroup>& B)
 	{
-		return IndicesMap.FindChecked(A->GetId()) < IndicesMap.FindChecked(B->GetId());
+		return IndicesMap.FindChecked(A->GetRCId()) < IndicesMap.FindChecked(B->GetRCId());
 	};
 
 	FieldGroups.Sort(SortFunc);
@@ -617,7 +587,7 @@ void SRCPanelExposedEntitiesList::OnFieldDeleted(const FGuid& GroupId, const FGu
 
 void SRCPanelExposedEntitiesList::OnFieldOrderChanged(const FGuid& GroupId, const TArray<FGuid>& Fields)
 {
-	if (TSharedPtr<SRCPanelGroup>* Group = FieldGroups.FindByPredicate([GroupId](const TSharedPtr<SRCPanelGroup>& InGroup) {return InGroup->GetId() == GroupId; }))
+	if (TSharedPtr<SRCPanelGroup>* Group = FieldGroups.FindByPredicate([GroupId](const TSharedPtr<SRCPanelGroup>& InGroup) {return InGroup->GetRCId() == GroupId; }))
 	{
 		// Sort the group's fields according to the fields array.
 		TMap<FGuid, int32> OrderMap;
@@ -633,7 +603,7 @@ void SRCPanelExposedEntitiesList::OnFieldOrderChanged(const FGuid& GroupId, cons
 			[&OrderMap]
 			(const TSharedPtr<SRCPanelTreeNode>& A, const TSharedPtr<SRCPanelTreeNode>& B)
 				{
-					return OrderMap.FindChecked(A->GetId()) < OrderMap.FindChecked(B->GetId());
+					return OrderMap.FindChecked(A->GetRCId()) < OrderMap.FindChecked(B->GetRCId());
 				});
 		}
 	}
diff --git a/Engine/Plugins/VirtualProduction/RemoteControl/Source/RemoteControlUI/Private/UI/SRCPanelExposedEntitiesList.h b/Engine/Plugins/VirtualProduction/RemoteControl/Source/RemoteControlUI/Private/UI/SRCPanelExposedEntitiesList.h
--- a/Engine/Plugins/VirtualProduction/RemoteControl/Source/RemoteControlUI/Private/UI/SRCPanelExposedEntitiesList.h
+++ b/Engine/Plugins/VirtualProduction/RemoteControl/Source/RemoteControlUI/Private/UI/SRCPanelExposedEntitiesList.h
@@ -2,13 +2,13 @@
 
 #pragma once
 
+#include "IRemoteControlUIModule.h"
 #include "RemoteControlPreset.h"
 #include "SRCPanelTreeNode.h"
 #include "Widgets/DeclarativeSyntaxSupport.h"
 #include "Widgets/SCompoundWidget.h"
 #include "Widgets/Views/STreeView.h"
 #include "UObject/StrongObjectPtr.h"
-#include "UObject/WeakObjectPtrTemplates.h"
 
 struct FRCPanelGroup;
 class FRCPanelWidgetRegistry;
@@ -51,10 +51,8 @@ class SRCPanelExposedEntitiesList : public SCompoundWidget
 public:
 	SLATE_BEGIN_ARGS(SRCPanelExposedEntitiesList)
 		: _EditMode(true)
-		, _DisplayValues(true)
 	{}
 		SLATE_ATTRIBUTE(bool, EditMode)
-		SLATE_ARGUMENT(bool, DisplayValues)
 		SLATE_EVENT(FSimpleDelegate, OnEntityListUpdated)
 	SLATE_END_ARGS()
 
@@ -91,11 +89,11 @@ private:
 	/** Handle selection changes. */
 	void OnSelectionChanged(TSharedPtr<SRCPanelTreeNode> Node, ESelectInfo::Type SelectInfo);
 	/** Handlers for drag/drop events. */
-	FReply OnDropOnGroup(const TSharedPtr<FDragDropOperation>& DragDropOperation, const TSharedPtr<SRCPanelTreeNode>& TargetEntity, const TSharedPtr<SRCPanelGroup>& DragTargetGroup);
+	FReply OnDropOnGroup(const TSharedPtr<FDragDropOperation>& DragDropOperation, const TSharedPtr<SRCPanelTreeNode>& TargetEntity, const TSharedPtr<SRCPanelTreeNode>& DragTargetGroup);
 	/** Get the id of the group that holds a particular widget. */
 	FGuid GetGroupId(const FGuid& EntityId);
 	/** Handles group deletion. */
-	void OnDeleteGroup(const TSharedPtr<SRCPanelGroup>& PanelGroup);
+	void OnDeleteGroup(const FGuid& GroupId);
 	/** Select actors in the current level. */
 	void SelectActorsInlevel(const TArray<UObject*>& Objects);
 	//~ Register to engine/editor events in order to correctly update widgets.
@@ -142,8 +140,6 @@ private:
 	FDelegateHandle OnPropertyChangedHandle;
 	/** Delegate called on selected group change. */
 	FOnSelectionChange OnSelectionChangeDelegate;
-	/** Whether to display the values in the list. */
-	bool bDisplayValues = false;
 	/** The column data shared between all tree nodes in order to share a splitter amongst all rows. */
 	FRCColumnSizeData ColumnSizeData;
 	/** The actual width of the right column.  The left column is 1-ColumnWidth */
diff --git a/Engine/Plugins/VirtualProduction/RemoteControl/Source/RemoteControlUI/Private/UI/SRCPanelExposedEntity.cpp b/Engine/Plugins/VirtualProduction/RemoteControl/Source/RemoteControlUI/Private/UI/SRCPanelExposedEntity.cpp
--- a/Engine/Plugins/VirtualProduction/RemoteControl/Source/RemoteControlUI/Private/UI/SRCPanelExposedEntity.cpp
+++ b/Engine/Plugins/VirtualProduction/RemoteControl/Source/RemoteControlUI/Private/UI/SRCPanelExposedEntity.cpp
@@ -2,21 +2,48 @@
 
 #include "SRCPanelExposedEntity.h"
 
+#include "EditorFontGlyphs.h"
 #include "Framework/Application/SlateApplication.h"
 #include "Framework/MultiBox/MultiBoxBuilder.h"
 #include "ActorTreeItem.h"
 #include "GameFramework/Actor.h"
 #include "RemoteControlEntity.h"
+#include "RemoteControlPreset.h"
 #include "RemoteControlField.h"
+#include "RemoteControlPanelStyle.h"
+#include "SRCPanelDragHandle.h"
 #include "SceneOutlinerFilters.h"
 #include "SceneOutlinerModule.h"
 #include "Modules/ModuleManager.h"
 #include "Widgets/Input/SComboButton.h"
 #include "Widgets/Layout/SBox.h"
+#include "Widgets/Layout/SBorder.h"
 #include "Widgets/SBoxPanel.h"
+#include "Widgets/Text/SInlineEditableTextBlock.h"
 
 #define LOCTEXT_NAMESPACE "RemoteControlPanel"
 
+void SRCPanelExposedEntity::Tick(const FGeometry&, const double, const float)
+{
+	if (bNeedsRename)
+	{
+		if (NameTextBox)
+		{
+			NameTextBox->EnterEditingMode();
+		}
+		bNeedsRename = false;
+	}
+}
+
+TSharedPtr<FRemoteControlEntity> SRCPanelExposedEntity::GetEntity() const
+{
+	if (Preset.IsValid())
+	{
+		return Preset->GetExposedEntity(EntityId).Pin();
+	}
+	return nullptr;
+}
+
 TSharedPtr<SWidget> SRCPanelExposedEntity::GetContextMenu()
 {
 	FMenuBuilder MenuBuilder(true, TSharedPtr<const FUICommandList>());
@@ -31,6 +58,21 @@ TSharedPtr<SWidget> SRCPanelExposedEntity::GetContextMenu()
 	return MenuBuilder.MakeWidget();
 }
 
+void SRCPanelExposedEntity::Initialize(const FGuid& InEntityId, URemoteControlPreset* InPreset, const TAttribute<bool>& InbEditMode)
+{
+	EntityId = InEntityId;
+	Preset = InPreset;
+	bEditMode = InbEditMode;
+
+	if (ensure(InPreset))
+	{
+		if (const TSharedPtr<FRemoteControlEntity> RCEntity = InPreset->GetExposedEntity(InEntityId).Pin())
+		{
+			CachedLabel = RCEntity->GetLabel();
+		}
+	}
+}
+
 TSharedRef<SWidget> SRCPanelExposedEntity::CreateInvalidWidget()
 {
 	return SNew(SHorizontalBox)
@@ -55,6 +97,11 @@ TSharedRef<SWidget> SRCPanelExposedEntity::CreateInvalidWidget()
 		];
 }
 
+EVisibility SRCPanelExposedEntity::GetVisibilityAccordingToEditMode(EVisibility NonEditModeVisibility) const
+{
+	return bEditMode.Get() ? EVisibility::Visible : NonEditModeVisibility;
+}
+
 TSharedRef<SWidget> SRCPanelExposedEntity::CreateRebindMenuContent()
 {
 	FSceneOutlinerModule& SceneOutlinerModule = FModuleManager::Get().LoadModuleChecked<FSceneOutlinerModule>("SceneOutliner");
@@ -70,6 +117,31 @@ TSharedRef<SWidget> SRCPanelExposedEntity::CreateRebindMenuContent()
 	];
 }
 
+bool SRCPanelExposedEntity::OnVerifyItemLabelChanged(const FText& InLabel, FText& OutErrorMessage)
+{
+	if (URemoteControlPreset* RCPreset = Preset.Get())
+	{
+		if (InLabel.ToString() != CachedLabel.ToString() && RCPreset->GetExposedEntityId(*InLabel.ToString()).IsValid())
+		{
+			OutErrorMessage = LOCTEXT("NameAlreadyExists", "This name already exists.");
+			return false;
+		}
+	}
+
+	return true;
+}
+
+void SRCPanelExposedEntity::OnLabelCommitted(const FText& InLabel, ETextCommit::Type InCommitInfo)
+{
+	if (URemoteControlPreset* RCPreset = Preset.Get())
+	{
+		FScopedTransaction Transaction(LOCTEXT("ModifyEntityLabel", "Modify exposed entity's label."));
+		RCPreset->Modify();
+		CachedLabel = RCPreset->RenameExposedEntity(EntityId, *InLabel.ToString());
+		NameTextBox->SetText(FText::FromName(CachedLabel));
+	}
+}
+
 void SRCPanelExposedEntity::OnActorSelected(AActor* InActor) const
 {
 	if (TSharedPtr<FRemoteControlEntity> Entity = GetEntity())
@@ -80,6 +152,11 @@ void SRCPanelExposedEntity::OnActorSelected(AActor* InActor) const
 	FSlateApplication::Get().DismissAllMenus();
 }
 
+const FSlateBrush* SRCPanelExposedEntity::GetBorderImage() const
+{
+	return FRemoteControlPanelStyle::Get()->GetBrush("RemoteControlPanel.ExposedFieldBorder");
+}
+
 bool SRCPanelExposedEntity::IsActorSelectable(const AActor* Actor) const
 {
 	if (TSharedPtr<FRemoteControlEntity> Entity = GetEntity())
@@ -89,5 +166,85 @@ bool SRCPanelExposedEntity::IsActorSelectable(const AActor* Actor) const
 	return false;
 }
 
+TSharedRef<SWidget> SRCPanelExposedEntity::CreateEntityWidget(TSharedPtr<SWidget> ValueWidget, const FText& OptionalWarningMessage)
+{
+	FMakeNodeWidgetArgs Args;
+
+	TSharedRef<SBorder> Widget = SNew(SBorder)
+		.Padding(0.0f)
+		.BorderImage_Raw(this, &SRCPanelExposedEntity::GetBorderImage);
+	
+	Args.DragHandle = SNew(SBox)
+		.Visibility_Raw(this, &SRCPanelExposedEntity::GetVisibilityAccordingToEditMode, EVisibility::Collapsed)
+		[
+			SNew(SRCPanelDragHandle<FExposedEntityDragDrop>, GetRCId())
+			.Widget(Widget)
+		];
+
+	Args.NameWidget = SNew(SHorizontalBox)
+		.Clipping(EWidgetClipping::OnDemand)
+		+ SHorizontalBox::Slot()
+		.VAlign(VAlign_Center)
+		.Padding(0.f, 0.f, 2.0f, 0.f)
+		.AutoWidth()
+		[
+			SNew(STextBlock)
+			.Visibility(!OptionalWarningMessage.IsEmpty() ? EVisibility::Visible : EVisibility::Collapsed)
+            .TextStyle(FRemoteControlPanelStyle::Get(), "RemoteControlPanel.Button.TextStyle")
+            .Font(FAppStyle::Get().GetFontStyle("FontAwesome.10"))
+            .ToolTipText(OptionalWarningMessage)
+            .Text(FEditorFontGlyphs::Exclamation_Triangle)
+		]
+		+ SHorizontalBox::Slot()
+		.AutoWidth()
+		[
+			SAssignNew(NameTextBox, SInlineEditableTextBlock)
+			.Text(FText::FromName(CachedLabel))
+			.OnTextCommitted_Raw(this, &SRCPanelExposedEntity::OnLabelCommitted)
+			.OnVerifyTextChanged_Raw(this, &SRCPanelExposedEntity::OnVerifyItemLabelChanged)
+			.IsReadOnly_Lambda([this]() { return !bEditMode.Get(); })
+		];
+
+	Args.RenameButton = SNew(SButton)
+		.Visibility_Raw(this, &SRCPanelExposedEntity::GetVisibilityAccordingToEditMode, EVisibility::Collapsed)
+		.ButtonStyle(FAppStyle::Get(), "FlatButton")
+		.OnClicked_Lambda([this]() {
+			bNeedsRename = true;
+			return FReply::Handled();
+		})
+		[
+			SNew(STextBlock)
+				.TextStyle(FRemoteControlPanelStyle::Get(), "RemoteControlPanel.Button.TextStyle")
+				.Font(FAppStyle::Get().GetFontStyle("FontAwesome.10"))
+				.Text(FText::FromString(FString(TEXT("\xf044"))) /*fa-edit*/)
+		];
+
+	Args.ValueWidget = ValueWidget;
+
+	Args.UnexposeButton = SNew(SButton)
+		.Visibility_Raw(this, &SRCPanelExposedEntity::GetVisibilityAccordingToEditMode, EVisibility::Collapsed)
+		.OnPressed_Raw(this, &SRCPanelExposedEntity::HandleUnexposeEntity)
+		.ButtonStyle(FRemoteControlPanelStyle::Get(), "RemoteControlPanel.UnexposeButton")
+		[
+			SNew(STextBlock)
+			.TextStyle(FRemoteControlPanelStyle::Get(), "RemoteControlPanel.Button.TextStyle")
+			.Font(FAppStyle::Get().GetFontStyle("FontAwesome.10"))
+			.Text(FText::FromString(FString(TEXT("\xf00d"))) /*fa-times*/)
+		];
+
+	Widget->SetContent(MakeNodeWidget(Args));
+	return Widget;
+}
+
+void SRCPanelExposedEntity::HandleUnexposeEntity()
+{
+	if (URemoteControlPreset* RCPreset = Preset.Get())
+	{
+		FScopedTransaction Transaction(LOCTEXT("UnexposeFunction", "Unexpose remote control entity"));
+		RCPreset->Modify();
+		RCPreset->Unexpose(EntityId);
+	}
+}
+
 
-#undef LOCTEXT_NAMESPACE
\ No newline at end of file
+#undef LOCTEXT_NAMESPACE
diff --git a/Engine/Plugins/VirtualProduction/RemoteControl/Source/RemoteControlUI/Private/UI/SRCPanelExposedEntity.h b/Engine/Plugins/VirtualProduction/RemoteControl/Source/RemoteControlUI/Private/UI/SRCPanelExposedEntity.h
--- a/Engine/Plugins/VirtualProduction/RemoteControl/Source/RemoteControlUI/Private/UI/SRCPanelExposedEntity.h
+++ b/Engine/Plugins/VirtualProduction/RemoteControl/Source/RemoteControlUI/Private/UI/SRCPanelExposedEntity.h
@@ -3,34 +3,109 @@
 #pragma once
 
 #include "SRCPanelTreeNode.h"
+#include "RemoteControlPanelStyle.h"
+#include "DragAndDrop/DecoratedDragDropOp.h"
 
 class AActor;
 struct FRemoteControlEntity;
+class URemoteControlPreset;
 class SWidget;
 
+/**
+ * @Note if you inherit from this struct, you must call SRCPanelExposedEntity::Initialize.
+ */
 struct SRCPanelExposedEntity : public SRCPanelTreeNode
 {
-	/**
-	 * Get the underlying exposed entity.
-	 */
-	virtual TSharedPtr<FRemoteControlEntity> GetEntity() const = 0;
-
-	/**
-	 * Get the context menu for this item.
-	 */
+	//~ SWidget interface
+	virtual void Tick(const FGeometry&, const double, const float) override;
+
+	//~ SRCPanelTreeNode interface
+	TSharedPtr<FRemoteControlEntity> GetEntity() const;
 	virtual TSharedPtr<SWidget> GetContextMenu() override;
+	virtual FGuid GetRCId() const override final { return EntityId; }
 
 protected:
+	void Initialize(const FGuid& InEntityId, URemoteControlPreset* InPreset, const TAttribute<bool>& InbEditMode);
+	
 	/** Create a widget that displays the rebind button. */
 	TSharedRef<SWidget> CreateInvalidWidget();
 
+	/** Get the widget's visibility according to the panel's mode. */
+	EVisibility GetVisibilityAccordingToEditMode(EVisibility NonEditModeVisibility) const;
+
+	/** Create an exposed entity widget with a drag handle and unexpose button. */
+	TSharedRef<SWidget> CreateEntityWidget(TSharedPtr<SWidget> ValueWidget, const FText& OptionalWarningMessage = FText::GetEmpty());
+
+protected:
+	/** Id of the entity. */
+	FGuid EntityId;
+	/** The underlying preset. */
+	TWeakObjectPtr<URemoteControlPreset> Preset;
+	/** Whether the panel is in edit mode or not. */
+	TAttribute<bool> bEditMode;
+	/** Display name of the entity. */
+	FName CachedLabel;
+
 private:
 	/** Handles changing the object this entity is bound to upon selecting an actor in the rebinding dropdown. */
 	void OnActorSelected(AActor* InActor) const;
-	
+	/** Get the widget's border. */
+	const FSlateBrush* GetBorderImage() const;
 	/** Create the content of the rebind button.  */
 	TSharedRef<SWidget> CreateRebindMenuContent();
-
+	/** Verifies that the entity's label doesn't already exist. */
+	bool OnVerifyItemLabelChanged(const FText& InLabel, FText& OutErrorMessage);
+	/** Handles committing a entity label. */
+	void OnLabelCommitted(const FText& InLabel, ETextCommit::Type InCommitInfo);
 	/** Returns whether or not the actor is selectable for a binding replacement. */
 	bool IsActorSelectable(const AActor* Parent) const;
+	/** Whether the editable text box for the label needs to enter edit mode. */
+	bool bNeedsRename = false;
+	/** Handle clicking on the unexpose button. */
+	void HandleUnexposeEntity();
+
+	/** The textbox for the row's name. */
+	TSharedPtr<SInlineEditableTextBlock> NameTextBox;
+};
+
+class FExposedEntityDragDrop : public FDecoratedDragDropOp
+{
+public:
+	DRAG_DROP_OPERATOR_TYPE(FExposedFieldDragDropOp, FDecoratedDragDropOp)
+
+	using WidgetType = SWidget;
+
+	FExposedEntityDragDrop(TSharedPtr<SWidget> InWidget, const FGuid& InId)
+		: Id(InId)
+	{
+		DecoratorWidget = SNew(SBorder)
+			.Padding(1.0f)
+			.BorderImage(FRemoteControlPanelStyle::Get()->GetBrush("RemoteControlPanel.ExposedFieldBorder"))
+			.Content()
+			[
+				InWidget.ToSharedRef()
+			];
+
+		Construct();
+	}
+
+	/** Get the ID of the represented entity or group. */
+	FGuid GetId() const
+	{
+		return Id;
+	}
+
+	virtual void OnDrop(bool bDropWasHandled, const FPointerEvent& MouseEvent) override
+	{
+		FDecoratedDragDropOp::OnDrop(bDropWasHandled, MouseEvent);
+	}
+
+	virtual TSharedPtr<SWidget> GetDefaultDecorator() const override
+	{
+		return DecoratorWidget;
+	}
+
+private:
+	FGuid Id;
+	TSharedPtr<SWidget> DecoratorWidget;
 };
\ No newline at end of file
diff --git a/Engine/Plugins/VirtualProduction/RemoteControl/Source/RemoteControlUI/Private/UI/SRCPanelExposedField.cpp b/Engine/Plugins/VirtualProduction/RemoteControl/Source/RemoteControlUI/Private/UI/SRCPanelExposedField.cpp
--- a/Engine/Plugins/VirtualProduction/RemoteControl/Source/RemoteControlUI/Private/UI/SRCPanelExposedField.cpp
+++ b/Engine/Plugins/VirtualProduction/RemoteControl/Source/RemoteControlUI/Private/UI/SRCPanelExposedField.cpp
@@ -4,27 +4,25 @@
 
 #include "Algo/Transform.h"
 #include "EditorFontGlyphs.h"
-#include "EditorStyleSet.h"
-#include "Framework/SlateDelegates.h"
 #include "RCPanelWidgetRegistry.h"
 #include "IDetailTreeNode.h"
+#include "IRemoteControlUIModule.h"
 #include "Layout/Visibility.h"
-#include "Modules/ModuleManager.h"
-#include "PropertyHandle.h"
 #include "RemoteControlField.h"
 #include "RemoteControlPanelStyle.h"
 #include "RemoteControlPreset.h"
 #include "RemoteControlSettings.h"
+#include "Styling/AppStyle.h"
 #include "ScopedTransaction.h"
-#include "SRCPanelDragHandle.h"
 #include "SRCPanelTreeNode.h"
-#include "SRemoteControlPanel.h"
 #include "Styling/SlateBrush.h"
 #include "UObject/Object.h"
+#include "Widgets/Layout/SBorder.h"
 #include "Widgets/Input/SButton.h"
 #include "Widgets/Text/SInlineEditableTextBlock.h"
 #include "Widgets/Text/STextBlock.h"
 
+
 #define LOCTEXT_NAMESPACE "RemoteControlPanel"
 
 namespace ExposedFieldUtils
@@ -59,21 +57,24 @@ namespace ExposedFieldUtils
 	}
 }
 
+TSharedPtr<SRCPanelTreeNode> SRCPanelExposedField::MakeInstance(const FGenerateWidgetArgs& Args)
+{
+	return SNew(SRCPanelExposedField, StaticCastSharedPtr<FRemoteControlField>(Args.Entity), Args.ColumnSizeData, Args.WidgetRegistry).Preset(Args.Preset);
+}
+
 void SRCPanelExposedField::Construct(const FArguments& InArgs, TWeakPtr<FRemoteControlField> InField, FRCColumnSizeData InColumnSizeData, TWeakPtr<FRCPanelWidgetRegistry> InWidgetRegistry)
 {
 	WeakField = MoveTemp(InField);
 
 	ColumnSizeData = MoveTemp(InColumnSizeData);
 	WidgetRegistry = MoveTemp(InWidgetRegistry);
-
-	bEditMode = InArgs._EditMode;
-	Preset = InArgs._Preset;
-	bDisplayValues = InArgs._DisplayValues;
 	
 	if (TSharedPtr<FRemoteControlField> FieldPtr = WeakField.Pin())
 	{
+		Initialize(FieldPtr->GetId(), InArgs._Preset.Get(), InArgs._EditMode);
+	
 		CachedLabel = FieldPtr->GetLabel();
-		FieldId = FieldPtr->GetId();
+		EntityId = FieldPtr->GetId();
 		
 		if (FieldPtr->FieldType == EExposedFieldType::Property)
 		{
@@ -86,34 +87,12 @@ void SRCPanelExposedField::Construct(const FArguments& InArgs, TWeakPtr<FRemoteC
 	}
 }
 
-void SRCPanelExposedField::Tick(const FGeometry&, const double, const float)
-{
-	if (bNeedsRename)
-	{
-		if (NameTextBox)
-		{
-			NameTextBox->EnterEditingMode();
-		}
-		bNeedsRename = false;
-	}
-}
-
 void SRCPanelExposedField::GetNodeChildren(TArray<TSharedPtr<SRCPanelTreeNode>>& OutChildren) const
 {
 	OutChildren.Append(ChildWidgets);
 }
 
-TSharedPtr<SRCPanelExposedField> SRCPanelExposedField::AsField()
-{
-	return SharedThis(this);
-}
-
-FGuid SRCPanelExposedField::GetId() const
-{
-	return GetFieldId();
-}
-
-SRCPanelTreeNode::ENodeType SRCPanelExposedField::GetType() const
+SRCPanelTreeNode::ENodeType SRCPanelExposedField::GetRCType() const
 {
 	return SRCPanelTreeNode::Field;
 }
@@ -123,11 +102,6 @@ FName SRCPanelExposedField::GetFieldLabel() const
 	return CachedLabel;
 }
 
-FGuid SRCPanelExposedField::GetFieldId() const
-{
-	return FieldId;
-}
-
 EExposedFieldType SRCPanelExposedField::GetFieldType() const
 {
 	if (TSharedPtr<FRemoteControlField> Field = WeakField.Pin())
@@ -170,34 +144,31 @@ void SRCPanelExposedField::GetBoundObjects(TSet<UObject*>& OutBoundObjects) cons
 
 TSharedRef<SWidget> SRCPanelExposedField::ConstructWidget()
 {
-	if (bDisplayValues)
+	if (TSharedPtr<FRemoteControlField> Field = WeakField.Pin())
 	{
-		if (TSharedPtr<FRemoteControlField> Field = WeakField.Pin())
+		// For the moment, just use the first object.
+		TArray<UObject*> Objects = Field->GetBoundObjects();
+		if (GetFieldType() == EExposedFieldType::Property && Objects.Num() > 0)
 		{
-			// For the moment, just use the first object.
-			TArray<UObject*> Objects = Field->GetBoundObjects();
-			if (GetFieldType() == EExposedFieldType::Property && Objects.Num() > 0)
+			if (TSharedPtr<FRCPanelWidgetRegistry> Registry = WidgetRegistry.Pin())
 			{
-				if (TSharedPtr<FRCPanelWidgetRegistry> Registry = WidgetRegistry.Pin())
+				if (TSharedPtr<IDetailTreeNode> Node = Registry->GetObjectTreeNode(Objects[0], Field->FieldPathInfo.ToPathPropertyString(), ERCFindNodeMethod::Path))
 				{
-					if (TSharedPtr<IDetailTreeNode> Node = Registry->GetObjectTreeNode(Objects[0], Field->FieldPathInfo.ToPathPropertyString(), ERCFindNodeMethod::Path))
-					{
-						TArray<TSharedRef<IDetailTreeNode>> ChildNodes;
-						Node->GetChildren(ChildNodes);
-						ChildWidgets.Reset(ChildNodes.Num());
-
-						for (const TSharedRef<IDetailTreeNode>& ChildNode : ChildNodes)
-						{
-							ChildWidgets.Add(SNew(SRCPanelFieldChildNode, ChildNode, ColumnSizeData));
-						}
+					TArray<TSharedRef<IDetailTreeNode>> ChildNodes;
+					Node->GetChildren(ChildNodes);
+					ChildWidgets.Reset(ChildNodes.Num());
 
-						return MakeFieldWidget(ExposedFieldUtils::CreateNodeValueWidget(MoveTemp(Node)));
+					for (const TSharedRef<IDetailTreeNode>& ChildNode : ChildNodes)
+					{
+						ChildWidgets.Add(SNew(SRCPanelFieldChildNode, ChildNode, ColumnSizeData));
 					}
+
+					return MakeFieldWidget(ExposedFieldUtils::CreateNodeValueWidget(MoveTemp(Node)));
 				}
 			}
-	
-			return MakeFieldWidget(CreateInvalidWidget());
 		}
+
+		return MakeFieldWidget(CreateInvalidWidget());
 	}
 
 	return MakeFieldWidget(SNullWidget::NullWidget);
@@ -205,8 +176,6 @@ TSharedRef<SWidget> SRCPanelExposedField::ConstructWidget()
 
 TSharedRef<SWidget> SRCPanelExposedField::MakeFieldWidget(const TSharedRef<SWidget>& InWidget)
 {
-	FMakeNodeWidgetArgs Args;
-
 	FText WarningMessage;
 
 	if (GetDefault<URemoteControlSettings>()->bDisplayInEditorOnlyWarnings)
@@ -215,7 +184,7 @@ TSharedRef<SWidget> SRCPanelExposedField::MakeFieldWidget(const TSharedRef<SWidg
             bIsEditableInPackaged = true,
             bIsCallableInPackaged = true;
 		
-		if (TSharedPtr<FRemoteControlField> RCField = GetRemoteControlField().Pin())
+		if (const TSharedPtr<FRemoteControlField> RCField = GetRemoteControlField().Pin())
 		{
 			bIsEditorOnly = RCField->IsEditorOnly();
 			if (RCField->FieldType == EExposedFieldType::Property)
@@ -247,125 +216,7 @@ TSharedRef<SWidget> SRCPanelExposedField::MakeFieldWidget(const TSharedRef<SWidg
 		WarningMessage = Builder.ToText();
 	}
 	
-	Args.DragHandle = SNew(SBox)
-		.Visibility(this, &SRCPanelExposedField::GetVisibilityAccordingToEditMode, EVisibility::Collapsed)
-		[
-			SNew(SRCPanelDragHandle<FExposedEntityDragDrop>, FieldId)
-			.Widget(AsShared())
-		];
-
-	Args.NameWidget = SNew(SHorizontalBox)
-		.Clipping(EWidgetClipping::OnDemand)
-		+ SHorizontalBox::Slot()
-		.VAlign(VAlign_Center)
-		.Padding(0.f, 0.f, 2.0f, 0.f)
-		.AutoWidth()
-		[
-			SNew(STextBlock)
-			.Visibility(!WarningMessage.IsEmpty() ? EVisibility::Visible : EVisibility::Collapsed)
-            .TextStyle(FRemoteControlPanelStyle::Get(), "RemoteControlPanel.Button.TextStyle")
-            .Font(FEditorStyle::Get().GetFontStyle("FontAwesome.10"))
-            .ToolTipText(WarningMessage)
-            .Text(FEditorFontGlyphs::Exclamation_Triangle)
-		]
-		+ SHorizontalBox::Slot()
-		.AutoWidth()
-		[
-			SAssignNew(NameTextBox, SInlineEditableTextBlock)
-			.Text(FText::FromName(CachedLabel))
-			.OnTextCommitted(this, &SRCPanelExposedField::OnLabelCommitted)
-			.OnVerifyTextChanged(this, &SRCPanelExposedField::OnVerifyItemLabelChanged)
-			.IsReadOnly_Lambda([this]() { return !bEditMode.Get(); })
-		];
-
-	Args.RenameButton = SNew(SButton)
-		.Visibility(this, &SRCPanelExposedField::GetVisibilityAccordingToEditMode, EVisibility::Collapsed)
-		.ButtonStyle(FEditorStyle::Get(), "FlatButton")
-		.OnClicked_Lambda([this]() {
-			bNeedsRename = true;
-			return FReply::Handled();
-		})
-		[
-			SNew(STextBlock)
-				.TextStyle(FRemoteControlPanelStyle::Get(), "RemoteControlPanel.Button.TextStyle")
-				.Font(FEditorStyle::Get().GetFontStyle("FontAwesome.10"))
-				.Text(FText::FromString(FString(TEXT("\xf044"))) /*fa-edit*/)
-		];
-
-	Args.ValueWidget = InWidget;
-
-	Args.UnexposeButton = SNew(SButton)
-		.Visibility(this, &SRCPanelExposedField::GetVisibilityAccordingToEditMode, EVisibility::Collapsed)
-		.OnPressed(this, &SRCPanelExposedField::HandleUnexposeField)
-		.ButtonStyle(FRemoteControlPanelStyle::Get(), "RemoteControlPanel.UnexposeButton")
-		[
-			SNew(STextBlock)
-			.TextStyle(FRemoteControlPanelStyle::Get(), "RemoteControlPanel.Button.TextStyle")
-			.Font(FEditorStyle::Get().GetFontStyle("FontAwesome.10"))
-			.Text(FText::FromString(FString(TEXT("\xf00d"))) /*fa-times*/)
-		];
-
-	return SNew(SBorder)
-		.Padding(0.0f)
-		.BorderImage(this, &SRCPanelExposedField::GetBorderImage)
-		[
-			MakeNodeWidget(Args)
-		];
-}
-
-EVisibility SRCPanelExposedField::GetVisibilityAccordingToEditMode(EVisibility NonEditModeVisibility) const
-{
-	return bEditMode.Get() ? EVisibility::Visible : NonEditModeVisibility;
-}
-
-const FSlateBrush* SRCPanelExposedField::GetBorderImage() const
-{
-	return FRemoteControlPanelStyle::Get()->GetBrush("RemoteControlPanel.ExposedFieldBorder");
-}
-
-void SRCPanelExposedField::HandleUnexposeField()
-{
-	if (URemoteControlPreset* RCPreset = Preset.Get())
-	{
-		FText TransactionText;
-		if (GetFieldType() == EExposedFieldType::Function)
-		{
-			TransactionText = LOCTEXT("UnexposeFunction", "Unexpose Function");
-		}
-		else
-		{
-			TransactionText = LOCTEXT("UnexposeProperty", "Unexpose Property");
-		}
-
-		FScopedTransaction Transaction(MoveTemp(TransactionText));
-		RCPreset->Modify();
-		RCPreset->Unexpose(FieldId);
-	}
-}
-
-bool SRCPanelExposedField::OnVerifyItemLabelChanged(const FText& InLabel, FText& OutErrorMessage)
-{
-	if (URemoteControlPreset* RCPreset = Preset.Get())
-	{
-		if (InLabel.ToString() != CachedLabel.ToString() && RCPreset->GetExposedEntityId(*InLabel.ToString()).IsValid())
-		{
-			OutErrorMessage = LOCTEXT("NameAlreadyExists", "This name already exists.");
-			return false;
-		}
-	}
-
-	return true;
-}
-
-void SRCPanelExposedField::OnLabelCommitted(const FText& InLabel, ETextCommit::Type InCommitInfo)
-{
-	if (URemoteControlPreset* RCPreset = Preset.Get())
-	{
-		FScopedTransaction Transaction(LOCTEXT("ModifyFieldLabel", "Modify exposed field's label."));
-		RCPreset->Modify();
-		CachedLabel = RCPreset->RenameExposedEntity(FieldId, *InLabel.ToString());
-		NameTextBox->SetText(FText::FromName(CachedLabel));
-	}
+	return CreateEntityWidget(InWidget, WarningMessage);
 }
 
 void SRCPanelExposedField::ConstructPropertyWidget()
@@ -383,43 +234,36 @@ void SRCPanelExposedField::ConstructFunctionWidget()
 		return;
 	}
 
-	if (TSharedPtr<FRemoteControlFunction> RCFunction = RCPreset->GetExposedEntity<FRemoteControlFunction>(FieldId).Pin())
+	if (const TSharedPtr<FRemoteControlFunction> RCFunction = RCPreset->GetExposedEntity<FRemoteControlFunction>(EntityId).Pin())
 	{
 		if (RCFunction->GetFunction() && RCFunction->GetBoundObjects().Num())
 		{
-			if (bDisplayValues)
+			if (TSharedPtr<FRCPanelWidgetRegistry> Registry = WidgetRegistry.Pin())
 			{
-				if (TSharedPtr<FRCPanelWidgetRegistry> Registry = WidgetRegistry.Pin())
+				Registry->Refresh(RCFunction->FunctionArguments);
+
+				TArray<TSharedPtr<SRCPanelFieldChildNode>> ChildNodes;
+				for (TFieldIterator<FProperty> It(RCFunction->GetFunction()); It; ++It)
 				{
-					Registry->Refresh(RCFunction->FunctionArguments);
+					const bool Param = It->HasAnyPropertyFlags(CPF_Parm);
+					const bool OutParam = It->HasAnyPropertyFlags(CPF_OutParm) && !It->HasAnyPropertyFlags(CPF_ConstParm);
+					const bool ReturnParam = It->HasAnyPropertyFlags(CPF_ReturnParm);
 
-					TArray<TSharedPtr<SRCPanelFieldChildNode>> ChildNodes;
-					for (TFieldIterator<FProperty> It(RCFunction->GetFunction()); It; ++It)
+					if (!Param || OutParam || ReturnParam)
 					{
-						bool bMustHaveParmFlag = !RCFunction->GetFunction()->HasAnyFunctionFlags(FUNC_Native);
-						const bool Param = It->HasAnyPropertyFlags(CPF_Parm);
-						const bool OutParam = It->HasAnyPropertyFlags(CPF_OutParm) && !It->HasAnyPropertyFlags(CPF_ConstParm);
-						const bool ReturnParam = It->HasAnyPropertyFlags(CPF_ReturnParm);
-
-						if (!Param || OutParam || ReturnParam)
-						{
-							continue;
-						}
-
-						if (TSharedPtr<IDetailTreeNode> PropertyNode = Registry->GetStructTreeNode(RCFunction->FunctionArguments, It->GetFName().ToString(), ERCFindNodeMethod::Name))
-						{
-							ChildNodes.Add(SNew(SRCPanelFieldChildNode, PropertyNode.ToSharedRef(), ColumnSizeData));
-						}
+						continue;
 					}
 
-					ChildSlot.AttachWidget(MakeFieldWidget(ConstructCallFunctionButton()));
-					ChildWidgets = ChildNodes;
+					if (TSharedPtr<IDetailTreeNode> PropertyNode = Registry->GetStructTreeNode(RCFunction->FunctionArguments, It->GetFName().ToString(), ERCFindNodeMethod::Name))
+					{
+						ChildNodes.Add(SNew(SRCPanelFieldChildNode, PropertyNode.ToSharedRef(), ColumnSizeData));
+					}
 				}
+
+				ChildSlot.AttachWidget(MakeFieldWidget(ConstructCallFunctionButton()));
+				ChildWidgets = ChildNodes;
 			}
-			else
-			{
-				ChildSlot.AttachWidget(MakeFieldWidget(SNullWidget::NullWidget));
-			}	
+			
 
 			return;
 		}
@@ -454,7 +298,7 @@ FReply SRCPanelExposedField::OnClickFunctionButton()
 
 	if (URemoteControlPreset* RCPreset = Preset.Get())
 	{
-		if (TSharedPtr<FRemoteControlFunction> Function = RCPreset->GetExposedEntity<FRemoteControlFunction>(FieldId).Pin())
+		if (TSharedPtr<FRemoteControlFunction> Function = RCPreset->GetExposedEntity<FRemoteControlFunction>(EntityId).Pin())
 		{
 			for (UObject* Object : Function->GetBoundObjects())
 			{
diff --git a/Engine/Plugins/VirtualProduction/RemoteControl/Source/RemoteControlUI/Private/UI/SRCPanelExposedField.h b/Engine/Plugins/VirtualProduction/RemoteControl/Source/RemoteControlUI/Private/UI/SRCPanelExposedField.h
--- a/Engine/Plugins/VirtualProduction/RemoteControl/Source/RemoteControlUI/Private/UI/SRCPanelExposedField.h
+++ b/Engine/Plugins/VirtualProduction/RemoteControl/Source/RemoteControlUI/Private/UI/SRCPanelExposedField.h
@@ -3,17 +3,15 @@
 #pragma once
 
 #include "CoreMinimal.h"
-#include "DragAndDrop/DecoratedDragDropOp.h"
-#include "EditorStyleSet.h"
 #include "RemoteControlField.h"
 #include "SRCPanelExposedEntity.h"
 #include "SlateFwd.h"
 #include "Widgets/DeclarativeSyntaxSupport.h"
-#include "Widgets/Layout/SBorder.h"
 #include "Widgets/SCompoundWidget.h"
 
 struct EVisibility;
 enum class EExposedFieldType : uint8;
+struct FGenerateWidgetArgs;
 struct FSlateBrush;
 class FRCPanelWidgetRegistry;
 struct FRemoteControlField;
@@ -26,35 +24,28 @@ class URemoteControlPreset;
 /**
  * Widget that displays an exposed field.
  */
-struct SRCPanelExposedField : public SCompoundWidget, public SRCPanelExposedEntity
+struct SRCPanelExposedField : public SRCPanelExposedEntity
 {
 	SLATE_BEGIN_ARGS(SRCPanelExposedField)
 		: _EditMode(true)
 		, _Preset(nullptr)
-		, _DisplayValues(true)
 	{}
 		SLATE_ATTRIBUTE(bool, EditMode)
 		SLATE_ATTRIBUTE(URemoteControlPreset*, Preset)
-		SLATE_ARGUMENT(bool, DisplayValues)
 	SLATE_END_ARGS()
 
-	using SWidget::SharedThis;
-	using SWidget::AsShared;
+	static TSharedPtr<SRCPanelTreeNode> MakeInstance(const FGenerateWidgetArgs& Args);
 
 	void Construct(const FArguments& InArgs, TWeakPtr<FRemoteControlField> Field, FRCColumnSizeData ColumnSizeData, TWeakPtr<FRCPanelWidgetRegistry> InWidgetRegistry);
 
-	void Tick(const FGeometry&, const double, const float);
-
 	//~ SRCPanelTreeNode Interface 
 	virtual void GetNodeChildren(TArray<TSharedPtr<SRCPanelTreeNode>>& OutChildren) const override;
-	virtual TSharedPtr<SRCPanelExposedField> AsField() override;
-	virtual FGuid GetId() const override;
-	virtual SRCPanelTreeNode::ENodeType GetType() const override;
+	virtual ENodeType GetRCType() const override;
 	virtual void Refresh() override;
 	//~ End SRCPanelTreeNode Interface
 
 	//~ SRCPanelExposedEntity Interface
-	virtual TSharedPtr<FRemoteControlEntity> GetEntity() const override { return WeakField.Pin(); }
+	virtual void SetIsHovered(bool bIsBeingHovered) override;
 	//~ End SRCPanelExposedEntity Interface
 
 	/** Get a weak pointer to the underlying remote control field. */
@@ -63,14 +54,9 @@ struct SRCPanelExposedField : public SCompoundWidget, public SRCPanelExposedEnti
 	/** Get this field's label. */
 	FName GetFieldLabel() const;
 
-	/** Get this field's id. */
-	FGuid GetFieldId() const;
-
 	/** Get this field's type. */
 	EExposedFieldType GetFieldType() const;
 
-	/** Set whether this widet is currently hovered when drag and dropping. */
-	void SetIsHovered(bool bInIsHovered);
 
 	/** Returns this widget's underlying objects. */
 	void GetBoundObjects(TSet<UObject*>& OutBoundObjects) const;
@@ -80,16 +66,6 @@ private:
 	TSharedRef<SWidget> ConstructWidget();
 	/** Create the wrapper around the field value widget. */
 	TSharedRef<SWidget> MakeFieldWidget(const TSharedRef<SWidget>& InWidget);
-	/** Get the widget's visibility according to the panel's mode. */
-	EVisibility GetVisibilityAccordingToEditMode(EVisibility NonEditModeVisibility) const;
-	/** Get the widget's border. */
-	const FSlateBrush* GetBorderImage() const;
-	/** Handle clicking on the unexpose button. */
-	void HandleUnexposeField();
-	/** Verifies that the field's label doesn't already exist. */
-	bool OnVerifyItemLabelChanged(const FText& InLabel, FText& OutErrorMessage);
-	/** Handles committing a field label. */
-	void OnLabelCommitted(const FText& InLabel, ETextCommit::Type InCommitInfo);
 	/** Construct this field widget as a property widget. */
 	void ConstructPropertyWidget();
 	/** Construct this field widget as a function widget. */
@@ -105,45 +81,26 @@ private:
 private:
 	/** Weak pointer to the underlying RC Field. */
 	TWeakPtr<FRemoteControlField> WeakField;
-	/** Display name of the field. */
-	FName CachedLabel;
-	/** Id of the field. */
-	FGuid FieldId;
-	/** Whether the row should display its options. */
-	bool bShowOptions = false;
 	/** Whether the widget is currently hovered by a drag and drop operation. */
 	bool bIsHovered = false;
-	/** Whether the editable text box for the label needs to enter edit mode. */
-	bool bNeedsRename = false;
-	/** The widget that displays the field's options ie. Function arguments or metadata. */
-	TSharedPtr<SWidget> OptionsWidget;
 	/** This exposed field's child widgets (ie. An array's rows) */
 	TArray<TSharedPtr<SRCPanelFieldChildNode>> ChildWidgets;
-	/** Whether the panel is in edit mode or not. */
-	TAttribute<bool> bEditMode;
-	/** The underlying preset. */
-	TAttribute<URemoteControlPreset*> Preset;
-	/** The textbox for the row's name. */
-	TSharedPtr<SInlineEditableTextBlock> NameTextBox;
-	/** Whether to display the call function button and the property values. */
-	bool bDisplayValues;
 	/** Holds the panel's cached widgets. */
 	TWeakPtr<FRCPanelWidgetRegistry> WidgetRegistry;
 };
 
 
 /** Represents a child of an exposed field widget. */
-struct SRCPanelFieldChildNode : public SCompoundWidget, public SRCPanelTreeNode
+struct SRCPanelFieldChildNode : public SRCPanelTreeNode
 {
 	SLATE_BEGIN_ARGS(SRCPanelFieldChildNode)
 	{}
 	SLATE_END_ARGS()
 
 	void Construct(const FArguments& InArgs, const TSharedRef<IDetailTreeNode>& InNode, FRCColumnSizeData InColumnSizeData);
-	virtual void GetNodeChildren(TArray<TSharedPtr<SRCPanelTreeNode>>& OutChildren) const { return OutChildren.Append(ChildrenNodes); }
-	virtual FGuid GetId() const { return FGuid(); }
-	virtual ENodeType GetType() const { return SRCPanelTreeNode::FieldChild; }
-	virtual TSharedPtr<SWidget> AsFieldChild() { return AsShared(); }
+	virtual void GetNodeChildren(TArray<TSharedPtr<SRCPanelTreeNode>>& OutChildren) const override { return OutChildren.Append(ChildrenNodes); }
+	virtual FGuid GetRCId() const override { return FGuid(); }
+	virtual ENodeType GetRCType() const override { return SRCPanelTreeNode::FieldChild; }
 
 	TArray<TSharedPtr<SRCPanelFieldChildNode>> ChildrenNodes;
 };
diff --git a/Engine/Plugins/VirtualProduction/RemoteControl/Source/RemoteControlUI/Private/UI/SRCPanelFieldGroup.cpp b/Engine/Plugins/VirtualProduction/RemoteControl/Source/RemoteControlUI/Private/UI/SRCPanelFieldGroup.cpp
--- a/Engine/Plugins/VirtualProduction/RemoteControl/Source/RemoteControlUI/Private/UI/SRCPanelFieldGroup.cpp
+++ b/Engine/Plugins/VirtualProduction/RemoteControl/Source/RemoteControlUI/Private/UI/SRCPanelFieldGroup.cpp
@@ -15,8 +15,6 @@
 #include "Widgets/Input/SButton.h"
 #include "Widgets/Text/SInlineEditableTextBlock.h"
 #include "Widgets/Text/STextBlock.h"
-#include "Widgets/Views/SExpanderArrow.h"
-#include "Widgets/Views/STableViewBase.h"
 #include "Widgets/SBoxPanel.h"
 
 #define LOCTEXT_NAMESPACE "RemoteControlPanel"
@@ -73,7 +71,7 @@ void SRCPanelGroup::Construct(const FArguments& InArgs, URemoteControlPreset* In
 		[
 			SAssignNew(NameTextBox, SInlineEditableTextBlock)
 			.ColorAndOpacity(this, &SRCPanelGroup::GetGroupNameTextColor)
-			.Font(FEditorStyle::GetFontStyle("DetailsView.CategoryFontStyle"))
+			.Font(FAppStyle::Get().GetFontStyle("DetailsView.CategoryFontStyle"))
 			.Text(FText::FromName(Name))
 			.OnTextCommitted(this, &SRCPanelGroup::OnLabelCommitted)
 			.OnVerifyTextChanged(this, &SRCPanelGroup::OnVerifyItemLabelChanged)
@@ -87,7 +85,7 @@ void SRCPanelGroup::Construct(const FArguments& InArgs, URemoteControlPreset* In
 		[
 			SNew(SButton)
 			.Visibility(this, &SRCPanelGroup::GetVisibilityAccordingToEditMode, EVisibility::Collapsed)
-			.ButtonStyle(FEditorStyle::Get(), "FlatButton")
+			.ButtonStyle(FAppStyle::Get(), "FlatButton")
 			.OnClicked_Lambda([this] () 
 				{
 					bNeedsRename = true;
@@ -96,7 +94,7 @@ void SRCPanelGroup::Construct(const FArguments& InArgs, URemoteControlPreset* In
 			[
 				SNew(STextBlock)
 				.TextStyle(FRemoteControlPanelStyle::Get(), "RemoteControlPanel.Button.TextStyle")
-				.Font(FEditorStyle::Get().GetFontStyle("FontAwesome.10"))
+				.Font(FAppStyle::Get().GetFontStyle("FontAwesome.10"))
 				.Text(FText::FromString(FString(TEXT("\xf044"))) /*fa-edit*/)
 			]
 		];
@@ -116,7 +114,7 @@ void SRCPanelGroup::Construct(const FArguments& InArgs, URemoteControlPreset* In
 			[
 				SNew(STextBlock)
 				.TextStyle(FRemoteControlPanelStyle::Get(), "RemoteControlPanel.Button.TextStyle")
-				.Font(FEditorStyle::Get().GetFontStyle("FontAwesome.10"))
+				.Font(FAppStyle::Get().GetFontStyle("FontAwesome.10"))
 				.Text(FText::FromString(FString(TEXT("\xf00d"))) /*fa-times*/)
 			]
 		];
@@ -157,24 +155,19 @@ void SRCPanelGroup::GetNodeChildren(TArray<TSharedPtr<SRCPanelTreeNode>>& OutChi
 	OutChildren.Append(Nodes);
 }
 
-FGuid SRCPanelGroup::GetId() const
+FGuid SRCPanelGroup::GetRCId() const
 {
 	return Id;
 }
 
-SRCPanelTreeNode::ENodeType SRCPanelGroup::GetType() const
+SRCPanelTreeNode::ENodeType SRCPanelGroup::GetRCType() const
 {
 	return SRCPanelTreeNode::Group;
 }
 
-TSharedPtr<SRCPanelGroup> SRCPanelGroup::AsGroup()
-{
-	return StaticCastSharedRef<SRCPanelGroup>(AsShared());
-}
-
 FReply SRCPanelGroup::OnFieldDropGroup(const FDragDropEvent& Event, TSharedPtr<SRCPanelTreeNode> TargetField)
 {
-	if (TSharedPtr<FExposedEntityDragDrop> DragDropOp = Event.GetOperationAs<FExposedEntityDragDrop>())
+	if (const TSharedPtr<FExposedEntityDragDrop> DragDropOp = Event.GetOperationAs<FExposedEntityDragDrop>())
 	{
 		return OnFieldDropGroup(DragDropOp, TargetField);
 	}
@@ -187,11 +180,11 @@ FReply SRCPanelGroup::OnFieldDropGroup(TSharedPtr<FDragDropOperation> DragDropOp
 	{
 		if (DragDropOperation->IsOfType<FExposedEntityDragDrop>() && OnFieldDropEvent.IsBound())
 		{
-			return OnFieldDropEvent.Execute(DragDropOperation, TargetField, AsGroup());
+			return OnFieldDropEvent.Execute(DragDropOperation, TargetField, StaticCastSharedRef<SRCPanelGroup>(AsShared()));
 		}
 		else if (DragDropOperation->IsOfType<FFieldGroupDragDropOp>() && OnFieldDropEvent.IsBound())
 		{
-			return OnFieldDropEvent.Execute(DragDropOperation, nullptr, AsGroup());
+			return OnFieldDropEvent.Execute(DragDropOperation, nullptr, StaticCastSharedRef<SRCPanelGroup>(AsShared()));
 		}
 	}
 
@@ -206,7 +199,7 @@ bool SRCPanelGroup::OnAllowDropFromOtherGroup(TSharedPtr<FDragDropOperation> Dra
 		{
 			if (OnGetGroupId.IsBound())
 			{
-				FGuid OriginGroupId = OnGetGroupId.Execute(DragDropOp->GetId());
+				const FGuid OriginGroupId = OnGetGroupId.Execute(DragDropOp->GetId());
 				if (OriginGroupId != Id)
 				{
 					return true;
@@ -227,7 +220,7 @@ bool SRCPanelGroup::OnAllowDropFromOtherGroup(TSharedPtr<FDragDropOperation> Dra
 
 FReply SRCPanelGroup::HandleDeleteGroup()
 {
-	OnDeleteGroup.ExecuteIfBound(AsGroup());
+	OnDeleteGroup.ExecuteIfBound(Id);
 	return FReply::Handled();
 }
 
diff --git a/Engine/Plugins/VirtualProduction/RemoteControl/Source/RemoteControlUI/Private/UI/SRCPanelFieldGroup.h b/Engine/Plugins/VirtualProduction/RemoteControl/Source/RemoteControlUI/Private/UI/SRCPanelFieldGroup.h
--- a/Engine/Plugins/VirtualProduction/RemoteControl/Source/RemoteControlUI/Private/UI/SRCPanelFieldGroup.h
+++ b/Engine/Plugins/VirtualProduction/RemoteControl/Source/RemoteControlUI/Private/UI/SRCPanelFieldGroup.h
@@ -20,16 +20,16 @@ class SInlineEditableTextBlock;
 class URemoteControlPreset;
 
 /** Widget representing a group. */
-class SRCPanelGroup : public SRCPanelTreeNode, public SCompoundWidget
+class SRCPanelGroup : public SRCPanelTreeNode
 {
 public:
 	using SWidget::SharedThis;
 	using SWidget::AsShared;
 
 public:
-	DECLARE_DELEGATE_RetVal_ThreeParams(FReply, FOnFieldDropEvent, const TSharedPtr<FDragDropOperation>& /* Event */, const TSharedPtr<SRCPanelTreeNode>& /* TargetField */, const TSharedPtr<SRCPanelGroup>& /* DragTargetGroup */);
+	DECLARE_DELEGATE_RetVal_ThreeParams(FReply, FOnFieldDropEvent, const TSharedPtr<FDragDropOperation>& /* Event */, const TSharedPtr<SRCPanelTreeNode>& /* TargetField */, const TSharedPtr<SRCPanelTreeNode>& /* DragTargetGroup */);
 	DECLARE_DELEGATE_RetVal_OneParam(FGuid, FOnGetGroupId, const FGuid& /* EntityId */);
-	DECLARE_DELEGATE_OneParam(FOnDeleteGroup, const TSharedPtr<SRCPanelGroup>&);
+	DECLARE_DELEGATE_OneParam(FOnDeleteGroup, const FGuid& /*GroupId*/);
 
 	SLATE_BEGIN_ARGS(SRCPanelGroup)
 		: _EditMode(true)
@@ -58,9 +58,8 @@ public:
 
 	//~ SRCPanelTreeNode Interface
 	virtual void GetNodeChildren(TArray<TSharedPtr<SRCPanelTreeNode>>& OutChildren) const override;
-	virtual FGuid GetId() const override;
-	virtual ENodeType GetType() const override;
-	virtual TSharedPtr<SRCPanelGroup> AsGroup() override;
+	virtual FGuid GetRCId() const override;
+	virtual ENodeType GetRCType() const override;
 
 private:
 	//~ Handle drag/drop events
@@ -117,7 +116,7 @@ public:
 	{
 		DecoratorWidget = SNew(SBorder)
 			.Padding(0.f)
-			.BorderImage(FEditorStyle::GetBrush("Graph.ConnectorFeedback.Border"))
+			.BorderImage(FAppStyle::Get().GetBrush("Graph.ConnectorFeedback.Border"))
 			.Content()
 			[
 				InWidget.ToSharedRef()
diff --git a/Engine/Plugins/VirtualProduction/RemoteControl/Source/RemoteControlUI/Private/UI/SRCPanelTreeNode.cpp b/Engine/Plugins/VirtualProduction/RemoteControl/Source/RemoteControlUI/Private/UI/SRCPanelTreeNode.cpp
--- a/Engine/Plugins/VirtualProduction/RemoteControl/Source/RemoteControlUI/Private/UI/SRCPanelTreeNode.cpp
+++ b/Engine/Plugins/VirtualProduction/RemoteControl/Source/RemoteControlUI/Private/UI/SRCPanelTreeNode.cpp
@@ -17,7 +17,7 @@ TSharedRef<SWidget> SRCPanelTreeNode::MakeSplitRow(TSharedRef<SWidget> LeftColum
 	RightColumnAttribute.BindRaw(this, &SRCPanelTreeNode::GetRightColumnWidth);
 	
 	return SNew(SSplitter)
-		.Style(FEditorStyle::Get(), "DetailsView.Splitter")
+		.Style(FAppStyle::Get(), "DetailsView.Splitter")
 		.PhysicalSplitterHandleSize(1.0f)
 		.HitDetectionSplitterHandleSize(5.0f)
 		+ SSplitter::Slot()
@@ -67,10 +67,12 @@ TSharedRef<SWidget> SRCPanelTreeNode::MakeNodeWidget(const FMakeNodeWidgetArgs&
 	TSharedRef<SWidget> RightColumn =
 		SNew(SOverlay)
 		+ SOverlay::Slot()
+
 		.HAlign(HAlign_Fill)
 		[
 			SNew(SBox)
 			.HAlign(HAlign_Left)
+			.VAlign(VAlign_Center)
 			[
 				WidgetOrNull(Args.ValueWidget)			
 			]
@@ -81,7 +83,7 @@ TSharedRef<SWidget> SRCPanelTreeNode::MakeNodeWidget(const FMakeNodeWidgetArgs&
 		.VAlign(VAlign_Center)
 		[
 			SNew(SBorder)
-			.BorderImage(FEditorStyle::GetBrush("ToolPanel.GroupBorder"))
+			.BorderImage(FAppStyle::Get().GetBrush("ToolPanel.GroupBorder"))
 			.Padding(0.f)
 			[
 				WidgetOrNull(Args.UnexposeButton)
@@ -99,19 +101,19 @@ void SRCPanelTreeNode::OnLeftColumnResized(float) const
 
 float SRCPanelTreeNode::GetLeftColumnWidth() const
 {
-	const float Offset = GetType() == ENodeType::Group ? SplitterOffset : 0;
+	const float Offset = GetRCType() == ENodeType::Group ? SplitterOffset : 0;
 	return FMath::Clamp(ColumnSizeData.LeftColumnWidth.Get() + Offset, 0.f, 1.f);
 }
 
 float SRCPanelTreeNode::GetRightColumnWidth() const
 {
-	const float Offset = GetType() == ENodeType::Group ? SplitterOffset : 0;
+	const float Offset = GetRCType() == ENodeType::Group ? SplitterOffset : 0;
 	return  FMath::Clamp(ColumnSizeData.RightColumnWidth.Get() - Offset, 0.f, 1.f);
 }
 
 void SRCPanelTreeNode::SetColumnWidth(float InWidth)
 {
-	const float Offset = GetType() == ENodeType::Group ? SplitterOffset : -SplitterOffset;
+	const float Offset = GetRCType() == ENodeType::Group ? SplitterOffset : -SplitterOffset;
 	ColumnSizeData.SetColumnWidth(FMath::Clamp(InWidth + SplitterOffset, 0.f, 1.f));
 }
 
diff --git a/Engine/Plugins/VirtualProduction/RemoteControl/Source/RemoteControlUI/Private/UI/SRCPanelTreeNode.h b/Engine/Plugins/VirtualProduction/RemoteControl/Source/RemoteControlUI/Private/UI/SRCPanelTreeNode.h
--- a/Engine/Plugins/VirtualProduction/RemoteControl/Source/RemoteControlUI/Private/UI/SRCPanelTreeNode.h
+++ b/Engine/Plugins/VirtualProduction/RemoteControl/Source/RemoteControlUI/Private/UI/SRCPanelTreeNode.h
@@ -2,31 +2,19 @@
 
 #pragma once
 
+#include "IRemoteControlUIModule.h"
 #include "Templates/SharedPointer.h"
 #include "Containers/Array.h"
-#include "DragAndDrop/DecoratedDragDropOp.h"
 #include "Misc/Guid.h"
-#include "Widgets/Layout/SSplitter.h"
 
 class SWidget;
 class SRCPanelGroup;
 struct SRCPanelExposedField;
 struct SRCPanelExposedActor;
-
-/**
- * Holds data about a row's columns' size.
- */
-struct FRCColumnSizeData
-{
-	TAttribute<float> LeftColumnWidth;
-	TAttribute<float> RightColumnWidth;
-	SSplitter::FOnSlotResized OnWidthChanged;
-
-	void SetColumnWidth(float InWidth) { OnWidthChanged.ExecuteIfBound(InWidth); }
-};
+struct SRCPanelExposedMaterial;
 
 /** A node in the panel tree view. */
-struct SRCPanelTreeNode
+struct SRCPanelTreeNode : public SCompoundWidget
 {
 	enum ENodeType
 	{
@@ -34,7 +22,8 @@ struct SRCPanelTreeNode
 		Group,
 		Field,
 		FieldChild,
-		Actor
+		Actor,
+		Material
 	};
 
 	virtual ~SRCPanelTreeNode() {}
@@ -42,19 +31,15 @@ struct SRCPanelTreeNode
 	/** Get this tree node's childen. */
 	virtual void GetNodeChildren(TArray<TSharedPtr<SRCPanelTreeNode>>& OutChildren) const {}
 	/** Get this node's ID if any. */
-	virtual FGuid GetId() const { return FGuid(); }
+	virtual FGuid GetRCId() const { return FGuid(); }
 	/** Get get this node's type. */
-	virtual ENodeType GetType() const { return Invalid; };
+	virtual ENodeType GetRCType() const { return ENodeType::Invalid; };
 	/** Refresh the node. */
 	virtual void Refresh() {};
 	/** Get the context menu for this node. */
 	virtual TSharedPtr<SWidget> GetContextMenu() { return nullptr; }
-
-	//~ Utiliy methods for not having to downcast 
-	virtual TSharedPtr<SRCPanelExposedField> AsField() { return nullptr; }
-	virtual TSharedPtr<SWidget> AsFieldChild() { return nullptr; }
-	virtual TSharedPtr<SRCPanelGroup> AsGroup() { return nullptr; }
-	virtual TSharedPtr<SRCPanelExposedActor> AsActor() { return nullptr; }
+	/** Set whether this widget is currently hovered when drag and dropping. */
+	virtual void SetIsHovered(bool bIsBeingHovered) {}
 
 protected:
 	struct FMakeNodeWidgetArgs
@@ -89,46 +74,5 @@ private:
 	static constexpr float SplitterOffset = 0.008f;
 };
 
-class FExposedEntityDragDrop : public FDecoratedDragDropOp
-{
-public:
-	DRAG_DROP_OPERATOR_TYPE(FExposedFieldDragDropOp, FDecoratedDragDropOp)
-
-	using WidgetType = SWidget;
-
-	FExposedEntityDragDrop(TSharedPtr<SWidget> InWidget, FGuid InId)
-		: Id(MoveTemp(InId))
-	{
-		DecoratorWidget = SNew(SBorder)
-			.Padding(1.0f)
-			.BorderImage(FEditorStyle::GetBrush("DetailsView.CategoryMiddle_Active"))
-			.Content()
-			[
-				InWidget.ToSharedRef()
-			];
-
-		Construct();
-	}
-
-	/** Get the ID of the represented entity or group. */
-	FGuid GetId()
-	{
-		return Id;
-	}
-
-	virtual void OnDrop(bool bDropWasHandled, const FPointerEvent& MouseEvent) override
-	{
-		FDecoratedDragDropOp::OnDrop(bDropWasHandled, MouseEvent);
-	}
-
-	virtual TSharedPtr<SWidget> GetDefaultDecorator() const override
-	{
-		return DecoratorWidget;
-	}
-
-private:
-	FGuid Id;
-	TSharedPtr<SWidget> DecoratorWidget;
-};
 
 #undef LOCTEXT_NAMESPACE /*RemoteControlPanelNode*/
diff --git a/Engine/Plugins/VirtualProduction/RemoteControl/Source/RemoteControlUI/Private/UI/SRemoteControlPanel.cpp b/Engine/Plugins/VirtualProduction/RemoteControl/Source/RemoteControlUI/Private/UI/SRemoteControlPanel.cpp
--- a/Engine/Plugins/VirtualProduction/RemoteControl/Source/RemoteControlUI/Private/UI/SRemoteControlPanel.cpp
+++ b/Engine/Plugins/VirtualProduction/RemoteControl/Source/RemoteControlUI/Private/UI/SRemoteControlPanel.cpp
@@ -78,14 +78,14 @@ namespace RemoteControlPanelUtils
 		return GEditor ? GEditor->GetEditorWorldContext(false).World() : nullptr;
 	}
 
-	template <typename EntityType> 
-	TSharedPtr<FStructOnScope> GetEntityOnScope(const TSharedPtr<EntityType>& Entity)
+	TSharedPtr<FStructOnScope> GetEntityOnScope(const TSharedPtr<FRemoteControlEntity>& Entity, const UScriptStruct* EntityType)
 	{
-		static_assert(TIsDerivedFrom<EntityType, FRemoteControlEntity>::Value, "EntityType must derive from FRemoteControlEntity.");
-		if (Entity)
+		if (ensure(Entity && EntityType))
 		{
-			return MakeShared<FStructOnScope>(EntityType::StaticStruct(), reinterpret_cast<uint8*>(Entity.Get()));
+			check(EntityType->IsChildOf(FRemoteControlEntity::StaticStruct()));
+			return MakeShared<FStructOnScope>(EntityType, reinterpret_cast<uint8*>(Entity.Get()));
 		}
+
 		return nullptr;
 	}
 }
@@ -106,7 +106,6 @@ void SRemoteControlPanel::Construct(const FArguments& InArgs, URemoteControlPres
 	EntityProtocolDetails = SNew(SBox);
 	
 	EntityList = SNew(SRCPanelExposedEntitiesList, Preset.Get(), WidgetRegistry)
-		.DisplayValues(true)
 		.OnEntityListUpdated_Lambda([this] ()
 		{
 			UpdateEntityDetailsView(EntityList->GetSelection());
@@ -153,11 +152,11 @@ void SRemoteControlPanel::Construct(const FArguments& InArgs, URemoteControlPres
 			[
 				SNew(SButton)
 				.Visibility_Lambda([this]() { return bIsInEditMode ? EVisibility::Visible : EVisibility::Collapsed; })
-				.ButtonStyle(FEditorStyle::Get(), "FlatButton")
+				.ButtonStyle(FAppStyle::Get(), "FlatButton")
 				.OnClicked(this, &SRemoteControlPanel::OnCreateGroup)
 				[
 					SNew(SImage)
-					.Image(FEditorStyle::GetBrush("SceneOutliner.NewFolderIcon"))
+					.Image(FAppStyle::Get().GetBrush("SceneOutliner.NewFolderIcon"))
 				]
 			]
 			// Function library picker
@@ -233,7 +232,7 @@ void SRemoteControlPanel::Construct(const FArguments& InArgs, URemoteControlPres
 				.AutoWidth()
 				[
 					SAssignNew(PresetNameTextBlock, STextBlock)
-					.Font(FEditorStyle::GetFontStyle("DetailsView.CategoryFontStyle"))
+					.Font(FAppStyle::Get().GetFontStyle("DetailsView.CategoryFontStyle"))
 					.Text(FText::FromName(Preset->GetFName()))
 				]
 				+ SHorizontalBox::Slot()
@@ -248,14 +247,15 @@ void SRemoteControlPanel::Construct(const FArguments& InArgs, URemoteControlPres
 				.AutoWidth()
 				[
 					SNew(SButton)
-					.ButtonStyle(FEditorStyle::Get(), "FlatButton")
+					.ContentPadding(2.0f)
+					.ButtonStyle(FAppStyle::Get(), "FlatButton")
 					.TextStyle(FRemoteControlPanelStyle::Get(), "RemoteControlPanel.Button.TextStyle")
 					.OnClicked_Raw(this, &SRemoteControlPanel::OnClickSettingsButton)
 					[
 						SNew(STextBlock)
 						.ColorAndOpacity(FColor::White)
 						.ToolTipText(LOCTEXT("OpenRemoteControlSettings", "Open Remote Control settings."))
-						.Font(FEditorStyle::Get().GetFontStyle("FontAwesome.10"))
+						.Font(FAppStyle::Get().GetFontStyle("FontAwesome.10"))
 						.Text(FEditorFontGlyphs::Cogs)
 					]
 				]
@@ -270,7 +270,7 @@ void SRemoteControlPanel::Construct(const FArguments& InArgs, URemoteControlPres
 			[
 				SNew(SBorder)
 				.Padding(FMargin(0.f, 5.f, 0.f, 0.f))
-				.BorderImage(FEditorStyle::GetBrush("ToolPanel.DarkGroupBorder"))
+				.BorderImage(FAppStyle::Get().GetBrush("ToolPanel.DarkGroupBorder"))
 				[
 					SNew(SWidgetSwitcher)
 					.WidgetIndex_Lambda([this](){ return !bIsInEditMode ? 0 : 1; })
@@ -278,7 +278,7 @@ void SRemoteControlPanel::Construct(const FArguments& InArgs, URemoteControlPres
 					[
 						// Exposed entities List
 						SNew(SBorder)
-						.BorderImage(FEditorStyle::GetBrush("ToolPanel.GroupBorder"))
+						.BorderImage(FAppStyle::Get().GetBrush("ToolPanel.GroupBorder"))
 						[
 							EntityList.ToSharedRef()
 						]
@@ -299,7 +299,7 @@ void SRemoteControlPanel::Construct(const FArguments& InArgs, URemoteControlPres
 						[
 							// Exposed entities List
 							SNew(SBorder)
-							.BorderImage(FEditorStyle::GetBrush("ToolPanel.GroupBorder"))
+							.BorderImage(FAppStyle::Get().GetBrush("ToolPanel.GroupBorder"))
 							[
 								EntityList.ToSharedRef()
 							]
@@ -314,7 +314,7 @@ void SRemoteControlPanel::Construct(const FArguments& InArgs, URemoteControlPres
 							.SizeRule(SSplitter::SizeToContent)
 							[
 								SNew(SBorder)
-								.BorderImage(FEditorStyle::GetBrush("ToolPanel.GroupBorder"))
+								.BorderImage(FAppStyle::Get().GetBrush("ToolPanel.GroupBorder"))
 								[
 									SNew(SBox)
 									.MinDesiredHeight(40.0f)
@@ -326,7 +326,7 @@ void SRemoteControlPanel::Construct(const FArguments& InArgs, URemoteControlPres
 							+ SSplitter::Slot()
 							[
 								SNew(SBorder)
-								.BorderImage(FEditorStyle::GetBrush("ToolPanel.GroupBorder"))
+								.BorderImage(FAppStyle::Get().GetBrush("ToolPanel.GroupBorder"))
 								[
 									EntityProtocolDetails.ToSharedRef()
 								]
@@ -347,8 +347,8 @@ void SRemoteControlPanel::Construct(const FArguments& InArgs, URemoteControlPres
 	{
 		// We want to insert the widgets before the edit mode buttons.
 		constexpr int32 NumEditModeWidgets = 2;
-		const int32 ExtensionsPosititon = ExtensionWidgets.Num() - NumEditModeWidgets;
-		TopExtensions->InsertSlot(ExtensionsPosititon)
+		const int32 ExtensionsPosition = ExtensionWidgets.Num() - NumEditModeWidgets;
+		TopExtensions->InsertSlot(ExtensionsPosition)
 		.VAlign(VAlign_Center)
 		.AutoWidth()
 		[
@@ -472,9 +472,9 @@ FGuid SRemoteControlPanel::GetSelectedGroup() const
 {
 	if (TSharedPtr<SRCPanelTreeNode> Node = EntityList->GetSelection())
 	{
-		if (Node->AsGroup())
+		if (Node->GetRCType() == SRCPanelTreeNode::Group)
 		{
-			return Node->GetId();		
+			return Node->GetRCId();		
 		}
 	}
 	return FGuid();
@@ -497,14 +497,14 @@ TSharedRef<SWidget> SRemoteControlPanel::CreateCPUThrottleButton() const
 	FText PerformanceWarningText = FText::Format(LOCTEXT("RemoteControlPerformanceWarning", "Warning: The editor setting '{PropertyName}' is currently enabled\nThis will stop editor windows from updating in realtime while the editor is not in focus"), Arguments);
 
 	return SNew(SButton)
-		.ButtonStyle(FEditorStyle::Get(), "FlatButton")
+		.ButtonStyle(FAppStyle::Get(), "FlatButton")
 		.Visibility_Lambda([]() {return GetDefault<UEditorPerformanceSettings>()->bThrottleCPUWhenNotForeground ? EVisibility::Visible : EVisibility::Collapsed; } )
 		.OnClicked_Raw(this, &SRemoteControlPanel::OnClickDisableUseLessCPU)
 		[
 			SNew(STextBlock)
 			.ToolTipText(MoveTemp(PerformanceWarningText))
 			.TextStyle(FRemoteControlPanelStyle::Get(), "RemoteControlPanel.Button.TextStyle")
-			.Font(FEditorStyle::Get().GetFontStyle("FontAwesome.10"))
+			.Font(FAppStyle::Get().GetFontStyle("FontAwesome.10"))
 			.Text(FEditorFontGlyphs::Exclamation_Triangle)
 		];
 }
@@ -587,14 +587,14 @@ TSharedRef<SWidget> SRemoteControlPanel::CreateExposeButton()
 	
 	return SAssignNew(ExposeComboButton, SComboButton)
 		.Visibility_Lambda([this]() { return this->bIsInEditMode ? EVisibility::Visible : EVisibility::Collapsed; })
-		.ButtonStyle(FEditorStyle::Get(), "PropertyEditor.AssetComboStyle")
+		.ButtonStyle(FAppStyle::Get(), "PropertyEditor.AssetComboStyle")
 		.ForegroundColor(FSlateColor::UseForeground())
 		.CollapseMenuOnParentFocus(true)
 		.ContentPadding(FMargin(10.0f, 0.f))
 		.ButtonContent()
 		[
 			SNew(STextBlock)
-			.TextStyle(FEditorStyle::Get(), "ContentBrowser.TopBar.Font")
+			.TextStyle(FAppStyle::Get(), "ContentBrowser.TopBar.Font")
 			.Text(LOCTEXT("ExposeButtonLabel", "Expose"))
 		]
 		.MenuContent()
@@ -937,55 +937,30 @@ TSharedRef<SWidget> SRemoteControlPanel::CreateEntityDetailsView()
 void SRemoteControlPanel::UpdateEntityDetailsView(const TSharedPtr<SRCPanelTreeNode>& SelectedNode)
 {
 	TSharedPtr<FStructOnScope> SelectedEntityPtr;
-	EExposedFieldType FieldType = EExposedFieldType::Invalid;
+
 	if (SelectedNode)
 	{
-		if (const TSharedPtr<SRCPanelExposedField> FieldWidget = SelectedNode->AsField())
+		if (SelectedNode->GetRCType() != SRCPanelTreeNode::Group)
 		{
-			if (const TSharedPtr<FRemoteControlField> Field = FieldWidget->GetRemoteControlField().Pin())
-			{
-				if(Field->GetStruct() == FRemoteControlProperty::StaticStruct())
-				{
-					SelectedEntityPtr = RemoteControlPanelUtils::GetEntityOnScope(StaticCastSharedPtr<FRemoteControlProperty>(Field));
-					FieldType = EExposedFieldType::Property;
-				}
-				else if(Field->GetStruct() == FRemoteControlFunction::StaticStruct())
-				{
-					SelectedEntityPtr = RemoteControlPanelUtils::GetEntityOnScope(StaticCastSharedPtr<FRemoteControlFunction>(Field));
-					FieldType = EExposedFieldType::Function;
-				}
-				else
-				{
-					checkNoEntry();
-				}
-				
-				SelectedEntity = Field;			
-			}
-		}
-		else if (const TSharedPtr<SRCPanelExposedActor> ActorWidget = SelectedNode->AsActor())
-		{
-			if (const TSharedPtr<FRemoteControlActor> Actor = ActorWidget->GetRemoteControlActor().Pin())
-			{
-				SelectedEntity = Actor;
-				SelectedEntityPtr = RemoteControlPanelUtils::GetEntityOnScope<FRemoteControlActor>(Actor);
-			}
+			TSharedPtr<FRemoteControlEntity> Entity = Preset->GetExposedEntity<FRemoteControlEntity>(SelectedNode->GetRCId()).Pin();
+			SelectedEntityPtr = RemoteControlPanelUtils::GetEntityOnScope(Entity, Preset->GetExposedEntityType(SelectedNode->GetRCId()));
 		}
 	}
+
 	if (ensure(EntityDetailsView))
 	{
 		EntityDetailsView->SetStructureData(SelectedEntityPtr);
 	}
 
 	static const FName ProtocolWidgetsModuleName = "RemoteControlProtocolWidgets";	
-	if(SelectedEntity && SelectedNode.IsValid() && FModuleManager::Get().IsModuleLoaded(ProtocolWidgetsModuleName))
+	if(SelectedNode.IsValid() && FModuleManager::Get().IsModuleLoaded(ProtocolWidgetsModuleName) && ensure(Preset.IsValid()))
 	{
-		// If the SelectedNode is valid, the Preset should be too.
-		if(ensure(Preset.IsValid()))
+		if (TSharedPtr<FRemoteControlEntity> RCEntity = Preset->GetExposedEntity(SelectedNode->GetRCId()).Pin())
 		{
-			if(SelectedEntity->IsBound())
+			if(RCEntity->IsBound())
 			{
 				IRemoteControlProtocolWidgetsModule& ProtocolWidgetsModule = FModuleManager::LoadModuleChecked<IRemoteControlProtocolWidgetsModule>(ProtocolWidgetsModuleName);
-				EntityProtocolDetails->SetContent(ProtocolWidgetsModule.GenerateDetailsForEntity(Preset.Get(), SelectedEntity->GetId(), FieldType));	
+				EntityProtocolDetails->SetContent(ProtocolWidgetsModule.GenerateDetailsForEntity(Preset.Get(), RCEntity->GetId()));
 			}
 			else
 			{
diff --git a/Engine/Plugins/VirtualProduction/RemoteControl/Source/RemoteControlUI/Private/UI/SRemoteControlPanel.h b/Engine/Plugins/VirtualProduction/RemoteControl/Source/RemoteControlUI/Private/UI/SRemoteControlPanel.h
--- a/Engine/Plugins/VirtualProduction/RemoteControl/Source/RemoteControlUI/Private/UI/SRemoteControlPanel.h
+++ b/Engine/Plugins/VirtualProduction/RemoteControl/Source/RemoteControlUI/Private/UI/SRemoteControlPanel.h
@@ -208,11 +208,6 @@ private:
 	TSharedPtr<class IStructureDetailsView> EntityDetailsView;
 	/** Holds the field's protocol details. */
 	TSharedPtr<SBox> EntityProtocolDetails;
-	/**
-	 * Pointer to the currently selected node.
-	 * Used in order to ensure that the generated details view keeps a valid pointer to the selected entity.
-	 */
-	TSharedPtr<FRemoteControlEntity> SelectedEntity;
 	/** Whether to show the rebind all button. */
 	bool bShowRebindButton = false;
 	/** Cache of exposed properties. */
diff --git a/Engine/Plugins/VirtualProduction/RemoteControl/Source/RemoteControlUI/Private/UI/SSearchableItemList.h b/Engine/Plugins/VirtualProduction/RemoteControl/Source/RemoteControlUI/Private/UI/SSearchableItemList.h
--- a/Engine/Plugins/VirtualProduction/RemoteControl/Source/RemoteControlUI/Private/UI/SSearchableItemList.h
+++ b/Engine/Plugins/VirtualProduction/RemoteControl/Source/RemoteControlUI/Private/UI/SSearchableItemList.h
@@ -3,18 +3,15 @@
 #pragma once
 
 #include "CoreMinimal.h"
-#include "SlateFwd.h"
 #include "Widgets/SCompoundWidget.h"
 #include "Misc/TextFilter.h"
+#include "Styling/AppStyle.h"
 #include "Widgets/Views/STableRow.h"
 #include "Widgets/Views/SListView.h"
 #include "Widgets/Text/STextBlock.h"
-#include "Algo/Transform.h"
-#include "EditorStyleSet.h"
 #include "Widgets/Input/SSearchBox.h"
 #include "Framework/Application/SlateApplication.h"
 #include "Widgets/Images/SImage.h"
-#include "SSearchableComboBox.h"
 
 /**
  * Widget that displays a searchable dropdown list.
@@ -114,7 +111,7 @@ private:
 			.VAlign(VAlign_Center)
 			[
 				SNew(SImage)
-				.Image(FEditorStyle::Get().GetBrush("GraphEditor.Function_16x"))
+				.Image(FAppStyle::Get().GetBrush("GraphEditor.Function_16x"))
 			]
 			+ SHorizontalBox::Slot()
 			.VAlign(VAlign_Center)
diff --git a/Engine/Plugins/VirtualProduction/RemoteControl/Source/RemoteControlUI/Private/UI/SSearchableTreeView.h b/Engine/Plugins/VirtualProduction/RemoteControl/Source/RemoteControlUI/Private/UI/SSearchableTreeView.h
--- a/Engine/Plugins/VirtualProduction/RemoteControl/Source/RemoteControlUI/Private/UI/SSearchableTreeView.h
+++ b/Engine/Plugins/VirtualProduction/RemoteControl/Source/RemoteControlUI/Private/UI/SSearchableTreeView.h
@@ -3,18 +3,15 @@
 #pragma once
 
 #include "CoreMinimal.h"
-#include "SlateFwd.h"
-#include "Widgets/SCompoundWidget.h"
+#include "Framework/Application/SlateApplication.h"
 #include "Misc/TextFilter.h"
+#include "Styling/AppStyle.h"
+#include "Widgets/SCompoundWidget.h"
 #include "Widgets/Views/STableRow.h"
 #include "Widgets/Views/STreeView.h"
 #include "Widgets/Text/STextBlock.h"
-#include "Algo/Transform.h"
-#include "EditorStyleSet.h"
 #include "Widgets/Input/SSearchBox.h"
-#include "Framework/Application/SlateApplication.h"
 #include "Widgets/Images/SImage.h"
-#include "SSearchableComboBox.h"
 
 /**
  * Widget that displays a searchable tree view.
@@ -138,7 +135,7 @@ private:
 			.VAlign(VAlign_Center)
 			[
 				SNew(SImage)
-				.Image(FEditorStyle::Get().GetBrush("GraphEditor.Function_16x"))
+				.Image(FAppStyle::Get().GetBrush("GraphEditor.Function_16x"))
 			]
 			+ SHorizontalBox::Slot()
 			.VAlign(VAlign_Center)
diff --git a/Engine/Plugins/VirtualProduction/RemoteControl/Source/RemoteControlUI/Public/IRemoteControlUIModule.h b/Engine/Plugins/VirtualProduction/RemoteControl/Source/RemoteControlUI/Public/IRemoteControlUIModule.h
--- a/Engine/Plugins/VirtualProduction/RemoteControl/Source/RemoteControlUI/Public/IRemoteControlUIModule.h
+++ b/Engine/Plugins/VirtualProduction/RemoteControl/Source/RemoteControlUI/Public/IRemoteControlUIModule.h
@@ -6,11 +6,35 @@
 #include "Modules/ModuleInterface.h"
 #include "Templates/SharedPointer.h"
 
+class FDelegateHandle;
+class FRCPanelWidgetRegistry;
+struct FRemoteControlEntity;
 class IPropertyHandle;
 class IDetailCategoryBuilder;
 class IDetailLayoutBuilder;
+struct SRCPanelTreeNode;
 class URemoteControlPreset;
 
+struct FRCColumnSizeData
+{
+	TAttribute<float> LeftColumnWidth;
+	TAttribute<float> RightColumnWidth;
+	SSplitter::FOnSlotResized OnWidthChanged;
+
+	void SetColumnWidth(float InWidth) { OnWidthChanged.ExecuteIfBound(InWidth); }
+};
+
+struct FGenerateWidgetArgs
+{
+	URemoteControlPreset* Preset = nullptr;
+	FRCColumnSizeData ColumnSizeData;
+	TSharedPtr<FRemoteControlEntity> Entity;
+	TAttribute<bool> bIsInEditMode;
+	TWeakPtr<FRCPanelWidgetRegistry> WidgetRegistry;
+};
+
+DECLARE_DELEGATE_RetVal_OneParam(TSharedPtr<SRCPanelTreeNode>, FOnGenerateRCWidget, const FGenerateWidgetArgs& /*Args*/);
+
 DECLARE_MULTICAST_DELEGATE_OneParam(FOnGenerateExtensions, TArray<TSharedRef<class SWidget>>& /*OutExtensions*/);
 
 /**
@@ -79,4 +103,14 @@ public:
 	 * Get the preset currently being edited in the editor.
 	 */
 	virtual URemoteControlPreset* GetActivePreset() const = 0;
+
+	/**
+	 * Register a widget factory to handle creating a widget in the control panel.
+	 */
+	virtual void RegisterWidgetFactoryForType(UScriptStruct* RemoteControlEntityType, const FOnGenerateRCWidget& OnGenerateRCWidgetDelegate) = 0;
+
+	/**
+	 * Unregister a previously registered widget factory.
+	 */
+	virtual void UnregisterWidgetFactoryForType(UScriptStruct* RemoteControlEntityType) = 0;
 };
