Fix non-unity build

diff --git a/Engine/Plugins/VirtualProduction/RemoteControl/Source/RemoteControlUI/Private/UI/SRCPanelExposedEntity.cpp b/Engine/Plugins/VirtualProduction/RemoteControl/Source/RemoteControlUI/Private/UI/SRCPanelExposedEntity.cpp
--- a/Engine/Plugins/VirtualProduction/RemoteControl/Source/RemoteControlUI/Private/UI/SRCPanelExposedEntity.cpp
+++ b/Engine/Plugins/VirtualProduction/RemoteControl/Source/RemoteControlUI/Private/UI/SRCPanelExposedEntity.cpp
@@ -14,7 +14,9 @@
 #include "SRCPanelDragHandle.h"
 #include "SceneOutlinerFilters.h"
 #include "SceneOutlinerModule.h"
+#include "ScopedTransaction.h"
 #include "Modules/ModuleManager.h"
+#include "Widgets/Input/SButton.h"
 #include "Widgets/Input/SComboButton.h"
 #include "Widgets/Layout/SBox.h"
 #include "Widgets/Layout/SBorder.h"
diff --git a/Engine/Plugins/VirtualProduction/RemoteControl/Source/RemoteControlUI/Private/UI/SRCPanelExposedEntity.h b/Engine/Plugins/VirtualProduction/RemoteControl/Source/RemoteControlUI/Private/UI/SRCPanelExposedEntity.h
--- a/Engine/Plugins/VirtualProduction/RemoteControl/Source/RemoteControlUI/Private/UI/SRCPanelExposedEntity.h
+++ b/Engine/Plugins/VirtualProduction/RemoteControl/Source/RemoteControlUI/Private/UI/SRCPanelExposedEntity.h
@@ -3,13 +3,15 @@
 #pragma once
 
 #include "SRCPanelTreeNode.h"
-#include "RemoteControlPanelStyle.h"
 #include "DragAndDrop/DecoratedDragDropOp.h"
+#include "RemoteControlPanelStyle.h"
+#include "Widgets/Layout/SBorder.h"
 
 class AActor;
 struct FRemoteControlEntity;
-class URemoteControlPreset;
+class SInlineEditableTextBlock;
 class SWidget;
+class URemoteControlPreset;
 
 /**
  * @Note if you inherit from this struct, you must call SRCPanelExposedEntity::Initialize.
diff --git a/Engine/Plugins/VirtualProduction/RemoteControl/Source/RemoteControlUI/Private/UI/SRCPanelTreeNode.cpp b/Engine/Plugins/VirtualProduction/RemoteControl/Source/RemoteControlUI/Private/UI/SRCPanelTreeNode.cpp
--- a/Engine/Plugins/VirtualProduction/RemoteControl/Source/RemoteControlUI/Private/UI/SRCPanelTreeNode.cpp
+++ b/Engine/Plugins/VirtualProduction/RemoteControl/Source/RemoteControlUI/Private/UI/SRCPanelTreeNode.cpp
@@ -3,6 +3,7 @@
 #include "SRCPanelTreeNode.h"
 
 #include "Widgets/Layout/SBox.h"
+#include "Widgets/Layout/SBorder.h"
 #include "Widgets/SBoxPanel.h"
 #include "Widgets/SOverlay.h"
 #include "Widgets/Text/STextBlock.h"
diff --git a/Engine/Plugins/VirtualProduction/RemoteControl/Source/RemoteControlUI/Private/UI/SRCPanelTreeNode.h b/Engine/Plugins/VirtualProduction/RemoteControl/Source/RemoteControlUI/Private/UI/SRCPanelTreeNode.h
--- a/Engine/Plugins/VirtualProduction/RemoteControl/Source/RemoteControlUI/Private/UI/SRCPanelTreeNode.h
+++ b/Engine/Plugins/VirtualProduction/RemoteControl/Source/RemoteControlUI/Private/UI/SRCPanelTreeNode.h
@@ -6,6 +6,7 @@
 #include "Templates/SharedPointer.h"
 #include "Containers/Array.h"
 #include "Misc/Guid.h"
+#include "Widgets/SCompoundWidget.h"
 
 class SWidget;
 class SRCPanelGroup;
diff --git a/Engine/Plugins/VirtualProduction/RemoteControl/Source/RemoteControlUI/Public/IRemoteControlUIModule.h b/Engine/Plugins/VirtualProduction/RemoteControl/Source/RemoteControlUI/Public/IRemoteControlUIModule.h
--- a/Engine/Plugins/VirtualProduction/RemoteControl/Source/RemoteControlUI/Public/IRemoteControlUIModule.h
+++ b/Engine/Plugins/VirtualProduction/RemoteControl/Source/RemoteControlUI/Public/IRemoteControlUIModule.h
@@ -2,9 +2,11 @@
 #pragma once
 
 #include "CoreMinimal.h"
+#include "Misc/Attribute.h"
 #include "Modules/ModuleManager.h"
 #include "Modules/ModuleInterface.h"
 #include "Templates/SharedPointer.h"
+#include "Widgets/Layout/SSplitter.h"
 
 class FDelegateHandle;
 class FRCPanelWidgetRegistry;
