Remote Control Protocol bool deserialization correction

diff --git a/Engine/Plugins/VirtualProduction/RemoteControl/Source/RemoteControl/Private/RemoteControlProtocolBinding.cpp b/Engine/Plugins/VirtualProduction/RemoteControl/Source/RemoteControl/Private/RemoteControlProtocolBinding.cpp
--- a/Engine/Plugins/VirtualProduction/RemoteControl/Source/RemoteControl/Private/RemoteControlProtocolBinding.cpp
+++ b/Engine/Plugins/VirtualProduction/RemoteControl/Source/RemoteControl/Private/RemoteControlProtocolBinding.cpp
@@ -45,7 +45,6 @@ namespace EntityInterpolation
 				if(Outer->GetClass() == FArrayProperty::StaticClass())
 				{
 					FScriptArrayHelper Helper(CastField<FArrayProperty>(Outer), RangePair.Mapping.GetData());
-					auto N = Helper.Num();
 					TypedPair.Mapping = *reinterpret_cast<ValueType*>(Helper.GetRawPtr(InArrayIndex));
 				}
 				else
@@ -209,7 +208,6 @@ namespace EntityInterpolation
 		{
 			bool bBoolValue = false;
 			bSuccess = InterpolateValue(BoolProperty, OuterProperty, InRangeMappingBuffers, InProtocolValue, bBoolValue, InArrayIndex);
-			bBoolValue = static_cast<uint8>(bBoolValue) > 0;
 			WritePropertyValue(InCborWriter, InProperty, bBoolValue, !bIsInArray);
 		}
 		else if (FNumericProperty* NumericProperty = CastField<FNumericProperty>(InProperty))
diff --git a/Engine/Plugins/VirtualProduction/RemoteControl/Source/RemoteControlCommon/Public/RCPropertyUtilities.h b/Engine/Plugins/VirtualProduction/RemoteControl/Source/RemoteControlCommon/Public/RCPropertyUtilities.h
--- a/Engine/Plugins/VirtualProduction/RemoteControl/Source/RemoteControlCommon/Public/RCPropertyUtilities.h
+++ b/Engine/Plugins/VirtualProduction/RemoteControl/Source/RemoteControlCommon/Public/RCPropertyUtilities.h
@@ -369,6 +369,24 @@ namespace RemoteControlPropertyUtilities
 		return true;
 	}
 
+	/** Specialization for FBoolProperty. */
+	template <>
+	inline bool Deserialize<FBoolProperty>(const FRCPropertyVariant& InSrc, FRCPropertyVariant& OutDst)
+	{
+		using ValueType = TRemoteControlPropertyTypeTraits<FBoolProperty>::ValueType;
+		
+		TArray<uint8>* SrcPropertyContainer = InSrc.GetPropertyContainer();
+		checkf(SrcPropertyContainer != nullptr, TEXT("Deserialize requires Src to have a backing container."));
+
+		OutDst.Init(InSrc.Size()); // initializes only if necessary
+
+		const ValueType* SrcCurrentValue = InSrc.GetPropertyValue<ValueType>();
+		ValueType* DstCurrentValue = OutDst.GetPropertyValue<ValueType>();
+		*DstCurrentValue = *SrcCurrentValue;
+
+		return true;
+	}
+
 	/** Specialization for FStructProperty. */
 	template <>
 	inline bool Deserialize<FStructProperty>(const FRCPropertyVariant& InSrc, FRCPropertyVariant& OutDst)
