Remove hack in web remote control

diff --git a/Engine/Plugins/VirtualProduction/RemoteControl/Source/WebRemoteControl/Private/RemoteControlReflectionUtils.cpp b/Engine/Plugins/VirtualProduction/RemoteControl/Source/WebRemoteControl/Private/RemoteControlReflectionUtils.cpp
new file mode 100644
--- /dev/null
+++ b/Engine/Plugins/VirtualProduction/RemoteControl/Source/WebRemoteControl/Private/RemoteControlReflectionUtils.cpp
@@ -0,0 +1,119 @@
+// Copyright Epic Games, Inc. All Rights Reserved.
+
+#include "RemoteControlReflectionUtils.h"
+
+#include "Engine/Engine.h"
+#include "IRemoteControlModule.h"
+#include "UObject/StructOnScope.h"
+
+UScriptStruct* UWebRCStructRegistry::GenerateStruct(FName StructureName, const FWebRCGenerateStructArgs& Args)
+{
+	if (UScriptStruct* CachedStruct = CachedStructs.FindRef(StructureName))
+	{
+		return CachedStruct;
+	}
+
+	return CachedStructs.Add(StructureName, GenerateStructInternal(StructureName, Args));
+}
+
+UScriptStruct* UWebRCStructRegistry::GenerateStructInternal(FName StructureName, const FWebRCGenerateStructArgs& Args)
+{
+	UScriptStruct* NewStruct = NewObject<UScriptStruct>(GetTransientPackage(), *StructureName.ToString(), RF_NoFlags);
+
+	for (FName ParamName : Args.StringProperties)
+	{
+		NewStruct->AddCppProperty(new FStrProperty(NewStruct, *ParamName.ToString(), RF_NoFlags));
+	}
+
+	for (const TPair<FName, UScriptStruct*>& Pair : Args.ArrayProperties)
+	{
+		FArrayProperty* NewArrayProperty = new FArrayProperty(NewStruct, Pair.Key, RF_NoFlags);
+
+		FName InnerPathPropName(Pair.Key.ToString() + TEXT("_Inner"));
+		FStructProperty* NewInnerProp = new FStructProperty(NewArrayProperty, InnerPathPropName, RF_NoFlags);
+		NewInnerProp->Struct = Pair.Value;
+		NewArrayProperty->Inner = NewInnerProp;
+
+		NewStruct->AddCppProperty(NewArrayProperty);
+	}
+
+	for (const TPair<FName, FProperty*>& Pair : Args.GenericProperties)
+	{
+		check(Pair.Value);
+		FProperty* DuplicatedField = CastField<FProperty>(FField::Duplicate(Pair.Value, NewStruct, Pair.Key));
+		NewStruct->AddCppProperty(DuplicatedField);
+	}
+
+	for (const TPair<FName, UScriptStruct*>& Pair : Args.StructProperties)
+	{
+		check(Pair.Value);
+		FStructProperty* StructProperty = new FStructProperty(NewStruct, Pair.Key, RF_NoFlags);
+		StructProperty->Struct = Pair.Value;
+		NewStruct->AddCppProperty(StructProperty);
+	}
+
+	// Finalize the struct
+	NewStruct->Bind();
+	NewStruct->PrepareCppStructOps();
+	NewStruct->StaticLink(/*bRelinkExistingProperties*/true);
+
+	return NewStruct;
+}
+
+UScriptStruct* UE::WebRCReflectionUtils::GenerateStruct(FName StructureName, const FWebRCGenerateStructArgs& Args)
+{
+	check(IsInGameThread());
+
+	UWebRCStructRegistry* StructRegistry = GEngine->GetEngineSubsystem<UWebRCStructRegistry>();
+	return StructRegistry->GenerateStruct(StructureName, Args);
+}
+
+void UE::WebRCReflectionUtils::SetStringPropertyValue(FName PropertyName, const FStructOnScope& TargetStruct, const FString& Value)
+{
+	FProperty* Prop = TargetStruct.GetStruct()->FindPropertyByName(PropertyName);
+	check(Prop);
+	*Prop->ContainerPtrToValuePtr<FString>((void*)TargetStruct.GetStructMemory()) = Value;
+}
+
+void UE::WebRCReflectionUtils::CopyPropertyValue(FName PropertyName, const FStructOnScope& TargetStruct, const FRCObjectReference& SourceObject)
+{
+	FProperty* TargetValueProp = TargetStruct.GetStruct()->FindPropertyByName(PropertyName);
+	FProperty* SourceValueProp = SourceObject.Property.Get();
+
+	if (TargetValueProp && SourceValueProp && TargetValueProp->GetClass() == SourceValueProp->GetClass() && TargetValueProp->GetSize() == SourceValueProp->GetSize())
+	{
+		void* DestPtr = TargetValueProp->ContainerPtrToValuePtr<void>((void*)TargetStruct.GetStructMemory());
+		void* SrcPtr = SourceObject.Property->ContainerPtrToValuePtr<void>(SourceObject.ContainerAdress);
+
+		if (DestPtr && SrcPtr)
+		{
+			TargetValueProp->CopyCompleteValue(DestPtr, SrcPtr);
+			return;
+		}
+	}
+
+	checkNoEntry();
+}
+
+void UE::WebRCReflectionUtils::SetStructArrayPropertyValue(FName PropertyName, const FStructOnScope& TargetStruct, const TArray<FStructOnScope>& ArrayElements)
+{
+	FArrayProperty* ArrayProp = CastFieldChecked<FArrayProperty>(TargetStruct.GetStruct()->FindPropertyByName(PropertyName));
+	FScriptArrayHelper_InContainer Helper{ ArrayProp, TargetStruct.GetStructMemory() };
+	for (int32 Index = 0; Index < ArrayElements.Num(); Index++)
+	{
+		const FStructOnScope& ArrayElement = ArrayElements[Index];
+		Helper.AddValue();
+		void* NewPosPtr = Helper.GetRawPtr(Index);
+		((UScriptStruct*)ArrayElement.GetStruct())->CopyScriptStruct(NewPosPtr, ArrayElement.GetStructMemory());
+	}
+}
+
+void UE::WebRCReflectionUtils::SetStructPropertyValue(FName PropertyName, const FStructOnScope& TargetStruct, const UScriptStruct* Struct, void* StructData)
+{
+	FProperty* DestProp = TargetStruct.GetStruct()->FindPropertyByName(PropertyName);
+	check(DestProp);
+	void* DestPtr = DestProp->ContainerPtrToValuePtr<void>((void*)TargetStruct.GetStructMemory());
+	check(DestPtr);
+	check(StructData);
+	Struct->CopyScriptStruct(DestPtr, StructData);
+}
diff --git a/Engine/Plugins/VirtualProduction/RemoteControl/Source/WebRemoteControl/Private/RemoteControlReflectionUtils.h b/Engine/Plugins/VirtualProduction/RemoteControl/Source/WebRemoteControl/Private/RemoteControlReflectionUtils.h
new file mode 100644
--- /dev/null
+++ b/Engine/Plugins/VirtualProduction/RemoteControl/Source/WebRemoteControl/Private/RemoteControlReflectionUtils.h
@@ -0,0 +1,100 @@
+// Copyright Epic Games, Inc. All Rights Reserved.
+
+#pragma once
+
+#include "CoreMinimal.h"
+#include "Subsystems/EngineSubsystem.h"
+#include "UObject/Class.h"
+
+#include "RemoteControlReflectionUtils.generated.h"
+
+struct FRCObjectReference;
+class FStructOnScope;
+
+struct FWebRCGenerateStructArgs
+{
+	/** String properties. */
+	TArray<FName> StringProperties;
+
+	/** Struct properties. */
+	TArray<TPair<FName, UScriptStruct*>> StructProperties;
+
+	/** Properties with an unknown type at compile time. */
+	TArray<TPair<FName, FProperty*>> GenericProperties;
+
+	/** List of array properties and the type they should hold. */
+	TArray<TPair<FName, UScriptStruct*>> ArrayProperties;
+};
+
+/** A subsystem to provide and cache dynamically created ustructs. */
+UCLASS()
+class UWebRCStructRegistry : public UEngineSubsystem
+{
+	GENERATED_BODY()
+	
+public:
+	/** 
+	 * Generates a struct based on provided arguments.
+	 * @param StructureName The name of the structure to generate.
+	 * @param Args the properties that should populate the structure.
+	 * @return The generated structure.
+	 */
+	UScriptStruct* GenerateStruct(FName StructureName, const FWebRCGenerateStructArgs& Args);
+
+private:	
+	/** Creates a new struct based on provided arguments. */
+	UScriptStruct* GenerateStructInternal(FName StructureName, const FWebRCGenerateStructArgs& Args);
+
+private:
+	/** Map of cached structs. */
+	UPROPERTY(Transient)
+	TMap<FName, UScriptStruct*> CachedStructs;
+};
+
+namespace UE
+{
+	namespace WebRCReflectionUtils
+	{
+		/** 
+		 * Generates a struct based on provided arguments.
+		 * @param StructureName The name of the structure to generate.
+		 * @param Args the properties that should populate the structure.
+		 * @return The generated structure.
+		 */
+		UScriptStruct* GenerateStruct(FName StructureName, const FWebRCGenerateStructArgs& Args);
+		
+		/**
+		 * Set a string property value in a provided struct on scope.
+		 * @param PropertyName the name of the property to set.
+		 * @param TargetStruct the structure that contains the property to modify.
+		 * @param Value the new property value.
+		 */
+		void SetStringPropertyValue(FName PropertyName, const FStructOnScope& TargetStruct, const FString& Value);
+
+		/**
+		 * Copy a property value from an object to a struct on scope.
+		 * @param PropertyName the name of the property to copy.
+		 * @param TargetStruct the structure that contains the property to overwrite.
+		 * @param SourceObject the object to copy the property from.
+		 */
+		void CopyPropertyValue(FName PropertyName, const FStructOnScope& TargetStruct, const FRCObjectReference& SourceObject);
+
+		/**
+		 * Sets an array of structures in a struct on scope.
+		 * @param PropertyName the name of the property to set.
+		 * @param TargetStruct the structure that contains the property to modify.
+		 * @param ArrayElements the array that should be copied into the target structure.
+		 */
+		void SetStructArrayPropertyValue(FName PropertyName, const FStructOnScope& TargetStruct, const TArray<FStructOnScope>& ArrayElements);
+
+		/**
+		 * Sets a struct property in the target struct.
+		 * @param PropertyName the name of the property to set.
+		 * @param TargetStruct the structure that contains the property to modify.
+		 * @param Struct the type of the structure to set.
+		 * @param StructData the data that should be copied into the target structure.
+		 */
+		void SetStructPropertyValue(FName PropertyName, const FStructOnScope& TargetStruct, const UScriptStruct* Struct, void* StructData);
+	}
+}
+
diff --git a/Engine/Plugins/VirtualProduction/RemoteControl/Source/WebRemoteControl/Private/WebRemoteControl.cpp b/Engine/Plugins/VirtualProduction/RemoteControl/Source/WebRemoteControl/Private/WebRemoteControl.cpp
--- a/Engine/Plugins/VirtualProduction/RemoteControl/Source/WebRemoteControl/Private/WebRemoteControl.cpp
+++ b/Engine/Plugins/VirtualProduction/RemoteControl/Source/WebRemoteControl/Private/WebRemoteControl.cpp
@@ -6,17 +6,17 @@
 #include "Modules/ModuleManager.h"
 
 #include "IRemoteControlModule.h"
-#include "WebRemoteControlUtils.h"
+#include "RemoteControlReflectionUtils.h"
 #include "RemoteControlRoute.h"
 #include "RemoteControlSettings.h"
 #include "RemoteControlPreset.h"
+#include "WebRemoteControlUtils.h"
 #include "WebSocketMessageHandler.h"
 
 #if WITH_EDITOR
 // Settings
+#include "Editor.h"
 #include "IRemoteControlUIModule.h"
-#include "ISettingsModule.h"
-#include "ISettingsSection.h"
 #include "Framework/Notifications/NotificationManager.h"
 #include "ScopedTransaction.h"
 #include "Widgets/Notifications/SNotificationList.h"
@@ -45,7 +45,6 @@
 #include "RemoteControlRequest.h"
 #include "RemoteControlResponse.h"
 #include "RemoteControlModels.h"
-#include "HttpServerHttpVersion.h"
 
 // Asset registry
 #include "AssetRegistry/IAssetRegistry.h"
@@ -54,16 +53,106 @@
 // Miscelleanous
 #include "Misc/App.h"
 #include "UObject/UnrealType.h"
-#include "UObject/UObjectHash.h"
 #include "Templates/UnrealTemplate.h"
 
 #define LOCTEXT_NAMESPACE "WebRemoteControl"
+
 // Boot the server on startup flag
 static TAutoConsoleVariable<int32> CVarWebControlStartOnBoot(TEXT("WebControl.EnableServerOnStartup"), 0, TEXT("Enable the Web Control servers (web and websocket) on startup."));
 
 // Enable experimental remote routes
 static TAutoConsoleVariable<int32> CVarWebControlEnableExperimentalRoutes(TEXT("WebControl.EnableExperimentalRoutes"), 0, TEXT("Enable the Web Control server experimental routes."));
 
+namespace WebRemoteControlStructUtils
+{
+	FName Struct_PropertyValue = "WebRCPropertyValue";
+	FName Prop_ObjectPath = "ObjectPath";
+	FName Prop_PropertyValue = "PropertyValue";
+
+	FName Struct_GetPropertyResponse = "WebRCGetPropertyResponse";
+	FName Prop_PropertyValues = "PropertyValues";
+	FName Prop_ExposedPropertyDescription = "ExposedPropertyDescription";
+	
+	UScriptStruct* CreatePropertyValueContainer(FProperty* InValueSrcProperty)
+	{
+		check(InValueSrcProperty);
+
+		static FGuid PropertyValueGuid = FGuid::NewGuid();
+
+		FWebRCGenerateStructArgs Args;
+		Args.GenericProperties.Emplace(Prop_PropertyValue, InValueSrcProperty);
+		Args.StringProperties.Add(Prop_ObjectPath);
+		
+		const FString StructName = FString::Format(TEXT("{0}_{1}_{2}_{3}"), { *Struct_PropertyValue.ToString(), *InValueSrcProperty->GetClass()->GetName(), *InValueSrcProperty->GetName(), PropertyValueGuid.ToString() });
+		return UE::WebRCReflectionUtils::GenerateStruct(*StructName, Args);
+	}
+
+	UScriptStruct* CreateActorPropertyValueContainer(FProperty* InValueSrcProperty)
+	{
+		check(InValueSrcProperty);
+
+		static FGuid ActorPropertyValueGuid = FGuid::NewGuid();
+
+		FWebRCGenerateStructArgs Args;
+		Args.GenericProperties.Emplace(Prop_PropertyValue, InValueSrcProperty);
+
+		const FString StructName = FString::Format(TEXT("{0}_{1}_{2}_{4}"), { *Struct_PropertyValue.ToString(), *InValueSrcProperty->GetClass()->GetName(), *InValueSrcProperty->GetName(), ActorPropertyValueGuid.ToString() });
+		return UE::WebRCReflectionUtils::GenerateStruct(*StructName, Args);
+	}
+
+	FStructOnScope CreatePropertyValueOnScope(const TSharedPtr<FRemoteControlProperty>& RCProperty, const FRCObjectReference& ObjectReference)
+	{
+		check(ObjectReference.IsValid());
+
+		UScriptStruct* Struct = CreatePropertyValueContainer(ObjectReference.Property.Get());
+		FStructOnScope StructOnScope{ Struct };
+
+		UE::WebRCReflectionUtils::SetStringPropertyValue(Prop_ObjectPath, StructOnScope, ObjectReference.Object->GetPathName());
+		UE::WebRCReflectionUtils::CopyPropertyValue(Prop_PropertyValue, StructOnScope, ObjectReference);
+
+		return StructOnScope;
+	}
+
+	UScriptStruct* CreateGetPropertyResponseStruct(UScriptStruct* CustomContainer)
+	{
+		check(CustomContainer);
+
+		FWebRCGenerateStructArgs Args;
+		Args.ArrayProperties.Emplace(Prop_PropertyValues, CustomContainer);
+		Args.StructProperties.Emplace(Prop_ExposedPropertyDescription, FRCExposedPropertyDescription::StaticStruct());
+		return UE::WebRCReflectionUtils::GenerateStruct(Struct_GetPropertyResponse, Args);
+	}
+
+	FStructOnScope CreateGetPropertyOnScope(const TSharedPtr<FRemoteControlProperty>& RCProperty, const FRCObjectReference& ObjectReference, FStructOnScope&& PropertyValueOnScope)
+	{
+		check(ObjectReference.IsValid() && RCProperty);
+
+		UScriptStruct* TopLevelStruct = CreateGetPropertyResponseStruct((UScriptStruct*)PropertyValueOnScope.GetStruct());
+		FStructOnScope TopLevelOnScope{ TopLevelStruct };
+
+		FRCExposedPropertyDescription PropertyDescription{ *RCProperty };
+		UE::WebRCReflectionUtils::SetStructPropertyValue(Prop_ExposedPropertyDescription, TopLevelOnScope, FRCExposedPropertyDescription::StaticStruct(), &PropertyDescription);
+
+		TArray<FStructOnScope> PropertyValueArray {};
+		PropertyValueArray.Add(MoveTemp(PropertyValueOnScope));
+		UE::WebRCReflectionUtils::SetStructArrayPropertyValue(Prop_PropertyValues, TopLevelOnScope, PropertyValueArray);
+
+		return TopLevelOnScope;
+	}
+
+	FStructOnScope CreateActorPropertyOnScope(const FRCObjectReference& ObjectReference)
+	{
+		check(ObjectReference.IsValid());
+
+		UScriptStruct* Struct = CreateActorPropertyValueContainer(ObjectReference.Property.Get());
+
+		FStructOnScope StructOnScope{ Struct };
+		UE::WebRCReflectionUtils::CopyPropertyValue(Prop_PropertyValue, StructOnScope, ObjectReference);
+
+		return StructOnScope;
+	}
+}
+
 namespace WebRemoteControl
 {
 	template <typename EntityType>
@@ -77,13 +166,12 @@ namespace WebRemoteControl
 		FGuid Id;
 		if (FGuid::ParseExact(PropertyLabelOrId, EGuidFormats::Digits, Id))
 		{
-        if (TSharedPtr<EntityType> Entity = Preset->GetExposedEntity<EntityType>(Id).Pin())
-        {
-        	return Entity;
-        }
+			if (TSharedPtr<EntityType> Entity = Preset->GetExposedEntity<EntityType>(Id).Pin())
+			{
+        		return Entity;
+			}
 		}
 
-
 		return Preset->GetExposedEntity<EntityType>(Preset->GetExposedEntityId(*PropertyLabelOrId)).Pin();
 	}
 
@@ -92,23 +180,45 @@ namespace WebRemoteControl
 		FGuid Id;
 		if (FGuid::ParseExact(PresetNameOrId, EGuidFormats::Digits, Id))
 		{
-		if (URemoteControlPreset* ResolvedPreset = IRemoteControlModule::Get().ResolvePreset(Id))
-		{
-			return ResolvedPreset;
-		}
+			if (URemoteControlPreset* ResolvedPreset = IRemoteControlModule::Get().ResolvePreset(Id))
+			{
+				return ResolvedPreset;
+			}
 		}
 
 		return IRemoteControlModule::Get().ResolvePreset(*PresetNameOrId);
 	}
+	
+	bool IsWebControlEnabledInEditor()
+	{
+		bool bIsEditor = false;
+
+#if WITH_EDITOR
+		if (GIsEditor)
+		{
+			bIsEditor = true;
+		}
+#endif
+
+		// By default, web remote control is disabled in -game and packaged game.
+		return bIsEditor || FParse::Param(FCommandLine::Get(), TEXT("RCWebControlEnable"));
+	}
 }
 
 void FWebRemoteControlModule::StartupModule()
 {
-	if (FParse::Param(FCommandLine::Get(), TEXT("RCWebControlDisable")))
+	if (FParse::Param(FCommandLine::Get(), TEXT("RCWebControlDisable")) || !FApp::CanEverRender())
 	{
 		return;
 	}
-	
+
+	// By default, disable web remote control in -game and packaged game.
+	if (!WebRemoteControl::IsWebControlEnabledInEditor())
+	{
+		UE_LOG(LogRemoteControl, Display, TEXT("Web remote control is disabled by default when running outside the editor. Use the -RCWebControlEnable flag when launching in order to use it."));
+		return;
+	}
+
 #if WITH_EDITOR
 	RegisterSettings();
 #endif
@@ -123,14 +233,12 @@ void FWebRemoteControlModule::StartupModule()
 	RegisterConsoleCommands();
 	RegisterRoutes();
 
-	const bool bIsHeadless = !FApp::CanEverRender();
-
-	if (!bIsHeadless && (GetDefault<URemoteControlSettings>()->bAutoStartWebServer || CVarWebControlStartOnBoot.GetValueOnAnyThread() > 0))
+	if (GetDefault<URemoteControlSettings>()->bAutoStartWebServer || CVarWebControlStartOnBoot.GetValueOnAnyThread() > 0)
 	{
 		StartHttpServer();
 	}
 
-	if (!bIsHeadless && (GetDefault<URemoteControlSettings>()->bAutoStartWebSocketServer || CVarWebControlStartOnBoot.GetValueOnAnyThread() > 0))
+	if (GetDefault<URemoteControlSettings>()->bAutoStartWebSocketServer || CVarWebControlStartOnBoot.GetValueOnAnyThread() > 0)
 	{
 		StartWebSocketServer();
 	}
@@ -138,7 +246,13 @@ void FWebRemoteControlModule::StartupModule()
 
 void FWebRemoteControlModule::ShutdownModule()
 {
-	if (FParse::Param(FCommandLine::Get(), TEXT("RCWebControlDisable")))
+	if (FParse::Param(FCommandLine::Get(), TEXT("RCWebControlDisable")) || !FApp::CanEverRender())
+	{
+		return;
+	}
+
+	// By default, web remote control is disabled in -game and packaged game.
+	if (!WebRemoteControl::IsWebControlEnabledInEditor())
 	{
 		return;
 	}
@@ -986,51 +1100,17 @@ bool FWebRemoteControlModule::HandlePresetGetPropertyRoute(const FHttpServerRequ
 
 	TArray<uint8> WorkingBuffer;
 	FMemoryWriter Writer(WorkingBuffer);
-	TSharedPtr<TJsonWriter<UCS2CHAR>> JsonWriter = TJsonWriter<UCS2CHAR>::Create(&Writer);
 
 	FRCObjectReference ObjectRef;
 
-	bool bSuccess = true;
-
-	JsonWriter->WriteObjectStart();
-	{
-		{
-			JsonWriter->WriteIdentifierPrefix(TEXT("ExposedPropertyDescription"));
-			FRCJsonStructSerializerBackend SerializeBackend{Writer};
-			FStructSerializer::Serialize(FRCExposedPropertyDescription{*RemoteControlProperty}, SerializeBackend, FStructSerializerPolicies());
-		}
-		
-		// Temporary workaround, reset the state of the json writer
-		JsonWriter->WriteRawJSONValue(TEXT(""));
-		
-		JsonWriter->WriteIdentifierPrefix("PropertyValues");
-		JsonWriter->WriteArrayStart();
-
-		for (UObject* Object : RemoteControlProperty->GetBoundObjects())
-		{
-			IRemoteControlModule::Get().ResolveObjectProperty(ERCAccess::READ_ACCESS, Object, RemoteControlProperty->FieldPathInfo.ToString(), ObjectRef);
-
-			JsonWriter->WriteObjectStart();
-			{
-				JsonWriter->WriteValue(TEXT("ObjectPath"), Object->GetPathName());
-				JsonWriter->WriteIdentifierPrefix(TEXT("PropertyValue"));
-
-				bSuccess &= RemotePayloadSerializer::SerializePartial(
-	                [&ObjectRef] (FJsonStructSerializerBackend& SerializerBackend)
-	                {
-	                    return IRemoteControlModule::Get().GetObjectProperties(ObjectRef, SerializerBackend);
-	                }
-	                , Writer);
-			}
-			JsonWriter->WriteObjectEnd();
-		}
-
-		JsonWriter->WriteArrayEnd();
-	}
-	JsonWriter->WriteObjectEnd();
+	bool bSuccess = IRemoteControlModule::Get().ResolveObjectProperty(ERCAccess::READ_ACCESS, RemoteControlProperty->GetBoundObjects()[0], RemoteControlProperty->FieldPathInfo.ToString(), ObjectRef);
 
 	if (bSuccess)
 	{
+		FStructOnScope PropertyValueOnScope = WebRemoteControlStructUtils::CreatePropertyValueOnScope(RemoteControlProperty, ObjectRef);
+		FStructOnScope FinalStruct = WebRemoteControlStructUtils::CreateGetPropertyOnScope(RemoteControlProperty, ObjectRef, MoveTemp(PropertyValueOnScope));
+		WebRemoteControlUtils::SerializeStructOnScope(FinalStruct, Writer);
+
 		Response->Code = EHttpServerResponseCodes::Ok;
 		WebRemoteControlUtils::ConvertToUTF8(WorkingBuffer, Response->Body);
 	}
@@ -1068,30 +1148,14 @@ bool FWebRemoteControlModule::HandlePresetGetExposedActorPropertyRoute(const FHt
 		{
 			TArray<uint8> WorkingBuffer;
 			FMemoryWriter Writer(WorkingBuffer);
-			TSharedPtr<TJsonWriter<UCS2CHAR>> JsonWriter = TJsonWriter<UCS2CHAR>::Create(&Writer);
 
 			FRCObjectReference ObjectRef;
 
-			bool bSuccess = true;
-
-			bSuccess &= IRemoteControlModule::Get().ResolveObjectProperty(ERCAccess::READ_ACCESS, Actor, FieldPath, ObjectRef);
-
-			FRCJsonStructSerializerBackend SerializerBackend(Writer);
-
-			JsonWriter->WriteObjectStart();
-			JsonWriter->WriteIdentifierPrefix(TEXT("PropertyValue"));
-
-			bSuccess &= RemotePayloadSerializer::SerializePartial(
-				[&ObjectRef](FJsonStructSerializerBackend& SerializerBackend)
-				{
-					return IRemoteControlModule::Get().GetObjectProperties(ObjectRef, SerializerBackend);
-				}
-			, Writer);
-
-			JsonWriter->WriteObjectEnd();
-			
-			if (bSuccess)
+			if (IRemoteControlModule::Get().ResolveObjectProperty(ERCAccess::READ_ACCESS, Actor, FieldPath, ObjectRef))
 			{
+				FStructOnScope ActorPropertyOnScope = WebRemoteControlStructUtils::CreateActorPropertyOnScope(ObjectRef);
+				WebRemoteControlUtils::SerializeStructOnScope(ActorPropertyOnScope, Writer);
+			
 				Response->Code = EHttpServerResponseCodes::Ok;
 				WebRemoteControlUtils::ConvertToUTF8(WorkingBuffer, Response->Body);
 			}
diff --git a/Engine/Plugins/VirtualProduction/RemoteControl/Source/WebRemoteControl/Private/WebRemoteControlUtils.cpp b/Engine/Plugins/VirtualProduction/RemoteControl/Source/WebRemoteControl/Private/WebRemoteControlUtils.cpp
--- a/Engine/Plugins/VirtualProduction/RemoteControl/Source/WebRemoteControl/Private/WebRemoteControlUtils.cpp
+++ b/Engine/Plugins/VirtualProduction/RemoteControl/Source/WebRemoteControl/Private/WebRemoteControlUtils.cpp
@@ -123,38 +123,6 @@ namespace RemotePayloadSerializer
 		JsonWriter->WriteObjectEnd();
 	}
 
-	bool SerializePartial(TFunctionRef<bool(FRCJsonStructSerializerBackend&)> SerializeFunction, FMemoryWriter& SerializedPayloadWriter)
-	{
-		TArray<uint8> WorkingBuffer;
-		FMemoryWriter TemporaryBufferWriter(WorkingBuffer);
-		FRCJsonStructSerializerBackend TemporaryBackend(TemporaryBufferWriter, FRCJsonStructSerializerBackend::DefaultSerializerFlags);
-
-		int32 ColonPosition = -1;
-		int32 LastBracketPosition = -1;
-
-		bool bSuccess = SerializeFunction(TemporaryBackend);
-
-		FStringView BufferView = FStringView((TCHAR*)WorkingBuffer.GetData(), WorkingBuffer.Num() / sizeof(TCHAR));
-		if (bSuccess)
-		{
-			BufferView.FindChar(TEXT(':'), ColonPosition);
-			BufferView.FindLastChar(TEXT('}'), LastBracketPosition);
-			ColonPosition++;
-		}
-
-		if (LastBracketPosition > ColonPosition)
-		{
-			BufferView.MidInline(ColonPosition, LastBracketPosition - ColonPosition);
-			SerializedPayloadWriter.Serialize((uint8*)BufferView.GetData(), BufferView.Len() * sizeof(TCHAR) / sizeof(uint8));
-		}
-		else
-		{
-			bSuccess = false;
-		}
-
-		return bSuccess;
-	}
-
 	bool DeserializeCall(const FHttpServerRequest& InRequest, FRCCall& OutCall, const FHttpResultCallback& InCompleteCallback)
 	{
 		// Create Json reader to read the payload, the payload will already be validated as being in TCHAR
diff --git a/Engine/Plugins/VirtualProduction/RemoteControl/Source/WebRemoteControl/Private/WebRemoteControlUtils.h b/Engine/Plugins/VirtualProduction/RemoteControl/Source/WebRemoteControl/Private/WebRemoteControlUtils.h
--- a/Engine/Plugins/VirtualProduction/RemoteControl/Source/WebRemoteControl/Private/WebRemoteControlUtils.h
+++ b/Engine/Plugins/VirtualProduction/RemoteControl/Source/WebRemoteControl/Private/WebRemoteControlUtils.h
@@ -37,22 +37,6 @@ namespace RemotePayloadSerializer
 	 */
 	void SerializeWrappedCallResponse(int32 RequestId, TUniquePtr<FHttpServerResponse> Response, FMemoryWriter& Writer);
 
-	/**
-	 * Calls a the provided serializer function then strips the identifier from the resulting JSON.
-	 * ie. A returned payload like:
-	 * {
-	 *   "RelativeLocation": {
-	 *		"prop": 2
-	 *    }
-	 * }
-	 *
-	 * Would result in:
-	 * "prop": 2
-	 *
-	 * @todo Modify struct serializer to allow serializing without the identifier since this method is pretty inefficient.
-	 */
-	bool SerializePartial(TFunctionRef<bool(FRCJsonStructSerializerBackend&)> SerializeFunction, FMemoryWriter& SerializedPayloadWriter);
-
 	bool DeserializeCall(const FHttpServerRequest& InRequest, FRCCall& OutCall, const FHttpResultCallback& InCompleteCallback);
 
 	bool SerializeCall(const FRCCall& InCall, TArray<uint8>& OutPayload, bool bOnlyReturn = false);
@@ -320,4 +304,17 @@ namespace WebRemoteControlUtils
 	 * @return Whether or not the content type matches the target content type.
 	 */
 	bool IsRequestContentType(const FHttpServerRequest& InRequest, const FString& InContentType, FString* OutErrorText);
+
+	/**
+	 * Serialize a struct on scope.
+	 * @param Struct the struct on scope to serialize.
+	 * @param Writer the memory archive to write to.
+	 */
+	template <typename SerializerBackendType = FRCJsonStructSerializerBackend>
+	void SerializeStructOnScope(const FStructOnScope& Struct, FMemoryWriter& Writer)
+	{
+		static_assert(TIsDerivedFrom<SerializerBackendType, IStructSerializerBackend>::IsDerived, "SerializerBackendType must inherit from IStructSerializerBackend.");
+		SerializerBackendType SerializerBackend(Writer);
+		FStructSerializer::Serialize(Struct.GetStructMemory(), *(UScriptStruct*)Struct.GetStruct(), SerializerBackend, FStructSerializerPolicies());
+	}
 }
diff --git a/Engine/Plugins/VirtualProduction/RemoteControl/Source/WebRemoteControl/Private/WebSocketMessageHandler.cpp b/Engine/Plugins/VirtualProduction/RemoteControl/Source/WebRemoteControl/Private/WebSocketMessageHandler.cpp
--- a/Engine/Plugins/VirtualProduction/RemoteControl/Source/WebRemoteControl/Private/WebSocketMessageHandler.cpp
+++ b/Engine/Plugins/VirtualProduction/RemoteControl/Source/WebRemoteControl/Private/WebSocketMessageHandler.cpp
@@ -12,11 +12,198 @@
 #include "RemoteControlRequest.h"
 #include "RemoteControlResponse.h"
 #include "RemoteControlRoute.h"
+#include "RemoteControlReflectionUtils.h"
 #include "WebRemoteControl.h"
 #include "WebRemoteControlUtils.h"
 
 static TAutoConsoleVariable<int32> CVarWebRemoteControlFramesBetweenPropertyNotifications(TEXT("WebControl.FramesBetweenPropertyNotifications"), 5, TEXT("The number of frames between sending batches of property notifications."));
 
+namespace WebSocketMessageHandlerStructUtils
+{
+	using namespace UE::WebRCReflectionUtils;
+
+	FName Struct_PropertyValue = "WEBRC_PropertyValue";
+	FName Prop_PropertyLabel = "PropertyLabel";
+	FName Prop_Id = "Id";
+	FName Prop_ObjectPath = "ObjectPath";
+	FName Prop_PropertyValue = "PropertyValue";
+
+	FName Struct_PresetFieldsChanged = "WEBRC_PresetFieldsChanged";
+	FName Prop_Type = "Type";
+	FName Prop_PresetName= "PresetName";
+	FName Prop_PresetId = "PresetId";
+	FName Prop_ChangedFields = "ChangedFields";
+
+	FName Struct_ActorPropertyValue= "WEBRC_ActorPropertyValue";
+	FName Prop_PropertyName = "PropertyName";
+
+	FName Struct_ModifiedActor = "WEBRC_ModifiedActor";
+	FName Prop_DisplayName = "DisplayName";
+	FName Prop_Path = "Path";
+	FName Prop_ModifiedProperties = "ModifiedProperties";
+
+	FName Struct_ModifiedActors = "WEBRC_ModifiedActors";
+	FName Prop_ModifiedActors = "ModifiedActors";
+	
+	UScriptStruct* CreatePropertyValueContainer(FProperty* InValueProperty)
+	{
+		check(InValueProperty);
+
+		static FGuid PropertyValueGuid = FGuid::NewGuid();
+
+		FWebRCGenerateStructArgs Args;
+
+		Args.StringProperties = 
+		{ 
+			Prop_PropertyLabel,
+			Prop_Id,
+			Prop_ObjectPath
+		};
+
+		Args.GenericProperties.Emplace(Prop_PropertyValue, InValueProperty);
+
+		const FString StructName = FString::Format(TEXT("{0}_{1}_{2}_{3}"), { *Struct_PropertyValue.ToString(), *InValueProperty->GetClass()->GetName(), *InValueProperty->GetName(), PropertyValueGuid.ToString() });
+
+		return GenerateStruct(*StructName, Args);
+	}
+
+	UScriptStruct* CreatePresetFieldsChangedStruct(UScriptStruct* PropertyValueStruct)
+	{
+		FWebRCGenerateStructArgs Args;
+		Args.StringProperties =
+		{ 
+			Prop_PresetId,
+			Prop_PresetName,
+			Prop_Type
+		};
+
+		Args.ArrayProperties.Emplace(Prop_ChangedFields, PropertyValueStruct);
+
+		return GenerateStruct(Struct_PresetFieldsChanged, Args);
+	}
+
+	UScriptStruct* CreateActorPropertyValueContainer(FProperty* InValueProperty)
+	{
+		static FGuid ActorPropertyValueGuid = FGuid::NewGuid();
+
+		FWebRCGenerateStructArgs Args;
+		Args.StringProperties =
+		{
+			Prop_PropertyName
+		};
+
+		Args.GenericProperties.Emplace(Prop_PropertyValue, InValueProperty);
+
+		const FString StructName = FString::Format(TEXT("{0}_{1}_{2}_{3}"), { *Struct_PropertyValue.ToString(), *InValueProperty->GetClass()->GetName(), *InValueProperty->GetName(), ActorPropertyValueGuid.ToString() });
+
+		return GenerateStruct(*StructName, Args);
+	}
+
+	UScriptStruct* CreateModifiedActorStruct(UScriptStruct* ModifiedPropertiesStruct)
+	{
+		FWebRCGenerateStructArgs Args;
+		Args.StringProperties =
+		{
+			Prop_Id,
+			Prop_DisplayName,
+			Prop_Path
+		};
+
+		Args.ArrayProperties.Emplace(Prop_ModifiedProperties, ModifiedPropertiesStruct);
+
+		return GenerateStruct(Struct_ModifiedActor, Args);
+	}
+
+	UScriptStruct* CreateModifiedActorsStruct(UScriptStruct* ModifiedActorStruct)
+	{
+		FWebRCGenerateStructArgs Args;
+		Args.StringProperties =
+		{
+			Prop_Type,
+			Prop_PresetName,
+			Prop_PresetId
+		};
+
+		Args.ArrayProperties.Emplace(Prop_ModifiedActors, ModifiedActorStruct);
+
+		return GenerateStruct(Struct_ModifiedActors, Args);
+	}
+
+	FStructOnScope CreatePropertyValueOnScope(const TSharedPtr<FRemoteControlProperty>& RCProperty, const FRCObjectReference& ObjectReference)
+	{
+		UScriptStruct* Struct = CreatePropertyValueContainer(ObjectReference.Property.Get());
+		FStructOnScope StructOnScope{ Struct };
+
+		SetStringPropertyValue(Prop_PropertyLabel, StructOnScope, RCProperty->GetLabel().ToString());
+		SetStringPropertyValue(Prop_Id, StructOnScope, RCProperty->GetId().ToString());
+		SetStringPropertyValue(Prop_ObjectPath, StructOnScope, ObjectReference.Object->GetPathName());
+		CopyPropertyValue(Prop_PropertyValue, StructOnScope, ObjectReference);
+
+		return StructOnScope;
+	}
+
+	FStructOnScope CreatePresetFieldsChangedStructOnScope(const URemoteControlPreset* Preset, const TArray<FStructOnScope>& PropertyValuesOnScope)
+	{
+		UScriptStruct* PropertyValueStruct = (UScriptStruct*)PropertyValuesOnScope[0].GetStruct();
+		check(PropertyValueStruct);
+
+		UScriptStruct* TopLevelStruct = CreatePresetFieldsChangedStruct(PropertyValueStruct);
+
+		FStructOnScope FieldsChangedOnScope{ TopLevelStruct };
+		SetStringPropertyValue(Prop_Type, FieldsChangedOnScope, TEXT("PresetFieldsChanged"));
+		SetStringPropertyValue(Prop_PresetName, FieldsChangedOnScope, *Preset->GetFName().ToString());
+		SetStringPropertyValue(Prop_PresetId, FieldsChangedOnScope, *Preset->GetPresetId().ToString());
+		SetStructArrayPropertyValue(Prop_ChangedFields, FieldsChangedOnScope, PropertyValuesOnScope);
+
+		return FieldsChangedOnScope;
+	}
+
+	FStructOnScope CreateActorPropertyValueOnScope(const URemoteControlPreset* Preset, const FRCObjectReference& ObjectReference)
+	{
+		UScriptStruct* Struct = CreateActorPropertyValueContainer(ObjectReference.Property.Get());
+		FStructOnScope StructOnScope{ Struct };
+
+		SetStringPropertyValue(Prop_PropertyName, StructOnScope, *ObjectReference.Property->GetName());
+		CopyPropertyValue(Prop_PropertyValue, StructOnScope, ObjectReference);
+
+		return StructOnScope;
+	}
+
+	FStructOnScope CreateModifiedActorStructOnScope(const URemoteControlPreset* Preset, const FRemoteControlActor& RCActor, const TArray<FStructOnScope>& ModifiedPropertiesOnScope)
+	{
+		check(ModifiedPropertiesOnScope.Num() > 0);
+		UScriptStruct* ModifiedPropertiesStruct = (UScriptStruct*)ModifiedPropertiesOnScope[0].GetStruct();
+		check(ModifiedPropertiesStruct);
+
+		UScriptStruct* TopLevelStruct = CreateModifiedActorStruct(ModifiedPropertiesStruct);
+		FStructOnScope FieldsChangedOnScope{ TopLevelStruct };
+
+		SetStringPropertyValue(Prop_Id, FieldsChangedOnScope, *RCActor.GetId().ToString());
+		SetStringPropertyValue(Prop_DisplayName, FieldsChangedOnScope, *RCActor.GetLabel().ToString());
+		SetStringPropertyValue(Prop_Path, FieldsChangedOnScope, *RCActor.Path.ToString());
+		SetStructArrayPropertyValue(Prop_ModifiedProperties, FieldsChangedOnScope, ModifiedPropertiesOnScope);
+
+		return FieldsChangedOnScope;
+	}
+
+	FStructOnScope CreateModifiedActorsStructOnScope(const URemoteControlPreset* Preset, const TArray<FStructOnScope>& ModifiedActorsOnScope)
+	{
+		check(ModifiedActorsOnScope.Num() > 0);
+		UScriptStruct* ModifiedActorStruct = (UScriptStruct*)ModifiedActorsOnScope[0].GetStruct();
+		check(ModifiedActorStruct);
+
+		UScriptStruct* TopLevelStruct = CreateModifiedActorsStruct(ModifiedActorStruct);
+		FStructOnScope FieldsChangedOnScope{ TopLevelStruct };
+
+		SetStringPropertyValue(Prop_Type, FieldsChangedOnScope, TEXT("PresetActorModified"));
+		SetStringPropertyValue(Prop_PresetName, FieldsChangedOnScope, *Preset->GetFName().ToString());
+		SetStringPropertyValue(Prop_PresetId, FieldsChangedOnScope, *Preset->GetPresetId().ToString());
+		SetStructArrayPropertyValue(Prop_ModifiedActors, FieldsChangedOnScope, ModifiedActorsOnScope);
+
+		return FieldsChangedOnScope;
+	}
+}
+
 FWebSocketMessageHandler::FWebSocketMessageHandler(FRCWebSocketServer* InServer, const FGuid& InActingClientId)
 	: Server(InServer)
 	, ActingClientId(InActingClientId)
@@ -229,7 +416,9 @@ void FWebSocketMessageHandler::ProcessChangedActorProperties()
 		for (const TPair<FGuid, TMap<FRemoteControlActor, TArray<FRCObjectReference>>>& ClientToModifications : Entry.Value)
 		{
 			TArray<uint8> WorkingBuffer;
-			if (ClientToModifications.Value.Num() && WriteActorPropertyChangePayload(Preset, ClientToModifications.Value, WorkingBuffer))
+			FMemoryWriter Writer(WorkingBuffer);
+
+			if (ClientToModifications.Value.Num() && WriteActorPropertyChangePayload(Preset, ClientToModifications.Value, Writer))
 			{
 				TArray<uint8> Payload;
 				WebRemoteControlUtils::ConvertToUTF8(WorkingBuffer, Payload);
@@ -606,141 +795,64 @@ bool FWebSocketMessageHandler::WritePropertyChangeEventPayload(URemoteControlPre
 {
 	bool bHasProperty = false;
 
-	FMemoryWriter Writer(OutBuffer);
-	TSharedPtr<TJsonWriter<UCS2CHAR>> JsonWriter = TJsonWriter<UCS2CHAR>::Create(&Writer);
-
-	//Might be a better idea to have defined structures for our web socket notification messages
-
-	//Response object
-	JsonWriter->WriteObjectStart();
+	TArray<FStructOnScope> PropValuesOnScope;
+	for (const FGuid& RCPropertyId : InModifiedPropertyIds)
 	{
-		JsonWriter->WriteValue(TEXT("Type"), TEXT("PresetFieldsChanged"));
-		JsonWriter->WriteValue("PresetName", *InPreset->GetFName().ToString());
-		JsonWriter->WriteValue("PresetId", *InPreset->GetPresetId().ToString());
-
-		JsonWriter->WriteIdentifierPrefix("ChangedFields");
-
-		//All exposed properties of this preset that changed
-		JsonWriter->WriteArrayStart();
+		FRCObjectReference ObjectRef;
+		if (TSharedPtr<FRemoteControlProperty> RCProperty = InPreset->GetExposedEntity<FRemoteControlProperty>(RCPropertyId).Pin())
 		{
-			for (const FGuid& RCPropertyId : InModifiedPropertyIds)
+			if (IRemoteControlModule::Get().ResolveObjectProperty(ERCAccess::READ_ACCESS, RCProperty->GetBoundObjects()[0], RCProperty->FieldPathInfo.ToString(), ObjectRef))
 			{
 				bHasProperty = true;
-
-				FRCObjectReference ObjectRef;
-				if (TSharedPtr<FRemoteControlProperty> RCProperty = InPreset->GetExposedEntity<FRemoteControlProperty>(RCPropertyId).Pin())
-				{
-					//Property object
-					JsonWriter->WriteObjectStart();
-					{
-						JsonWriter->WriteValue(TEXT("PropertyLabel"), *RCProperty->GetLabel().ToString());
-						JsonWriter->WriteValue(TEXT("Id"), *RCProperty->GetId().ToString());
-
-						for (UObject* Object : RCProperty->GetBoundObjects())
-						{
-							bHasProperty = true;
-
-							IRemoteControlModule::Get().ResolveObjectProperty(ERCAccess::READ_ACCESS, Object, RCProperty->FieldPathInfo.ToString(), ObjectRef);
-
-							JsonWriter->WriteValue(TEXT("ObjectPath"), Object->GetPathName());
-							JsonWriter->WriteIdentifierPrefix(TEXT("PropertyValue"));
-
-							RemotePayloadSerializer::SerializePartial(
-								[&ObjectRef](FJsonStructSerializerBackend& SerializerBackend)
-								{
-									return IRemoteControlModule::Get().GetObjectProperties(ObjectRef, SerializerBackend);
-								}
-							, Writer);
-
-						}
-					}
-					JsonWriter->WriteObjectEnd();
-				}
-
+				PropValuesOnScope.Add(WebSocketMessageHandlerStructUtils::CreatePropertyValueOnScope(RCProperty, ObjectRef));
 			}
 		}
-		JsonWriter->WriteArrayEnd();
 	}
-	JsonWriter->WriteObjectEnd();
+
+	if (PropValuesOnScope.Num())
+	{
+		FStructOnScope FieldsChangedEventOnScope = WebSocketMessageHandlerStructUtils::CreatePresetFieldsChangedStructOnScope(InPreset, PropValuesOnScope);
+
+		FMemoryWriter Writer(OutBuffer);
+		WebRemoteControlUtils::SerializeStructOnScope(FieldsChangedEventOnScope, Writer);
+	}
 
 	return bHasProperty;
 }
 
 
-bool FWebSocketMessageHandler::WriteActorPropertyChangePayload(URemoteControlPreset* InPreset, const TMap<FRemoteControlActor, TArray<FRCObjectReference>>& InModifications, TArray<uint8>& OutBuffer)
+bool FWebSocketMessageHandler::WriteActorPropertyChangePayload(URemoteControlPreset* InPreset, const TMap<FRemoteControlActor, TArray<FRCObjectReference>>& InModifications, FMemoryWriter& InWriter)
 {
 	bool bHasProperty = false;
 
-	FMemoryWriter Writer(OutBuffer);
-	TSharedPtr<TJsonWriter<UCS2CHAR>> JsonWriter = TJsonWriter<UCS2CHAR>::Create(&Writer);
-
-	//Might be a better idea to have defined structures for our web socket notification messages
+	TArray<FStructOnScope> ModifiedActorsOnScope;
 
-	//Response object
-	JsonWriter->WriteObjectStart();
+	for (const TPair<FRemoteControlActor, TArray<FRCObjectReference>>& Pair : InModifications)
 	{
-		JsonWriter->WriteValue(TEXT("Type"), TEXT("PresetActorModified"));
-		JsonWriter->WriteValue("PresetName", *InPreset->GetFName().ToString());
-		JsonWriter->WriteValue("PresetId", *InPreset->GetPresetId().ToString());
-
-		JsonWriter->WriteIdentifierPrefix("ModifiedActors");
-
-		JsonWriter->WriteArrayStart();
-
-		for (const TPair<FRemoteControlActor, TArray<FRCObjectReference>>& Pair : InModifications)
+		if (AActor* ModifiedActor = Cast<AActor>(Pair.Key.Path.ResolveObject()))
 		{
-			JsonWriter->WriteObjectStart();
-			
-			AActor* ModifiedActor = Cast<AActor>(Pair.Key.Path.ResolveObject());
-			if (!ModifiedActor)
-			{
-				continue;
-			}
-
-			FString RCActorName = Pair.Key.GetLabel().ToString();
-			
-			JsonWriter->WriteValue(TEXT("Id"), *Pair.Key.GetId().ToString());
-			JsonWriter->WriteValue(TEXT("DisplayName"), RCActorName);
-			JsonWriter->WriteValue(TEXT("Path"), Pair.Key.Path.ToString());
+			TArray<FStructOnScope> PropertyValuesOnScope;
 
-			JsonWriter->WriteIdentifierPrefix("ModifiedProperties");
-			//All exposed properties of this preset that changed
-			JsonWriter->WriteArrayStart();
+			for (const FRCObjectReference& Ref : Pair.Value)
 			{
-				for (const FRCObjectReference& Ref : Pair.Value)
-				{
-					const FProperty* Property = Ref.Property.Get();
-
-					if (!Property || !Ref.IsValid())
-					{
-						continue;
-					}
+				const FProperty* Property = Ref.Property.Get();
 
+				if (Property && Ref.IsValid())
+				{
 					bHasProperty = true;
-
-					//Property object
-					JsonWriter->WriteObjectStart();
-					{
-						JsonWriter->WriteValue(TEXT("PropertyName"), Property->GetName());
-						JsonWriter->WriteIdentifierPrefix(TEXT("PropertyValue"));
-						RemotePayloadSerializer::SerializePartial(
-							[&Ref](FJsonStructSerializerBackend& SerializerBackend)
-							{
-								return IRemoteControlModule::Get().GetObjectProperties(Ref, SerializerBackend);
-							}
-						, Writer);
-
-					}
-					JsonWriter->WriteObjectEnd();
+					PropertyValuesOnScope.Add(WebSocketMessageHandlerStructUtils::CreateActorPropertyValueOnScope(InPreset, Ref));
 				}
 			}
-			JsonWriter->WriteArrayEnd();
-			JsonWriter->WriteObjectEnd();
-		}
-		JsonWriter->WriteArrayEnd();
 
+			if (PropertyValuesOnScope.Num())
+			{
+				ModifiedActorsOnScope.Add(WebSocketMessageHandlerStructUtils::CreateModifiedActorStructOnScope(InPreset, Pair.Key, PropertyValuesOnScope));
+			}
+		}
 	}
-	JsonWriter->WriteObjectEnd();
+
+	FStructOnScope ActorsModifedOnScope = WebSocketMessageHandlerStructUtils::CreateModifiedActorsStructOnScope(InPreset, ModifiedActorsOnScope);
+	WebRemoteControlUtils::SerializeStructOnScope(ActorsModifedOnScope, InWriter);
 
 	return bHasProperty;
 }
diff --git a/Engine/Plugins/VirtualProduction/RemoteControl/Source/WebRemoteControl/Private/WebSocketMessageHandler.h b/Engine/Plugins/VirtualProduction/RemoteControl/Source/WebRemoteControl/Private/WebSocketMessageHandler.h
--- a/Engine/Plugins/VirtualProduction/RemoteControl/Source/WebRemoteControl/Private/WebSocketMessageHandler.h
+++ b/Engine/Plugins/VirtualProduction/RemoteControl/Source/WebRemoteControl/Private/WebSocketMessageHandler.h
@@ -98,7 +98,7 @@ private:
 	/**
 	 * Write the provided list of actor modifications to a buffer.
 	 */
-	bool WriteActorPropertyChangePayload(URemoteControlPreset* InPreset, const TMap<FRemoteControlActor, TArray<FRCObjectReference>>& InModifications, TArray<uint8>& OutBuffer);
+	bool WriteActorPropertyChangePayload(URemoteControlPreset* InPreset, const TMap<FRemoteControlActor, TArray<FRCObjectReference>>& InModifications, FMemoryWriter& InWriter);
 
 private:
 
diff --git a/Engine/Plugins/VirtualProduction/RemoteControlWebInterface/Source/RemoteControlWebInterface/Private/RCWebInterface.cpp b/Engine/Plugins/VirtualProduction/RemoteControlWebInterface/Source/RemoteControlWebInterface/Private/RCWebInterface.cpp
--- a/Engine/Plugins/VirtualProduction/RemoteControlWebInterface/Source/RemoteControlWebInterface/Private/RCWebInterface.cpp
+++ b/Engine/Plugins/VirtualProduction/RemoteControlWebInterface/Source/RemoteControlWebInterface/Private/RCWebInterface.cpp
@@ -10,26 +10,54 @@
 #include "Misc/Parse.h"
 
 #if WITH_EDITOR
+#include "Editor.h"
 #include "Misc/CoreDelegates.h"
 #include "RCWebInterfaceCustomizations.h"
 #endif
 
 #define LOCTEXT_NAMESPACE "FRemoteControlWebInterfaceModule"
 
+namespace RCWebInterface
+{
+	bool IsWebInterfaceEnabledInEditor()
+	{
+		bool bIsEditor = false;
+
+#if WITH_EDITOR
+		if (GEditor)
+		{
+			bIsEditor = true;
+		}
+#endif
+
+		// By default, remote control web interface is disabled in -game and packaged game.
+		return bIsEditor || FParse::Param(FCommandLine::Get(), TEXT("RCWebInterfaceEnable"));
+	}
+}
+
 void FRemoteControlWebInterfaceModule::StartupModule()
 {
 	bRCWebInterfaceDisable = FParse::Param(FCommandLine::Get(), TEXT("RCWebInterfaceDisable")) || IsRunningCommandlet();
 
-	WebApp = MakeShared<FRemoteControlWebInterfaceProcess>();
-	if (!bRCWebInterfaceDisable)
+	if (bRCWebInterfaceDisable)
 	{
-		TSharedPtr<FRemoteControlWebInterfaceProcess> WebAppLocal = WebApp;
-		FCoreDelegates::OnFEngineLoopInitComplete.AddLambda([WebAppLocal]()
-		{
-			WebAppLocal->Start();
-		});
+		return;
+	}
+	
+	if (!RCWebInterface::IsWebInterfaceEnabledInEditor())
+	{
+		UE_LOG(LogRemoteControlWebInterface, Display, TEXT("Remote Control Web Interface is disabled by default when running outside the editor. Use the -RCWebInterfaceEnable flag when launching in order to use it."));
+		return;
 	}
+	
+	WebApp = MakeShared<FRemoteControlWebInterfaceProcess>();
 
+	TSharedPtr<FRemoteControlWebInterfaceProcess> WebAppLocal = WebApp;
+	FCoreDelegates::OnFEngineLoopInitComplete.AddLambda([WebAppLocal]()
+	{
+		WebAppLocal->Start();
+	});
+	
 	if (IWebRemoteControlModule* WebRemoteControlModule = FModuleManager::GetModulePtr<IWebRemoteControlModule>("WebRemoteControl"))
 	{
 		WebSocketServerStartedDelegate = WebRemoteControlModule->OnWebSocketServerStarted().AddLambda([this](uint32) 
@@ -51,11 +79,18 @@ void FRemoteControlWebInterfaceModule::StartupModule()
 
 void FRemoteControlWebInterfaceModule::ShutdownModule()
 {
-	if (!bRCWebInterfaceDisable)
+	if (bRCWebInterfaceDisable)
 	{
-		WebApp->Shutdown();
+		return;
 	}
 
+	if (!RCWebInterface::IsWebInterfaceEnabledInEditor())
+    {
+    	return;
+    }
+    
+	WebApp->Shutdown();
+	
 	FCoreDelegates::OnFEngineLoopInitComplete.RemoveAll(this);
 
 #if WITH_EDITOR
@@ -70,16 +105,11 @@ void FRemoteControlWebInterfaceModule::ShutdownModule()
 
 void FRemoteControlWebInterfaceModule::OnSettingsModified(UObject* Settings, FPropertyChangedEvent& PropertyChangedEvent)
 {
-	if (bRCWebInterfaceDisable)
-	{
-		return;
-	}
-
 	FName PropertyName = PropertyChangedEvent.GetPropertyName();
 	if (PropertyName == GET_MEMBER_NAME_CHECKED(URemoteControlSettings, RemoteControlWebInterfacePort) || PropertyName == GET_MEMBER_NAME_CHECKED(URemoteControlSettings, bForceWebAppBuildAtStartup))
 	{
-	WebApp->Shutdown();
-	WebApp->Start();
+		WebApp->Shutdown();
+		WebApp->Start();
 	}
 }
 
