Manually bring CLs from Release-4.27
CLs: 17552296 17556294 17570864

diff --git a/Engine/Plugins/VirtualProduction/RemoteControl/Source/RemoteControl/Private/RemoteControlModule.cpp b/Engine/Plugins/VirtualProduction/RemoteControl/Source/RemoteControl/Private/RemoteControlModule.cpp
--- a/Engine/Plugins/VirtualProduction/RemoteControl/Source/RemoteControl/Private/RemoteControlModule.cpp
+++ b/Engine/Plugins/VirtualProduction/RemoteControl/Source/RemoteControl/Private/RemoteControlModule.cpp
@@ -6,6 +6,7 @@
 #include "IRemoteControlModule.h"
 #include "IStructDeserializerBackend.h"
 #include "IStructSerializerBackend.h"
+#include "Components/ActorComponent.h"
 #include "Components/LightComponent.h"
 #include "RCPropertyUtilities.h"
 #include "RemoteControlFieldPath.h"
@@ -522,6 +523,11 @@ bool FRemoteControlModule::InvokeCall(FRCCall& InCall, ERCPayloadType InPayloadT
 				if (OngoingModification->bHasStartedTransaction)
 				{
 					SnapshotTransactionBuffer(InCall.CallRef.Object.Get());
+					
+					if (UActorComponent* Component = Cast<UActorComponent>(InCall.CallRef.Object.Get()))
+					{
+						SnapshotTransactionBuffer(Component->GetOwner());
+					}
 				}
 			}
 			else
diff --git a/Engine/Plugins/VirtualProduction/RemoteControl/Source/RemoteControlUI/Private/AssetEditor/RemoteControlPresetEditorToolkit.cpp b/Engine/Plugins/VirtualProduction/RemoteControl/Source/RemoteControlUI/Private/AssetEditor/RemoteControlPresetEditorToolkit.cpp
--- a/Engine/Plugins/VirtualProduction/RemoteControl/Source/RemoteControlUI/Private/AssetEditor/RemoteControlPresetEditorToolkit.cpp
+++ b/Engine/Plugins/VirtualProduction/RemoteControl/Source/RemoteControlUI/Private/AssetEditor/RemoteControlPresetEditorToolkit.cpp
@@ -7,16 +7,14 @@
 #include "RemoteControlUIModule.h"
 #include "Subsystems/AssetEditorSubsystem.h"
 #include "Widgets/Docking/SDockTab.h"
-#include "Widgets/SWidget.h"
-#include "UI/SRCPanelExposedEntitiesList.h"
 #include "UI/SRCPanelTreeNode.h"
 #include "UI/SRemoteControlPanel.h"
 
 #define LOCTEXT_NAMESPACE "RemoteControlPresetEditorToolkit"
 
-const FName FRemoteControlPresetEditorToolkit::PanelTabId(TEXT("RemoteControl_RemoteControlPanel"));
 const FName FRemoteControlPresetEditorToolkit::RemoteControlPanelAppIdentifier(TEXT("RemoteControlPanel"));
 
+
 TSharedRef<FRemoteControlPresetEditorToolkit> FRemoteControlPresetEditorToolkit::CreateEditor(const EToolkitMode::Type Mode, const TSharedPtr<class IToolkitHost>& InitToolkitHost, URemoteControlPreset* InPreset)
 {
 	TSharedRef<FRemoteControlPresetEditorToolkit> NewEditor = MakeShared<FRemoteControlPresetEditorToolkit>();
@@ -30,7 +28,9 @@ void FRemoteControlPresetEditorToolkit::InitRemoteControlPresetEditor(const EToo
 {
 	Preset = InPreset;
 
-	PanelTab = FRemoteControlUIModule::Get().CreateRemoteControlPanel(InPreset);
+	PanelTab = FRemoteControlUIModule::Get().CreateRemoteControlPanel(InPreset, InitToolkitHost);
+
+	RegisterTabSpawners(InitToolkitHost->GetTabManager().ToSharedRef());
 
 	const TSharedRef<FTabManager::FLayout> StandaloneDefaultLayout = FTabManager::NewLayout("Standalone_RemoteControlPresetEditor")
 		->AddArea
@@ -39,7 +39,7 @@ void FRemoteControlPresetEditorToolkit::InitRemoteControlPresetEditor(const EToo
 			->Split
 			(
 				FTabManager::NewStack()
-				->AddTab(PanelTabId, ETabState::OpenedTab)
+				->AddTab(FRemoteControlUIModule::RemoteControlPanelTabName, ETabState::OpenedTab)
 			)
 		);
 
@@ -51,9 +51,9 @@ void FRemoteControlPresetEditorToolkit::InitRemoteControlPresetEditor(const EToo
 	// tab ID will already be registered within EditorTabManager
 	FLevelEditorModule& LevelEditorModule = FModuleManager::GetModuleChecked<FLevelEditorModule>("LevelEditor");
 	TSharedPtr<FTabManager> EditorTabManager = LevelEditorModule.GetLevelEditorTabManager();
-	if (EditorTabManager->FindExistingLiveTab(PanelTabId).IsValid())
+	if (EditorTabManager->FindExistingLiveTab(FRemoteControlUIModule::RemoteControlPanelTabName).IsValid())
 	{
-		EditorTabManager->TryInvokeTab(PanelTabId)->RequestCloseTab();
+		EditorTabManager->TryInvokeTab(FRemoteControlUIModule::RemoteControlPanelTabName)->RequestCloseTab();
 	}
 	
 	FAssetEditorToolkit::InitAssetEditor(Mode, InitToolkitHost, RemoteControlPanelAppIdentifier, StandaloneDefaultLayout, bCreateDefaultStandaloneMenu, bCreateDefaultStandaloneMenu, InPreset);
@@ -69,7 +69,7 @@ FRemoteControlPresetEditorToolkit::~FRemoteControlPresetEditorToolkit()
 		if (TSharedPtr<FTabManager> EditorTabManager = LevelEditorModule.GetLevelEditorTabManager())
 		{
 			UnregisterTabSpawners(EditorTabManager.ToSharedRef());
-			if (TSharedPtr<SDockTab> Tab = EditorTabManager->FindExistingLiveTab(PanelTabId))
+			if (TSharedPtr<SDockTab> Tab = EditorTabManager->FindExistingLiveTab(FRemoteControlUIModule::RemoteControlPanelTabName))
 			{
 				Tab->RequestCloseTab();
 			}
@@ -81,15 +81,21 @@ void FRemoteControlPresetEditorToolkit::RegisterTabSpawners(const TSharedRef<FTa
 {
 	WorkspaceMenuCategory = InTabManager->AddLocalWorkspaceMenuCategory(LOCTEXT("WorkspaceMenu_RemoteControlPanel", "Remote Control Panel"));
 
-	InTabManager->RegisterTabSpawner(PanelTabId, FOnSpawnTab::CreateSP(this, &FRemoteControlPresetEditorToolkit::HandleTabManagerSpawnPanelTab))
+	InTabManager->RegisterTabSpawner(FRemoteControlUIModule::RemoteControlPanelTabName, FOnSpawnTab::CreateSP(this, &FRemoteControlPresetEditorToolkit::HandleTabManagerSpawnPanelTab))
 		.SetDisplayName(LOCTEXT("RemoteControlPanelMainTab", "Remote Control Panel"))
 		.SetGroup(WorkspaceMenuCategory.ToSharedRef())
 		.SetIcon(FSlateIcon(FEditorStyle::GetStyleSetName(), "LevelEditor.GameSettings.Small"));
+
+	InTabManager->RegisterTabSpawner(FRemoteControlUIModule::EntityDetailsTabName, FOnSpawnTab::CreateSP(this, &FRemoteControlPresetEditorToolkit::HandleTabManagerSpawnDetailsTab))
+		.SetDisplayName(LOCTEXT("RemoteControlPanelDetailsTab", "Entity Details"))
+		.SetGroup(WorkspaceMenuCategory.ToSharedRef())
+		.SetIcon(FSlateIcon(FEditorStyle::GetStyleSetName(), "LevelEditor.Tabs.Details"));
 }
 
 void FRemoteControlPresetEditorToolkit::UnregisterTabSpawners(const TSharedRef<FTabManager>& InTabManager)
 {
-	InTabManager->UnregisterTabSpawner(PanelTabId);
+	InTabManager->UnregisterTabSpawner(FRemoteControlUIModule::EntityDetailsTabName);
+	InTabManager->UnregisterTabSpawner(FRemoteControlUIModule::RemoteControlPanelTabName);
 }
 
 bool FRemoteControlPresetEditorToolkit::OnRequestClose()
@@ -126,7 +132,7 @@ FString FRemoteControlPresetEditorToolkit::GetWorldCentricTabPrefix() const
 
 TSharedRef<SDockTab> FRemoteControlPresetEditorToolkit::HandleTabManagerSpawnPanelTab(const FSpawnTabArgs& Args)
 {
-	check(Args.GetTabId() == PanelTabId);
+	check(Args.GetTabId() == FRemoteControlUIModule::RemoteControlPanelTabName);
 
 	return SNew(SDockTab)
 		.Label(LOCTEXT("ControlPanelLabel", "Control Panel"))
@@ -136,15 +142,25 @@ TSharedRef<SDockTab> FRemoteControlPresetEditorToolkit::HandleTabManagerSpawnPan
 		];
 }
 
+TSharedRef<SDockTab> FRemoteControlPresetEditorToolkit::HandleTabManagerSpawnDetailsTab(const FSpawnTabArgs& Args)
+{
+	check(Args.GetTabId() == FRemoteControlUIModule::EntityDetailsTabName);
+
+	return SNew(SDockTab)
+		.Label(LOCTEXT("ControlPanelEntityDetailsLabel", "Entity Details"))
+		.TabColorScale(GetTabColorScale())	
+		[
+			SNullWidget::NullWidget
+		];
+}
+
 void FRemoteControlPresetEditorToolkit::InvokePanelTab()
 {
 	struct Local
 	{
 		static void OnRemoteControlPresetClosed(TSharedRef<SDockTab> DockTab, TWeakPtr<IAssetEditorInstance> InRemoteControlPreset)
 		{
-			TSharedPtr<IAssetEditorInstance> AssetEditorInstance = InRemoteControlPreset.Pin();
-
-			if (AssetEditorInstance.IsValid())
+			if (const TSharedPtr<IAssetEditorInstance> AssetEditorInstance = InRemoteControlPreset.Pin())
 			{
 				InRemoteControlPreset.Pin()->CloseWindow();
 			}
@@ -155,11 +171,12 @@ void FRemoteControlPresetEditorToolkit::InvokePanelTab()
 	FLevelEditorModule& LevelEditorModule = FModuleManager::GetModuleChecked<FLevelEditorModule>("LevelEditor");
 	TSharedPtr<FTabManager> EditorTabManager = LevelEditorModule.GetLevelEditorTabManager();
 
-	if (TSharedPtr<SDockTab> Tab = EditorTabManager->TryInvokeTab(PanelTabId))
+	if (TSharedPtr<SDockTab> Tab = EditorTabManager->TryInvokeTab(FRemoteControlUIModule::RemoteControlPanelTabName))
 	{
 		Tab->SetContent(PanelTab.ToSharedRef());
 		Tab->SetOnTabClosed(SDockTab::FOnTabClosedCallback::CreateStatic(&Local::OnRemoteControlPresetClosed, TWeakPtr<IAssetEditorInstance>(SharedThis(this))));
 	}
 }
 
+
 #undef LOCTEXT_NAMESPACE /*RemoteControlPresetEditorToolkit*/
\ No newline at end of file
diff --git a/Engine/Plugins/VirtualProduction/RemoteControl/Source/RemoteControlUI/Private/AssetEditor/RemoteControlPresetEditorToolkit.h b/Engine/Plugins/VirtualProduction/RemoteControl/Source/RemoteControlUI/Private/AssetEditor/RemoteControlPresetEditorToolkit.h
--- a/Engine/Plugins/VirtualProduction/RemoteControl/Source/RemoteControlUI/Private/AssetEditor/RemoteControlPresetEditorToolkit.h
+++ b/Engine/Plugins/VirtualProduction/RemoteControl/Source/RemoteControlUI/Private/AssetEditor/RemoteControlPresetEditorToolkit.h
@@ -41,14 +41,13 @@ public:
 	/** End IAssetEditorInstance interface */
 
 private:
-	/** Handle spawning the tab that holds the remote control panel tab. */
+	//~ Handlers that spawn remote control tabs.
 	TSharedRef<SDockTab> HandleTabManagerSpawnPanelTab(const FSpawnTabArgs& Args);
+	TSharedRef<SDockTab> HandleTabManagerSpawnDetailsTab(const FSpawnTabArgs& Args);
 
 	/** Handle invoking the remote control tab. */
 	void InvokePanelTab();
 private:
-	/** Holds the remote control panel tab id. */
-	static const FName PanelTabId;
 	/** Holds the remote control panel app identifier. */
 	static const FName RemoteControlPanelAppIdentifier;
 	/** Holds the preset being edited. */
diff --git a/Engine/Plugins/VirtualProduction/RemoteControl/Source/RemoteControlUI/Private/RemoteControlUIModule.cpp b/Engine/Plugins/VirtualProduction/RemoteControl/Source/RemoteControlUI/Private/RemoteControlUIModule.cpp
--- a/Engine/Plugins/VirtualProduction/RemoteControl/Source/RemoteControlUI/Private/RemoteControlUIModule.cpp
+++ b/Engine/Plugins/VirtualProduction/RemoteControl/Source/RemoteControlUI/Private/RemoteControlUIModule.cpp
@@ -4,6 +4,7 @@
 
 #include "AssetToolsModule.h"
 #include "AssetTools/RemoteControlPresetActions.h"
+#include "EditorStyleSet.h"
 #include "Framework/MultiBox/MultiBoxBuilder.h"
 #include "HAL/PlatformApplicationMisc.h"
 #include "MaterialEditor/DEditorParameterValue.h"
@@ -12,6 +13,7 @@
 #include "RemoteControlField.h"
 #include "RemoteControlPreset.h"
 #include "RemoteControlSettings.h"
+#include "AssetEditor/RemoteControlPresetEditorToolkit.h"
 #include "Styling/AppStyle.h"
 #include "Textures/SlateIcon.h"
 #include "UI/Customizations/RemoteControlEntityCustomization.h"
@@ -24,6 +26,10 @@
 
 #define LOCTEXT_NAMESPACE "RemoteControlUI"
 
+const FName FRemoteControlUIModule::EntityDetailsTabName = "RemoteControl_EntityDetails";
+const FName FRemoteControlUIModule::RemoteControlPanelTabName = "RemoteControl_RemoteControlPanel";
+
+
 namespace RemoteControlUIModule
 {
 	const TArray<FName>& GetCustomizedStructNames()
@@ -37,6 +43,12 @@ namespace RemoteControlUIModule
 
 		return CustomizedStructNames;
 	};
+	
+	static void OnOpenRemoteControlPanel(const EToolkitMode::Type Mode, const TSharedPtr<class IToolkitHost>& EditWithinLevelEditor, class URemoteControlPreset* Preset)
+	{
+		TSharedRef<FRemoteControlPresetEditorToolkit> Toolkit = FRemoteControlPresetEditorToolkit::CreateEditor(Mode, EditWithinLevelEditor,  Preset);
+		Toolkit->InitRemoteControlPresetEditor(Mode, EditWithinLevelEditor, Preset);
+	}
 }
 
 void FRemoteControlUIModule::StartupModule()
@@ -82,14 +94,14 @@ void FRemoteControlUIModule::UnregisterMetadataCustomization(FName MetadataKey)
 	ExternalEntityMetadataCustomizations.Remove(MetadataKey);
 }
 
-TSharedRef<SRemoteControlPanel> FRemoteControlUIModule::CreateRemoteControlPanel(URemoteControlPreset* Preset)
+TSharedRef<SRemoteControlPanel> FRemoteControlUIModule::CreateRemoteControlPanel(URemoteControlPreset* Preset, const TSharedPtr<IToolkitHost>& ToolkitHost)
 {
 	if (TSharedPtr<SRemoteControlPanel> Panel = WeakActivePanel.Pin())
 	{
 		Panel->SetEditMode(false);
 	}
 
-	TSharedRef<SRemoteControlPanel> PanelRef = SAssignNew(WeakActivePanel, SRemoteControlPanel, Preset)
+	TSharedRef<SRemoteControlPanel> PanelRef = SAssignNew(WeakActivePanel, SRemoteControlPanel, Preset, ToolkitHost)
 		.OnEditModeChange_Lambda(
 			[this](TSharedPtr<SRemoteControlPanel> Panel, bool bEditMode) 
 			{
diff --git a/Engine/Plugins/VirtualProduction/RemoteControl/Source/RemoteControlUI/Private/RemoteControlUIModule.h b/Engine/Plugins/VirtualProduction/RemoteControl/Source/RemoteControlUI/Private/RemoteControlUIModule.h
--- a/Engine/Plugins/VirtualProduction/RemoteControl/Source/RemoteControlUI/Private/RemoteControlUIModule.h
+++ b/Engine/Plugins/VirtualProduction/RemoteControl/Source/RemoteControlUI/Private/RemoteControlUIModule.h
@@ -9,6 +9,7 @@
 #include "PropertyEditorDelegates.h"
 #include "PropertyHandle.h"
 
+class IToolkitHost;
 class SRemoteControlPanel;
 class URemoteControlPreset;
 
@@ -53,7 +54,7 @@ public:
 	 * @param The preset to display.
 	 * @return The remote control widget.
 	 */
-	TSharedRef<SRemoteControlPanel> CreateRemoteControlPanel(URemoteControlPreset* Preset);
+	TSharedRef<SRemoteControlPanel> CreateRemoteControlPanel(URemoteControlPreset* Preset, const TSharedPtr<IToolkitHost>& ToolkitHost = TSharedPtr<IToolkitHost>());
 
 	/**
 	 * Get the map of entity metadata entry customizations.
@@ -63,6 +64,10 @@ public:
 		return ExternalEntityMetadataCustomizations;
 	}
 
+public:
+	static const FName EntityDetailsTabName;
+	static const FName RemoteControlPanelTabName;
+
 	TSharedPtr<SRCPanelTreeNode> GenerateEntityWidget(const FGenerateWidgetArgs& Args);
 
 private:
diff --git a/Engine/Plugins/VirtualProduction/RemoteControl/Source/RemoteControlUI/Private/UI/SRemoteControlPanel.cpp b/Engine/Plugins/VirtualProduction/RemoteControl/Source/RemoteControlUI/Private/UI/SRemoteControlPanel.cpp
--- a/Engine/Plugins/VirtualProduction/RemoteControl/Source/RemoteControlUI/Private/UI/SRemoteControlPanel.cpp
+++ b/Engine/Plugins/VirtualProduction/RemoteControl/Source/RemoteControlUI/Private/UI/SRemoteControlPanel.cpp
@@ -3,7 +3,6 @@
 #include "SRemoteControlPanel.h"
 
 #include "ActorEditorUtils.h"
-#include "AssetData.h"
 #include "AssetRegistryModule.h"
 #include "ClassViewerFilter.h"
 #include "ClassViewerModule.h"
@@ -40,14 +39,14 @@
 #include "SRCPanelFunctionPicker.h"
 #include "SRCPanelExposedActor.h"
 #include "SRCPanelExposedField.h"
-#include "SRCPanelFieldGroup.h"
 #include "SRCPanelTreeNode.h"
 #include "Subsystems/Subsystem.h"
 #include "Templates/SharedPointer.h"
 #include "Templates/SubclassOf.h"
 #include "Templates/UnrealTypeTraits.h"
-#include "UObject/SoftObjectPtr.h"
+#include "Toolkits/IToolkitHost.h"
 #include "Widgets/DeclarativeSyntaxSupport.h"
+#include "Widgets/Docking/SDockTab.h"
 #include "Widgets/Input/SButton.h"
 #include "Widgets/Input/SCheckBox.h"
 #include "Widgets/Layout/SBorder.h"
@@ -58,7 +57,6 @@
 #include "Widgets/SBoxPanel.h"
 #include "Widgets/Text/STextBlock.h"
 
-
 #define LOCTEXT_NAMESPACE "RemoteControlPanel"
 
 namespace RemoteControlPanelUtils
@@ -90,11 +88,12 @@ namespace RemoteControlPanelUtils
 	}
 }
 
-void SRemoteControlPanel::Construct(const FArguments& InArgs, URemoteControlPreset* InPreset)
+void SRemoteControlPanel::Construct(const FArguments& InArgs, URemoteControlPreset* InPreset, TSharedPtr<IToolkitHost> InToolkitHost)
 {
 	OnEditModeChange = InArgs._OnEditModeChange;
 	Preset = TStrongObjectPtr<URemoteControlPreset>(InPreset);
 	WidgetRegistry = MakeShared<FRCPanelWidgetRegistry>();
+	ToolkitHost = InToolkitHost;
 
 	UpdateRebindButtonVisibility();
 
@@ -188,6 +187,20 @@ void SRemoteControlPanel::Construct(const FArguments& InArgs, URemoteControlPres
 						.Text(LOCTEXT("RebindButtonText", "Rebind All"))
 					]
 				]
+				.VAlign(VAlign_Center)
+				.AutoWidth()
+				[
+					SNew(SButton)
+					.ButtonStyle(FEditorStyle::Get(), "FlatButton")
+					.TextStyle(FRemoteControlPanelStyle::Get(), "RemoteControlPanel.Button.TextStyle")
+					.ToolTipText(LOCTEXT("EntityDetailsToolTip", "Open the details panel for the selected exposed entity."))
+					.OnClicked_Lambda([this](){ ToggleDetailsView(); return FReply::Handled(); })
+					[
+						SNew(SImage)
+						.Image(FEditorStyle::Get().GetBrush("LevelEditor.Tabs.Details"))
+					]
+				]
+
 				+ SHorizontalBox::Slot()
 				.VAlign(VAlign_Center)
 				.Padding(4.0f, 0)
@@ -307,29 +320,10 @@ void SRemoteControlPanel::Construct(const FArguments& InArgs, URemoteControlPres
 						+ SSplitter::Slot()
 						.Value(TreeBindingSplitRatioBottom)
 						[
-							SNew(SSplitter)
-							.Orientation(Orient_Vertical)
-							.HitDetectionSplitterHandleSize(0.0f) // Effectively disables user manipulation (but keeps visuals)
-							+ SSplitter::Slot()
-							.SizeRule(SSplitter::SizeToContent)
-							[
-								SNew(SBorder)
-								.BorderImage(FAppStyle::Get().GetBrush("ToolPanel.GroupBorder"))
-								[
-									SNew(SBox)
-									.MinDesiredHeight(40.0f)
-									[
-										CreateEntityDetailsView()
-									]
-								]
-							]
-							+ SSplitter::Slot()
+							SNew(SBorder)
+							.BorderImage(FAppStyle::Get().GetBrush("ToolPanel.GroupBorder"))
 							[
-								SNew(SBorder)
-								.BorderImage(FAppStyle::Get().GetBrush("ToolPanel.GroupBorder"))
-								[
-									EntityProtocolDetails.ToSharedRef()
-								]
+								EntityProtocolDetails.ToSharedRef()
 							]
 						]
 					]
@@ -915,6 +909,28 @@ void SRemoteControlPanel::ExposeActor(AActor* Actor)
 	}
 }
 
+void SRemoteControlPanel::ToggleDetailsView()
+{
+	const FTabId TabId = FTabId(FRemoteControlUIModule::EntityDetailsTabName);
+	
+	if (ToolkitHost)
+	{
+		if (TSharedPtr<SDockTab> ExistingTab = ToolkitHost->GetTabManager()->FindExistingLiveTab(TabId))
+		{
+			ExistingTab->RequestCloseTab();
+		}
+		else
+		{
+			// Request the Tab Manager to invoke the tab. This will spawn the tab if needed, otherwise pull it to focus. This assumes
+			// that the Toolkit Host's Tab Manager has already registered a tab with a NullWidget for content.
+			if (TSharedPtr<SDockTab> EntityDetailsTab = ToolkitHost->GetTabManager()->TryInvokeTab(TabId))
+			{
+				EntityDetailsTab->SetContent(CreateEntityDetailsView());
+			}
+		}
+	}
+}
+
 TSharedRef<SWidget> SRemoteControlPanel::CreateEntityDetailsView()
 {
 	FDetailsViewArgs Args;
@@ -947,7 +963,7 @@ void SRemoteControlPanel::UpdateEntityDetailsView(const TSharedPtr<SRCPanelTreeN
 		}
 	}
 
-	if (ensure(EntityDetailsView))
+	if (EntityDetailsView)
 	{
 		EntityDetailsView->SetStructureData(SelectedEntityPtr);
 	}
diff --git a/Engine/Plugins/VirtualProduction/RemoteControl/Source/RemoteControlUI/Private/UI/SRemoteControlPanel.h b/Engine/Plugins/VirtualProduction/RemoteControl/Source/RemoteControlUI/Private/UI/SRemoteControlPanel.h
--- a/Engine/Plugins/VirtualProduction/RemoteControl/Source/RemoteControlUI/Private/UI/SRemoteControlPanel.h
+++ b/Engine/Plugins/VirtualProduction/RemoteControl/Source/RemoteControlUI/Private/UI/SRemoteControlPanel.h
@@ -17,6 +17,7 @@ struct FListEntry;
 class FRCPanelWidgetRegistry;
 struct FRemoteControlEntity;
 class FReply;
+class IToolkitHost;
 class IPropertyRowGenerator;
 class IPropertyHandle;
 class SBox;
@@ -44,7 +45,7 @@ class SRemoteControlPanel : public SCompoundWidget, public FSelfRegisteringEdito
 	SLATE_END_ARGS()
 
 public:
-	void Construct(const FArguments& InArgs, URemoteControlPreset* InPreset);
+	void Construct(const FArguments& InArgs, URemoteControlPreset* InPreset, TSharedPtr<IToolkitHost> InToolkitHost);
 	~SRemoteControlPanel();
 
 	//~ FEditorUndoClient interface
@@ -127,6 +128,9 @@ private:
 	/** Handle exposing an actor. */
 	void ExposeActor(AActor* Actor);
 
+	/** Opens or closes the entity details view tab. */
+	void ToggleDetailsView();
+
 	/** Handles disabling CPU throttling. */
 	FReply OnClickDisableUseLessCPU() const;
 
@@ -216,4 +220,7 @@ private:
 	TSharedPtr<STextBlock> PresetNameTextBlock;
 	/** Holds a cache of widgets. */
 	TSharedPtr<FRCPanelWidgetRegistry> WidgetRegistry;
+	
+	/** The toolkit that hosts this panel. */
+	TSharedPtr<IToolkitHost> ToolkitHost;
 };
