RC materials

diff --git a/Engine/Plugins/VirtualProduction/RemoteControl/Source/RemoteControl/Private/RemoteControlField.cpp b/Engine/Plugins/VirtualProduction/RemoteControl/Source/RemoteControl/Private/RemoteControlField.cpp
--- a/Engine/Plugins/VirtualProduction/RemoteControl/Source/RemoteControl/Private/RemoteControlField.cpp
+++ b/Engine/Plugins/VirtualProduction/RemoteControl/Source/RemoteControl/Private/RemoteControlField.cpp
@@ -298,6 +298,31 @@ bool FRemoteControlProperty::IsBound() const
 	return !!GetProperty();
 }
 
+bool FRemoteControlProperty::CheckIsBoundToPropertyPath(const FString& InPath) const
+{
+	return FieldPathInfo.ToPathPropertyString() == InPath;
+}
+
+bool FRemoteControlProperty::ContainsBoundObjects(TArray<UObject*> InObjects) const
+{
+	const TArray<UObject*> BoundObjects = GetBoundObjects();
+
+	if (BoundObjects.Num() == 0 || InObjects.Num() == 0)
+	{
+		return false;
+	}
+
+	for (UObject* BoundObject : BoundObjects)
+	{
+		if (InObjects.Contains(BoundObject))
+		{
+			return true;
+		}
+	}
+	
+	return false;
+}
+
 void FRemoteControlProperty::InitializeMetadata()
 {
 #if WITH_EDITOR
diff --git a/Engine/Plugins/VirtualProduction/RemoteControl/Source/RemoteControl/Private/RemoteControlInstanceMaterial.cpp b/Engine/Plugins/VirtualProduction/RemoteControl/Source/RemoteControl/Private/RemoteControlInstanceMaterial.cpp
new file mode 100644
--- /dev/null
+++ b/Engine/Plugins/VirtualProduction/RemoteControl/Source/RemoteControl/Private/RemoteControlInstanceMaterial.cpp
@@ -0,0 +1,398 @@
+﻿// Copyright Epic Games, Inc. All Rights Reserved.
+
+#include "RemoteControlInstanceMaterial.h"
+
+#include "IRemoteControlModule.h"
+#include "RemoteControlPreset.h"
+
+#include "Materials/MaterialInstanceConstant.h"
+
+#if WITH_EDITOR
+#include "MaterialEditor/DEditorFontParameterValue.h"
+#include "MaterialEditor/DEditorScalarParameterValue.h"
+#include "MaterialEditor/DEditorTextureParameterValue.h"
+#include "MaterialEditor/DEditorVectorParameterValue.h"
+#include "MaterialEditor/MaterialEditorInstanceConstant.h"
+#endif
+
+#if WITH_EDITOR
+struct FMaterialHelper
+{
+private:
+	/** Try to get a path for material property */
+	static FString GetFieldPathCallback(const FName& InPropertyName, const FRCFieldPathInfo& InOriginalFieldPath, const int32 InIndex)
+	{
+		FString FieldPathString = FString::Printf(TEXT("%s[%d].ParameterValue"), *InPropertyName.ToString(), InIndex);
+
+		if (InOriginalFieldPath.Segments.Num() > 1)
+		{
+			FieldPathString += TEXT(".") + InOriginalFieldPath.Segments[1].ToString();
+		}
+
+		return FieldPathString;
+	}
+
+	/** Try to get a path for font material property */
+	static FString GetFontFieldPathCallback(const FName& InPropertyName, const FRCFieldPathInfo& InOriginalFieldPath, const int32 InIndex)
+	{
+		FString FieldPathString;
+
+		if (InOriginalFieldPath.Segments.Num() > 1)
+		{
+			FieldPathString = FString::Printf(TEXT("%s[%d].%s"), *InPropertyName.ToString(), InIndex, *InOriginalFieldPath.Segments[1].ToString());
+		}
+		else
+		{
+			FieldPathString = FString::Printf(TEXT("%s[%d].FontValue"), *InPropertyName.ToString(), InIndex);
+		}
+
+		return FieldPathString;
+	}
+	
+public:
+	/** Field path Callback for getting Field path */
+	using FFieldPathCallback = TFunctionRef<FString(const FName& InPropertyName, const FRCFieldPathInfo&, const int32 InIndex)>;
+	
+	/** Try to find parameter index by given name */
+	template <typename ParameterType>
+	static int32 FindParameterIndexByName(const TArray<ParameterType>& Parameters, const FHashedMaterialParameterInfo& ParameterInfo)
+	{
+		for (int32 ParameterIndex = 0; ParameterIndex < Parameters.Num(); ++ParameterIndex)
+		{
+			const ParameterType* Parameter = &Parameters[ParameterIndex];
+			if (Parameter->ParameterInfo == ParameterInfo)
+			{
+				return ParameterIndex;
+			}
+		}
+
+		return INDEX_NONE;
+	}
+
+	/** Try replace Field Path based on given input */
+	template <typename ParameterType>
+	static bool ReplaceFieldPath(FFieldPathCallback InFieldPathCallback, const FName& InPropertyName,
+	                             const TArray<ParameterType>& Parameters,
+	                             const FHashedMaterialParameterInfo& ParameterInfo,
+	                             const FRCFieldPathInfo& OriginalFieldPath, FRCFieldPathInfo& OutFieldPath)
+	{
+		const int32 Index = FindParameterIndexByName(Parameters, ParameterInfo);
+		if (Index == INDEX_NONE)
+		{
+			return false;
+		}
+
+		const FString FieldPathString = InFieldPathCallback(InPropertyName, OriginalFieldPath, Index);
+
+		OutFieldPath = FieldPathString;
+
+		return true;
+	}
+
+
+	/** Replace Object and Field Path */
+	static bool ReplaceFieldPath(UMaterialInstance* InMaterialInstance, UDEditorParameterValue* InEditorParameterValue, const FRCFieldPathInfo& OriginalFieldPath, FRCFieldPathInfo& OutFieldPath)
+	{
+		bool bReplaced = false;
+
+		const FMaterialParameterInfo& ParameterInfo = InEditorParameterValue->ParameterInfo;
+			
+		if (InEditorParameterValue->IsA<UDEditorScalarParameterValue>())
+		{
+			bReplaced = ReplaceFieldPath(GetFieldPathCallback,
+				GET_MEMBER_NAME_STRING_CHECKED(UMaterialInstance, ScalarParameterValues),
+				InMaterialInstance->ScalarParameterValues, ParameterInfo, OriginalFieldPath,  OutFieldPath);
+		}
+		else if (InEditorParameterValue->IsA<UDEditorVectorParameterValue>())
+		{
+			bReplaced = ReplaceFieldPath(GetFieldPathCallback,
+				GET_MEMBER_NAME_STRING_CHECKED(UMaterialInstance, VectorParameterValues),
+				InMaterialInstance->VectorParameterValues, ParameterInfo, OriginalFieldPath, OutFieldPath);
+		}
+		else if (InEditorParameterValue->IsA<UDEditorTextureParameterValue>())
+		{
+			bReplaced = ReplaceFieldPath(GetFieldPathCallback,
+				GET_MEMBER_NAME_STRING_CHECKED(UMaterialInstance, TextureParameterValues),
+				InMaterialInstance->TextureParameterValues, ParameterInfo, OriginalFieldPath, OutFieldPath);
+		}
+		else if (InEditorParameterValue->IsA<UDEditorFontParameterValue>())
+		{			
+			bReplaced = ReplaceFieldPath(GetFontFieldPathCallback,
+				GET_MEMBER_NAME_STRING_CHECKED(UMaterialInstance, FontParameterValues),
+				InMaterialInstance->FontParameterValues, ParameterInfo, OriginalFieldPath, OutFieldPath);
+		}
+
+		return bReplaced;
+		}
+
+	template<typename ParameterType>
+	static const ParameterType* FindParameterValue(const TArray<ParameterType>& Parameters, const FName& InName)
+	{
+		return Parameters.FindByPredicate([&InName](const ParameterType& InValue)
+		{
+			return InValue.ParameterInfo.Name == InName;
+		});
+	}
+};
+
+#endif
+
+#if WITH_EDITOR
+
+FRemoteControlInstanceMaterial::FRemoteControlInstanceMaterial(URemoteControlPreset* InPreset, FName InLabel,
+                                                               UDEditorParameterValue* InEditorParameterValue, const FRCFieldPathInfo& InOriginalFieldPathInfo,
+                                                               UObject* InInstance, const FRCFieldPathInfo& InFieldPath,
+                                                               const TArray<URemoteControlBinding*>& InBindings)
+	: FRemoteControlProperty(InPreset, InLabel, InFieldPath, InBindings)
+	, OriginalFieldPathInfo(InOriginalFieldPathInfo)
+	, InstancePath(InInstance)
+{
+	OriginalClass = InEditorParameterValue->GetClass();
+	ParameterInfo = InEditorParameterValue->ParameterInfo;
+}
+#endif
+
+#if WITH_EDITOR
+void FRemoteControlInstanceMaterial::OnObjectPropertyChanged(UObject* InObject, FPropertyChangedEvent&)
+{
+	UMaterialInstance* InputMaterialInstance =  Cast<UMaterialInstance>(InObject);
+	UMaterialInstance* MaterialInstance =  Cast<UMaterialInstance>(InstancePath.ResolveObject());
+	
+	if (!InputMaterialInstance || !MaterialInstance || InputMaterialInstance != MaterialInstance)
+	{
+		return;
+	}
+
+	UDEditorParameterValue* DEditorParameterValue = [InputMaterialInstance, this]()
+	{
+		UDEditorParameterValue* ReturnParameterValue = nullptr;
+		
+		for (TObjectIterator<UDEditorParameterValue> It; It; ++It)
+		{
+			UDEditorParameterValue* ParameterValue = *It;
+			if (!ParameterValue || !ParameterValue->IsValidLowLevelFast())
+			{
+				continue;
+			}
+
+			if (ParameterValue->ParameterInfo.Name != ParameterInfo.Name)
+			{
+				continue;
+			}
+
+			if(UMaterialEditorInstanceConstant* MaterialEditorInstanceConstant = Cast<UMaterialEditorInstanceConstant>(ParameterValue->GetOuter()))
+			{
+				if (MaterialEditorInstanceConstant->SourceInstance == InputMaterialInstance)
+				{
+					ReturnParameterValue = ParameterValue;
+				}
+			}
+		}
+
+		return ReturnParameterValue;
+	}();
+
+	if (!DEditorParameterValue)
+	{
+		return;
+	}
+
+	const FName EditorValueParameterName = DEditorParameterValue->ParameterInfo.Name;
+
+	if (UDEditorScalarParameterValue* ScalarParameterValue = Cast<UDEditorScalarParameterValue>(DEditorParameterValue))
+	{
+		const FScalarParameterValue* FoundScalarParameterValue = FMaterialHelper::FindParameterValue(MaterialInstance->ScalarParameterValues, EditorValueParameterName);
+
+		if (FoundScalarParameterValue)
+		{
+			if (!FMath::IsNearlyEqual(ScalarParameterValue->ParameterValue ,FoundScalarParameterValue->ParameterValue))
+			{
+				ScalarParameterValue->ParameterValue = FoundScalarParameterValue->ParameterValue;
+			}
+		}
+	}
+	else if (UDEditorVectorParameterValue* VectorParameterValue = Cast<UDEditorVectorParameterValue>(DEditorParameterValue))
+	{
+		const FVectorParameterValue* FoundVectorParameterValue = FMaterialHelper::FindParameterValue(MaterialInstance->VectorParameterValues, EditorValueParameterName);
+
+		if (FoundVectorParameterValue)
+		{
+			if (VectorParameterValue->ParameterValue != FoundVectorParameterValue->ParameterValue)
+			{
+				VectorParameterValue->ParameterValue = FoundVectorParameterValue->ParameterValue;
+			}
+		}
+	}
+	else if (UDEditorTextureParameterValue* TextureParameterValue = Cast<UDEditorTextureParameterValue>(DEditorParameterValue))
+	{
+		const FTextureParameterValue* FoundTextureParameterValue = FMaterialHelper::FindParameterValue(MaterialInstance->TextureParameterValues, EditorValueParameterName);
+
+		if (FoundTextureParameterValue)
+		{
+			if (TextureParameterValue->ParameterValue != FoundTextureParameterValue->ParameterValue)
+			{
+				TextureParameterValue->ParameterValue = FoundTextureParameterValue->ParameterValue;
+			}
+		}
+	}
+	else if (UDEditorFontParameterValue* FontParameterValue = Cast<UDEditorFontParameterValue>(DEditorParameterValue))
+	{
+		const FFontParameterValue* FoundFontParameterValue = FMaterialHelper::FindParameterValue(MaterialInstance->FontParameterValues, EditorValueParameterName);
+
+		if (FoundFontParameterValue)
+		{
+			if (FontParameterValue->ParameterValue.FontValue != FoundFontParameterValue->FontValue)
+			{
+				FontParameterValue->ParameterValue.FontValue = FoundFontParameterValue->FontValue;
+				FontParameterValue->ParameterValue.FontPage = FoundFontParameterValue->FontPage;
+			}
+		}
+	}
+}
+#endif
+
+bool FRemoteControlInstanceMaterial::CheckIsBoundToPropertyPath(const FString& InPath) const
+{
+	return OriginalFieldPathInfo.ToString() == InPath;
+}
+
+bool FRemoteControlInstanceMaterial::ContainsBoundObjects(TArray<UObject*> InObjects) const
+{
+#if WITH_EDITOR	
+	TArray<UObject*> ReplacesObjects;
+
+	for (UObject* Object : InObjects)
+	{
+		UDEditorParameterValue* EditorParameterValue = Cast<UDEditorParameterValue>(Object);
+
+		if (!EditorParameterValue)
+		{
+			return false;
+		}
+
+		if (OriginalClass != EditorParameterValue->GetClass())
+		{
+			return false;
+		}
+	
+		UMaterialInstance* ParameterMaterialInstance = nullptr;
+		if(UMaterialEditorInstanceConstant* MaterialEditorInstanceConstant = Cast<UMaterialEditorInstanceConstant>(EditorParameterValue->GetOuter()))
+		{
+			ParameterMaterialInstance = MaterialEditorInstanceConstant->SourceInstance;
+		}
+
+		if (!ParameterMaterialInstance)
+		{
+			return false;
+		}
+	
+		UMaterialInstance* MaterialInstance =  Cast<UMaterialInstance>(InstancePath.ResolveObject());
+		if (!MaterialInstance || MaterialInstance != ParameterMaterialInstance)
+		{
+			return false;
+		}
+
+		ReplacesObjects.Add(MaterialInstance);
+	}
+
+	return FRemoteControlProperty::ContainsBoundObjects(ReplacesObjects);
+#endif
+
+	return false;
+}
+
+void FRemoteControlInstanceMaterial::PostLoad()
+{
+#if WITH_EDITOR
+	// Preload material instance asset in editor
+	InstancePath.TryLoad();
+#endif
+}
+
+TSharedRef<IRemoteControlPropertyFactory> FRemoteControlInstanceMaterialFactory::MakeInstance()
+{
+	return MakeShared<FRemoteControlInstanceMaterialFactory>();
+}
+
+TSharedPtr<FRemoteControlProperty> FRemoteControlInstanceMaterialFactory::CreateRemoteControlProperty(URemoteControlPreset* Preset, UObject* Object, FRCFieldPathInfo FieldPath, FRemoteControlPresetExposeArgs Args)
+{
+#if WITH_EDITOR
+	if (UDEditorParameterValue* EditorParameterValue = Cast<UDEditorParameterValue>(Object))
+	{
+		FRCFieldPathInfo OriginalFieldPath = FieldPath;
+
+		UMaterialEditorInstanceConstant* MaterialEditorInstanceConstant = Cast<UMaterialEditorInstanceConstant>(EditorParameterValue->GetOuter());
+		if (!MaterialEditorInstanceConstant)
+		{
+			return nullptr;
+		}
+		
+		UMaterialInstance* MaterialInstance = MaterialEditorInstanceConstant->SourceInstance;
+		if (!MaterialInstance)
+		{
+			return nullptr;
+		}
+
+		if (EditorParameterValue->bOverride == false)
+		{
+			// Override the property value in material
+			EditorParameterValue->bOverride = true;
+			EditorParameterValue->PostEditChange();
+
+			FPropertyChangedEvent OverrideEvent(nullptr);
+			MaterialEditorInstanceConstant->PostEditChangeProperty( OverrideEvent );
+		}
+
+		// If we are dealing with DEditor Parameter Value it should have a Object and Field Path replacement
+		// For Material Instance it should be replacement of UDEditorParameterValue with UMaterial Instance
+		FRCFieldPathInfo ReplacedFieldPath;
+		if (ensure(FMaterialHelper::ReplaceFieldPath(MaterialInstance, EditorParameterValue, FieldPath, ReplacedFieldPath)))
+		{
+			if (!ReplacedFieldPath.Resolve(MaterialInstance))
+			{
+				return nullptr;
+			}
+			
+			FName DesiredName = *Args.Label;
+			if (OriginalFieldPath.Segments.Num() > 1)
+			{
+				DesiredName = Preset->GetEntityName(DesiredName, MaterialInstance, ReplacedFieldPath);
+			}
+			if (DesiredName == NAME_None)
+			{
+				DesiredName = *FString::Printf(TEXT("%s (%s)"), *EditorParameterValue->ParameterInfo.Name.ToString(), *MaterialInstance->GetName());
+			}
+				
+			FRemoteControlInstanceMaterial RCProperty{
+				Preset, Preset->GenerateUniqueLabel(DesiredName), EditorParameterValue, MoveTemp(OriginalFieldPath),
+				MaterialInstance, MoveTemp(ReplacedFieldPath), {Preset->FindOrAddBinding(MaterialInstance)}
+			};
+        		
+			return StaticCastSharedPtr<FRemoteControlProperty>(Preset->Expose(MoveTemp(RCProperty), FRemoteControlInstanceMaterial::StaticStruct(), Args.GroupId));
+		}
+	}
+	
+#endif
+	
+	return nullptr;
+}
+
+void FRemoteControlInstanceMaterialFactory::PostSetObjectProperties(UObject* InObject, bool bInSuccess) const
+{
+	if (bInSuccess && !GIsEditor)
+	{
+		if (UMaterialInstance* MaterialInstance = Cast<UMaterialInstance>(InObject))
+		{
+			// This modifies material variables used for rendering
+			MaterialInstance->InitStaticPermutation();
+		}
+	}
+}
+
+bool FRemoteControlInstanceMaterialFactory::SupportExposedClass(UClass* InClass) const
+{
+#if WITH_EDITOR
+	return InClass->IsChildOf(UDEditorParameterValue::StaticClass());
+#endif
+	return false;
+}
\ No newline at end of file
diff --git a/Engine/Plugins/VirtualProduction/RemoteControl/Source/RemoteControl/Private/RemoteControlModule.cpp b/Engine/Plugins/VirtualProduction/RemoteControl/Source/RemoteControl/Private/RemoteControlModule.cpp
--- a/Engine/Plugins/VirtualProduction/RemoteControl/Source/RemoteControl/Private/RemoteControlModule.cpp
+++ b/Engine/Plugins/VirtualProduction/RemoteControl/Source/RemoteControl/Private/RemoteControlModule.cpp
@@ -12,6 +12,7 @@
 #include "RemoteControlFieldPath.h"
 #include "RemoteControlInterceptionHelpers.h"
 #include "RemoteControlInterceptionProcessor.h"
+#include "RemoteControlInstanceMaterial.h"
 #include "RemoteControlPreset.h"
 #include "StructDeserializer.h"
 #include "StructSerializer.h"
@@ -345,6 +346,9 @@ void FRemoteControlModule::StartupModule()
 #if WITH_EDITOR
 	FCoreDelegates::OnPostEngineInit.AddRaw(this, &FRemoteControlModule::HandleEnginePostInit);
 #endif
+
+	// Register Property Factories
+	RegisterEntityFactory(FRemoteControlInstanceMaterial::StaticStruct()->GetFName(), FRemoteControlInstanceMaterialFactory::MakeInstance());
 }
 
 void FRemoteControlModule::ShutdownModule()
@@ -364,6 +368,9 @@ void FRemoteControlModule::ShutdownModule()
 
 	// Unregister the interceptor feature on module shutdown
 	IModularFeatures::Get().UnregisterModularFeature(IRemoteControlInterceptionFeatureProcessor::GetName(), RCIProcessor.Get());
+
+	// Unregister Property Factories
+	UnregisterEntityFactory(FRemoteControlInstanceMaterial::StaticStruct()->GetFName());
 }
 
 IRemoteControlModule::FOnPresetRegistered& FRemoteControlModule::OnPresetRegistered()
@@ -901,7 +908,10 @@ bool FRemoteControlModule::SetObjectProperties(const FRCObjectReference& ObjectA
 			MutableObjectReference.Object->PostEditChangeProperty(PropertyEvent);
 		}
 #endif
-		return bSuccess;
+		for (const TPair<FName, TSharedPtr<IRemoteControlPropertyFactory>>& EntityFactoryPair : EntityFactories)
+		{
+			EntityFactoryPair.Value->PostSetObjectProperties(ObjectAccess.Object.Get(), bSuccess);
+		}
 	}
 	return false;
 }
@@ -1100,6 +1110,24 @@ bool FRemoteControlModule::PropertySupportsRawModificationWithoutEditor(FPropert
 	return Property && (RemoteControlUtil::IsPropertyAllowed(Property, ERCAccess::WRITE_ACCESS, bInGameOrPackage) || !!RemoteControlPropertyUtilities::FindSetterFunction(Property, OwnerClass));
 }
 
+void FRemoteControlModule::RegisterEntityFactory(const FName InFactoryName, const TSharedRef<IRemoteControlPropertyFactory>& InFactory)
+{
+	if( InFactoryName != NAME_None )
+	{
+		EntityFactories.Add(InFactoryName, InFactory);
+	}
+	else
+	{
+		UE_LOG(LogRemoteControl, Error, TEXT("Factory should have a name"));
+		ensure(false);
+	}
+}
+
+void FRemoteControlModule::UnregisterEntityFactory(const FName InFactoryName)
+{
+	EntityFactories.Remove(InFactoryName);
+}
+
 void FRemoteControlModule::CachePresets()
 {
 	TArray<FAssetData> Assets;
diff --git a/Engine/Plugins/VirtualProduction/RemoteControl/Source/RemoteControl/Private/RemoteControlModule.h b/Engine/Plugins/VirtualProduction/RemoteControl/Source/RemoteControl/Private/RemoteControlModule.h
--- a/Engine/Plugins/VirtualProduction/RemoteControl/Source/RemoteControl/Private/RemoteControlModule.h
+++ b/Engine/Plugins/VirtualProduction/RemoteControl/Source/RemoteControl/Private/RemoteControlModule.h
@@ -42,6 +42,9 @@ public:
 	virtual bool RegisterDefaultEntityMetadata(FName MetadataKey, FEntityMetadataInitializer MetadataInitializer) override;
 	virtual void UnregisterDefaultEntityMetadata(FName MetadataKey) override;
 	virtual bool PropertySupportsRawModificationWithoutEditor(FProperty* Property, UClass* OwnerClass = nullptr) const override;
+	virtual void RegisterEntityFactory( const FName InFactoryName, const TSharedRef<IRemoteControlPropertyFactory>& InFactory) override;
+	virtual void UnregisterEntityFactory( const FName InFactoryName ) override;
+	virtual const TMap<FName, TSharedPtr<IRemoteControlPropertyFactory>>& GetEntityFactories() const override { return EntityFactories; };
 	//~ End IRemoteControlModule
 
 private:
@@ -124,6 +127,9 @@ private:
 	/** Delay before we check if a modification is no longer ongoing. */
 	static constexpr float SecondsBetweenOngoingChangeCheck = 0.2f;
 #endif
+
+	/** Map of the factories which is responsible for the Remote Control property creation */
+	TMap<FName, TSharedPtr<IRemoteControlPropertyFactory>> EntityFactories;
 };
 
 PRAGMA_ENABLE_DEPRECATION_WARNINGS
diff --git a/Engine/Plugins/VirtualProduction/RemoteControl/Source/RemoteControl/Private/RemoteControlPreset.cpp b/Engine/Plugins/VirtualProduction/RemoteControl/Source/RemoteControl/Private/RemoteControlPreset.cpp
--- a/Engine/Plugins/VirtualProduction/RemoteControl/Source/RemoteControl/Private/RemoteControlPreset.cpp
+++ b/Engine/Plugins/VirtualProduction/RemoteControl/Source/RemoteControl/Private/RemoteControlPreset.cpp
@@ -16,8 +16,10 @@
 #include "RemoteControlFieldPath.h"
 #include "RemoteControlActor.h"
 #include "RemoteControlBinding.h"
+#include "RemoteControlEntityFactory.h"
 #include "RemoteControlObjectVersion.h"
 #include "RemoteControlPresetRebindingManager.h"
+
 #include "UObject/Object.h"
 #include "UObject/ObjectMacros.h"
 #if WITH_EDITOR
@@ -689,6 +691,8 @@ void URemoteControlPreset::PostLoad()
 	RegisterEntityDelegates();
 
 	CreatePropertyWatchers();
+
+	PostLoadProperties();
 }
 
 void URemoteControlPreset::PostDuplicate(bool bDuplicateForPIE)
@@ -739,6 +743,12 @@ TWeakPtr<FRemoteControlActor> URemoteControlPreset::ExposeActor(AActor* Actor, F
 	return StaticCastSharedPtr<FRemoteControlActor>(Expose(MoveTemp(RCActor), FRemoteControlActor::StaticStruct(), Args.GroupId));
 }
 
+
+FName URemoteControlPreset::GenerateUniqueLabel(const FName InDesiredName) const
+{
+	return Registry->GenerateUniqueLabel(InDesiredName);
+}
+
 TWeakPtr<FRemoteControlProperty> URemoteControlPreset::ExposeProperty(UObject* Object, FRCFieldPathInfo FieldPath, FRemoteControlPresetExposeArgs Args)
 {
 	if (!Object)
@@ -746,48 +756,36 @@ TWeakPtr<FRemoteControlProperty> URemoteControlPreset::ExposeProperty(UObject* O
 		return nullptr;
 	}
 
-	if (!FieldPath.Resolve(Object))
+	TSharedPtr<FRemoteControlProperty> RCPropertyPtr;
+	const TMap<FName, TSharedPtr<IRemoteControlPropertyFactory>>& PropertyFactories = IRemoteControlModule::Get().GetEntityFactories();
+	for (const TPair<FName, TSharedPtr<IRemoteControlPropertyFactory>>& EntityFactoryPair : PropertyFactories)
 	{
-		return nullptr;
+		if (EntityFactoryPair.Value->SupportExposedClass(Object->GetClass()))
+		{
+			RCPropertyPtr = EntityFactoryPair.Value->CreateRemoteControlProperty(this, Object, FieldPath, Args);
+			break;
+		}
 	}
 
-	FProperty* Property = FieldPath.GetResolvedData().Field;
-	check(Property);
-
-	FString FieldName;
-
-#if WITH_EDITOR
-	FieldName = Property->GetDisplayNameText().ToString();
-#else
-	FieldName = FieldPath.GetFieldName().ToString();
-#endif
-
-	FName DesiredName = *Args.Label;
-
-	if (DesiredName == NAME_None)
+	// Create default one
+	if (!RCPropertyPtr)
 	{
-		FString ObjectName;
-#if WITH_EDITOR
-		if (AActor* Actor = Cast<AActor>(Object))
+		if (!FieldPath.Resolve(Object))
 		{
-			ObjectName = Actor->GetActorLabel();
-		}
-		else if(UActorComponent* Component = Cast<UActorComponent>(Object))
-		{
-			ObjectName = Component->GetOwner()->GetActorLabel();
-		}
-		else
-#endif
-		{
-			ObjectName = Object->GetName();
+			return nullptr;
 		}
 
-		DesiredName = *FString::Printf(TEXT("%s (%s)"), *FieldName, *ObjectName);
-	}
+		const FName FieldName = GetEntityName(*Args.Label, Object, FieldPath);
+
+		FRemoteControlProperty RCProperty{ this, Registry->GenerateUniqueLabel(FieldName), MoveTemp(FieldPath), { FindOrAddBinding(Object) } };
 
-	FRemoteControlProperty RCProperty{ this, Registry->GenerateUniqueLabel(DesiredName), MoveTemp(FieldPath), { FindOrAddBinding(Object) } };
+		RCPropertyPtr = StaticCastSharedPtr<FRemoteControlProperty>(Expose(MoveTemp(RCProperty), FRemoteControlProperty::StaticStruct(), Args.GroupId));
+	}
 
-	TSharedPtr<FRemoteControlProperty> RCPropertyPtr = StaticCastSharedPtr<FRemoteControlProperty>(Expose(MoveTemp(RCProperty), FRemoteControlProperty::StaticStruct(), Args.GroupId));
+	if (!RCPropertyPtr)
+	{
+		return nullptr;
+	}
 
 	RCPropertyPtr->EnableEditCondition();
 
@@ -1044,6 +1042,53 @@ void URemoteControlPreset::CreatePropertyWatchers()
 	}
 }
 
+void URemoteControlPreset::PostLoadProperties()
+{
+	for (const TSharedPtr<FRemoteControlProperty>& ExposedProperty : Registry->GetExposedEntities<FRemoteControlProperty>())
+	{
+		ExposedProperty->PostLoad();
+	}
+}
+
+FName URemoteControlPreset::GetEntityName(const FName InDesiredName, UObject* InObject, const FRCFieldPathInfo& InFieldPath) const
+{
+	FName DesiredName = InDesiredName;
+	
+	if (DesiredName == NAME_None)
+	{
+		FString ObjectName;
+#if WITH_EDITOR
+		if (AActor* Actor = Cast<AActor>(InObject))
+		{
+			ObjectName = Actor->GetActorLabel();
+		}
+		else if(UActorComponent* Component = Cast<UActorComponent>(InObject))
+		{
+			ObjectName = Component->GetOwner()->GetActorLabel();
+		}
+		else
+#endif
+		{
+			ObjectName = InObject->GetName();
+		}
+
+		FProperty* Property = InFieldPath.GetResolvedData().Field;
+		check(Property);
+
+		FString FieldPath;
+
+#if WITH_EDITOR
+		FieldPath = Property->GetDisplayNameText().ToString();
+#else
+		FieldPath = InFieldPath.GetFieldName().ToString();
+#endif
+			
+		DesiredName = *FString::Printf(TEXT("%s (%s)"), *FieldPath, *ObjectName);
+	}
+
+	return DesiredName;
+}
+
 TOptional<FRemoteControlFunction> URemoteControlPreset::GetFunction(FName FunctionLabel) const
 {
 PRAGMA_DISABLE_DEPRECATION_WARNINGS
@@ -1631,6 +1676,7 @@ void URemoteControlPreset::OnObjectPropertyChanged(UObject* Object, struct FProp
 				{
 					UE_LOG(LogRemoteControl, VeryVerbose, TEXT("(%s) Change detected on %s::%s"), *GetName(), *Object->GetName(), *Event.Property->GetName());
 					PerFrameModifiedProperties.Add(Property->GetId());
+					Property->OnObjectPropertyChanged(Object, Event);
 					Iter.RemoveCurrent();
 				}
 			}
diff --git a/Engine/Plugins/VirtualProduction/RemoteControl/Source/RemoteControl/Public/IRemoteControlModule.h b/Engine/Plugins/VirtualProduction/RemoteControl/Source/RemoteControl/Public/IRemoteControlModule.h
--- a/Engine/Plugins/VirtualProduction/RemoteControl/Source/RemoteControl/Public/IRemoteControlModule.h
+++ b/Engine/Plugins/VirtualProduction/RemoteControl/Source/RemoteControl/Public/IRemoteControlModule.h
@@ -18,6 +18,7 @@ PRAGMA_DISABLE_DEPRECATION_WARNINGS
 struct FExposedProperty;
 PRAGMA_ENABLE_DEPRECATION_WARNINGS
 struct FRemoteControlProperty;
+class IRemoteControlPropertyFactory;
 
 /**
  * Delegate called to initialize an exposed entity metadata entry registered with the RegisterDefaultEntityMetadata method.
@@ -105,6 +106,9 @@ enum class ERCAccess : uint8
  */
 struct FRCObjectReference
 {
+	/** Callback type for post set object properties */
+	using FPostSetObjectPropertyCallback = TFunction<bool(UObject* /*InObject*/, const FRCFieldPathInfo& /*InPathInfo*/, bool /*bInSuccess*/)>;
+	
 	FRCObjectReference() = default;
 
 	FRCObjectReference(ERCAccess InAccessType, UObject* InObject)
@@ -340,4 +344,20 @@ public:
 	 * Returns whether the property can be modified through SetObjectProperties when running without an editor.
 	 */
 	virtual bool PropertySupportsRawModificationWithoutEditor(FProperty* Property, UClass* OwnerClass = nullptr) const = 0;
+
+	/**
+	 * Register factory 
+	 * @param InFactoryName the factory unique name
+	 * @param InFactory Factory instance
+	 */
+	virtual void RegisterEntityFactory( const FName InFactoryName, const TSharedRef<IRemoteControlPropertyFactory>& InFactory) = 0;
+
+	/**
+	 * Remove factory by name
+	 * @param InFactoryName the factory unique name
+	 */
+	virtual void UnregisterEntityFactory( const FName InFactoryName ) = 0;
+
+	/** Get map of the factories which is responsible for the Remote Control property creation */
+	virtual const TMap<FName, TSharedPtr<IRemoteControlPropertyFactory>>& GetEntityFactories() const = 0;
 };
diff --git a/Engine/Plugins/VirtualProduction/RemoteControl/Source/RemoteControl/Public/RemoteControlEntityFactory.h b/Engine/Plugins/VirtualProduction/RemoteControl/Source/RemoteControl/Public/RemoteControlEntityFactory.h
new file mode 100644
--- /dev/null
+++ b/Engine/Plugins/VirtualProduction/RemoteControl/Source/RemoteControl/Public/RemoteControlEntityFactory.h
@@ -0,0 +1,45 @@
+﻿// Copyright Epic Games, Inc. All Rights Reserved.
+
+#pragma once
+
+#include "CoreMinimal.h"
+
+class URemoteControlPreset;
+struct FRCFieldPathInfo;
+struct FRemoteControlPresetExposeArgs;
+struct FRemoteControlProperty;
+
+/**
+ * Factory which is creates the instances of FRemoteControlProperty
+ */
+class IRemoteControlPropertyFactory : public TSharedFromThis<IRemoteControlPropertyFactory>
+{
+public:
+	/** Virtual destructor */
+	virtual ~IRemoteControlPropertyFactory(){}
+
+	/**
+	 * Expose a property from given preset
+	 * @param Preset the remote control preset object.
+	 * @param Object the object that holds the property.
+	 * @param FieldPath The name/path to the property.
+	 * @param Args Optional arguments used to expose the property.
+	 * @return The exposed property.
+	 */
+	virtual TSharedPtr<FRemoteControlProperty> CreateRemoteControlProperty(URemoteControlPreset* Preset, UObject* Object, FRCFieldPathInfo FieldPath, FRemoteControlPresetExposeArgs Args) = 0;
+
+	/**
+	 * Post set property action
+	 * @param Object the object that holds the property.
+	 * @param bInSuccess Whether the property was set.
+	 */
+	virtual void PostSetObjectProperties(UObject* Object, bool bInSuccess) const {}
+
+	/**
+	 * Whether the factory support exposed object
+	 * @param Class class to expose
+	 * @return true if the object is supported by given factory
+	 */
+	virtual bool SupportExposedClass(UClass* Class) const = 0;
+};
+
diff --git a/Engine/Plugins/VirtualProduction/RemoteControl/Source/RemoteControl/Public/RemoteControlField.h b/Engine/Plugins/VirtualProduction/RemoteControl/Source/RemoteControl/Public/RemoteControlField.h
--- a/Engine/Plugins/VirtualProduction/RemoteControl/Source/RemoteControl/Public/RemoteControlField.h
+++ b/Engine/Plugins/VirtualProduction/RemoteControl/Source/RemoteControl/Public/RemoteControlField.h
@@ -8,6 +8,7 @@
 #include "RemoteControlField.generated.h"
 
 class IRemoteControlPropertyHandle;
+struct FPropertyChangedEvent;
 
 /**
  * The type of the exposed field.
@@ -127,6 +128,33 @@ public:
 	virtual bool IsBound() const override;
 	//~ End FRemoteControlEntity interface
 
+	/**
+	 * Check whether given property path is bound to the property 
+	 * @param InPropertyPath Given Property Path
+	 * return true if the property bound to given path
+	 */
+	virtual bool CheckIsBoundToPropertyPath(const FString& InPropertyPath) const;
+
+	/**
+	 * Check whether any of the given objects are bound to this property
+	 * @param InObjects Objects to check
+	 * return true if at least one object is bound to the property
+	 */
+	virtual bool ContainsBoundObjects(TArray<UObject*> InObjects) const;
+
+	/**
+	 * Called on Post Load on Owner UObject.
+	 * Useful to initialize specific functionality of the inherit property classes
+	 */
+	virtual void PostLoad() {}
+
+	/**
+	 * Should be called when property chanced on bound object
+	 * @param InObject Edited object
+	 * @param InEvent  Change Event
+	 */
+	virtual void OnObjectPropertyChanged(UObject* InObject, FPropertyChangedEvent& InEvent)  {};
+
 	/**
 	 * Get the underlying property.
 	 * @return The exposed property or nullptr if it couldn't be resolved.
diff --git a/Engine/Plugins/VirtualProduction/RemoteControl/Source/RemoteControl/Public/RemoteControlFieldPath.h b/Engine/Plugins/VirtualProduction/RemoteControl/Source/RemoteControl/Public/RemoteControlFieldPath.h
--- a/Engine/Plugins/VirtualProduction/RemoteControl/Source/RemoteControl/Public/RemoteControlFieldPath.h
+++ b/Engine/Plugins/VirtualProduction/RemoteControl/Source/RemoteControl/Public/RemoteControlFieldPath.h
@@ -4,6 +4,7 @@
 #include "CoreMinimal.h"
 #include "RemoteControlFieldPath.generated.h"
 
+
 /** Small container for the resolved data of a remote control field segment */
 struct REMOTECONTROL_API FRCFieldResolvedData
 {
@@ -135,6 +136,16 @@ public:
 	 */
 	FRCFieldPathInfo(FProperty* Property);
 
+	bool operator == (const FRCFieldPathInfo& Other) const
+	{
+		return PathHash == Other.PathHash;
+	}
+
+	bool operator != (const FRCFieldPathInfo& Other) const
+	{
+		return PathHash != Other.PathHash;
+	}
+
 public:
 	/** Go through each segment and finds the property associated + container address for a given UObject owner */
 	bool Resolve(UObject* Owner);
diff --git a/Engine/Plugins/VirtualProduction/RemoteControl/Source/RemoteControl/Public/RemoteControlInstanceMaterial.h b/Engine/Plugins/VirtualProduction/RemoteControl/Source/RemoteControl/Public/RemoteControlInstanceMaterial.h
new file mode 100644
--- /dev/null
+++ b/Engine/Plugins/VirtualProduction/RemoteControl/Source/RemoteControl/Public/RemoteControlInstanceMaterial.h
@@ -0,0 +1,74 @@
+﻿// Copyright Epic Games, Inc. All Rights Reserved.
+
+#pragma once
+
+#include "RemoteControlField.h"
+#include "RemoteControlEntityFactory.h"
+
+#include "Materials/MaterialLayersFunctions.h"
+
+#include "RemoteControlInstanceMaterial.generated.h"
+
+class UDEditorParameterValue;
+class URemoteControlBinding;
+class URemoteControlPreset;
+
+/**
+ * Represents a material instance property that has been exposed to remote control.
+ */
+USTRUCT(BlueprintType)
+struct REMOTECONTROL_API FRemoteControlInstanceMaterial : public FRemoteControlProperty
+{
+	GENERATED_BODY()
+
+public:
+	FRemoteControlInstanceMaterial() = default;
+
+	// Material Instance property available only in editor
+#if WITH_EDITOR
+	FRemoteControlInstanceMaterial(URemoteControlPreset* InPreset, FName InLabel,
+	                               UDEditorParameterValue* InEditorParameterValue,
+	                               const FRCFieldPathInfo& InOriginalFieldPathInfo, UObject* InInstance,
+	                               const FRCFieldPathInfo& InFieldPath,
+	                               const TArray<URemoteControlBinding*>& InBindings);
+#endif
+
+	//~ Begin FRemoteControlProperty interface
+#if WITH_EDITOR
+	virtual void OnObjectPropertyChanged(UObject* InObject, FPropertyChangedEvent& InEvent) override;
+#endif
+	virtual bool CheckIsBoundToPropertyPath(const FString& InPath) const override;
+	virtual bool ContainsBoundObjects(TArray<UObject*> InObjects) const override;
+	virtual void PostLoad() override;
+	//~ End FRemoteControlProperty interface
+
+public:
+	/** Stores original binding class */
+	UPROPERTY()
+	UClass* OriginalClass;
+
+	/** Stores original material parameter info */
+	UPROPERTY()
+	FMaterialParameterInfo ParameterInfo;
+
+	/** Store original property path */
+	UPROPERTY()
+	FRCFieldPathInfo OriginalFieldPathInfo;
+
+	/** Store path to Material Instance */
+	UPROPERTY()
+	FSoftObjectPath InstancePath;
+};
+
+class FRemoteControlInstanceMaterialFactory : public IRemoteControlPropertyFactory
+{
+public:
+	static TSharedRef<IRemoteControlPropertyFactory> MakeInstance();
+
+	//~ Begin IRemoteControlPropertyFactory interface
+	virtual TSharedPtr<FRemoteControlProperty> CreateRemoteControlProperty(URemoteControlPreset* Preset, UObject* Object, FRCFieldPathInfo FieldPath, FRemoteControlPresetExposeArgs Args) override;
+	virtual void PostSetObjectProperties(UObject* Object, bool bInSuccess) const override;
+	virtual bool SupportExposedClass(UClass* Class) const override;
+	//~ End IRemoteControlPropertyFactory interface
+};
+
diff --git a/Engine/Plugins/VirtualProduction/RemoteControl/Source/RemoteControl/Public/RemoteControlPreset.h b/Engine/Plugins/VirtualProduction/RemoteControl/Source/RemoteControl/Public/RemoteControlPreset.h
--- a/Engine/Plugins/VirtualProduction/RemoteControl/Source/RemoteControl/Public/RemoteControlPreset.h
+++ b/Engine/Plugins/VirtualProduction/RemoteControl/Source/RemoteControl/Public/RemoteControlPreset.h
@@ -776,6 +776,18 @@ public:
 
 	virtual void Serialize(FArchive& Ar) override;
 
+	/**
+	* Get entity name from given Desired or UObject + Field Path 
+	* @param InDesiredName		Given name, it might be none
+	* @param InObject			Bound object input
+	* @param InFieldPath		Bound field path
+	* @return Computing name of the entity
+	*/
+	FName GetEntityName(const FName InDesiredName, UObject* InObject, const FRCFieldPathInfo& InFieldPath) const;
+
+	/** Generate label from Registry */
+	FName GenerateUniqueLabel(const FName InDesiredName) const;
+
 public:
 	/** The visual layout for this preset. */
 	UPROPERTY()
@@ -850,11 +862,14 @@ private:
 	TArray<TSharedPtr<FRemoteControlEntity>> GetEntities(UScriptStruct* EntityType);
 	TArray<TSharedPtr<const FRemoteControlEntity>> GetEntities(UScriptStruct* EntityType) const;
 
+public:	
 	/** Expose an entity in the registry. */
 	TSharedPtr<FRemoteControlEntity> Expose(FRemoteControlEntity&& Entity, UScriptStruct* EntityType, const FGuid& GroupId);
-	
+
+
 	/** Try to get a binding and creates a new one if it doesn't exist. */
 	URemoteControlBinding* FindOrAddBinding(const TSoftObjectPtr<UObject>& Object);
+private:
 
 	/** Handler called upon an entity being modified. */
 	void OnEntityModified(const FGuid& EntityId);
@@ -879,6 +894,9 @@ private:
 
 	/** Create property watchers for exposed properties that need them. */
 	void CreatePropertyWatchers();
+
+	/** Call post load function for exposed properties. */
+	void PostLoadProperties();
 	
 private:
 	/** Preset unique ID */
diff --git a/Engine/Plugins/VirtualProduction/RemoteControl/Source/RemoteControlUI/Private/RemoteControlUIModule.cpp b/Engine/Plugins/VirtualProduction/RemoteControl/Source/RemoteControlUI/Private/RemoteControlUIModule.cpp
--- a/Engine/Plugins/VirtualProduction/RemoteControl/Source/RemoteControlUI/Private/RemoteControlUIModule.cpp
+++ b/Engine/Plugins/VirtualProduction/RemoteControl/Source/RemoteControlUI/Private/RemoteControlUIModule.cpp
@@ -11,6 +11,7 @@
 #include "PropertyHandle.h"
 #include "RemoteControlActor.h"
 #include "RemoteControlField.h"
+#include "RemoteControlInstanceMaterial.h"
 #include "RemoteControlPreset.h"
 #include "RemoteControlSettings.h"
 #include "AssetEditor/RemoteControlPresetEditorToolkit.h"
@@ -49,6 +50,11 @@ namespace RemoteControlUIModule
 		TSharedRef<FRemoteControlPresetEditorToolkit> Toolkit = FRemoteControlPresetEditorToolkit::CreateEditor(Mode, EditWithinLevelEditor,  Preset);
 		Toolkit->InitRemoteControlPresetEditor(Mode, EditWithinLevelEditor, Preset);
 	}
+
+	bool IsWhitelistedTransientObject(UObject* Object)
+	{
+		return Object && Object->IsA<UDEditorParameterValue>();
+	}
 }
 
 void FRemoteControlUIModule::StartupModule()
@@ -336,7 +342,7 @@ bool FRemoteControlUIModule::ShouldDisplayExposeIcon(const TSharedRef<IPropertyH
 				}
 
 				// Don't display an expose icon for transient objects such as material editor parameters.
-				if (OuterObjects[0]->GetOutermost()->HasAnyFlags(RF_Transient))
+				if (OuterObjects[0]->GetOutermost()->HasAnyFlags(RF_Transient) && !RemoteControlUIModule::IsWhitelistedTransientObject(OuterObjects[0]))
 				{
 					return false;
 				}
@@ -421,6 +427,7 @@ void FRemoteControlUIModule::RegisterWidgetFactories()
 	RegisterWidgetFactoryForType(FRemoteControlActor::StaticStruct(), FOnGenerateRCWidget::CreateStatic(&SRCPanelExposedActor::MakeInstance));
 	RegisterWidgetFactoryForType(FRemoteControlProperty::StaticStruct(), FOnGenerateRCWidget::CreateStatic(&SRCPanelExposedField::MakeInstance));
 	RegisterWidgetFactoryForType(FRemoteControlFunction::StaticStruct(), FOnGenerateRCWidget::CreateStatic(&SRCPanelExposedField::MakeInstance));
+	RegisterWidgetFactoryForType(FRemoteControlInstanceMaterial::StaticStruct(), FOnGenerateRCWidget::CreateStatic(&SRCPanelExposedField::MakeInstance));
 }
 
 TSharedPtr<SRCPanelTreeNode> FRemoteControlUIModule::GenerateEntityWidget(const FGenerateWidgetArgs& Args)
diff --git a/Engine/Plugins/VirtualProduction/RemoteControl/Source/RemoteControlUI/Private/UI/SRemoteControlPanel.cpp b/Engine/Plugins/VirtualProduction/RemoteControl/Source/RemoteControlUI/Private/UI/SRemoteControlPanel.cpp
--- a/Engine/Plugins/VirtualProduction/RemoteControl/Source/RemoteControlUI/Private/UI/SRemoteControlPanel.cpp
+++ b/Engine/Plugins/VirtualProduction/RemoteControl/Source/RemoteControlUI/Private/UI/SRemoteControlPanel.cpp
@@ -386,15 +386,14 @@ bool SRemoteControlPanel::IsExposed(const TSharedPtr<IPropertyHandle>& PropertyH
 	
 	TArray<UObject*> OuterObjects;
 	PropertyHandle->GetOuterObjects(OuterObjects);
-
-	FString Path = PropertyHandle->GeneratePathToProperty();
+	const FString Path = PropertyHandle->GeneratePathToProperty();
 
 	TArray<TSharedPtr<FRemoteControlProperty>, TInlineAllocator<1>> PotentialMatches;
 	for (const TWeakPtr<FRemoteControlProperty>& WeakProperty : Preset->GetExposedEntities<FRemoteControlProperty>())
 	{
 		if (TSharedPtr<FRemoteControlProperty> Property = WeakProperty.Pin())
 		{
-			if (Property->FieldPathInfo.ToPathPropertyString() == Path)
+			if (Property->CheckIsBoundToPropertyPath(Path))
 			{
 				PotentialMatches.Add(Property);
 			}
@@ -409,7 +408,7 @@ bool SRemoteControlPanel::IsExposed(const TSharedPtr<IPropertyHandle>& PropertyH
 
 		for (const TSharedPtr<FRemoteControlProperty>& Property : PotentialMatches)
 		{
-			if (Property->GetBoundObjects().Contains(OuterObject))
+			if (Property->ContainsBoundObjects(OuterObjects))
 			{
 				bFoundPropForObject = true;
 				break;
@@ -798,8 +797,7 @@ void SRemoteControlPanel::Unexpose(const TSharedPtr<IPropertyHandle>& Handle)
 {
 	TArray<UObject*> OuterObjects;
 	Handle->GetOuterObjects(OuterObjects);
-
-	FString Path = Handle->GeneratePathToProperty();
+	const FString Path = Handle->GeneratePathToProperty();
 
 	// Find an exposed property with the same path.
 	TArray<TSharedPtr<FRemoteControlProperty>, TInlineAllocator<1>> PotentialMatches;
@@ -807,7 +805,7 @@ void SRemoteControlPanel::Unexpose(const TSharedPtr<IPropertyHandle>& Handle)
 	{
 		if (TSharedPtr<FRemoteControlProperty> Property = WeakProperty.Pin())
 		{
-			if (Property->FieldPathInfo.ToPathPropertyString() == Path)
+			if (Property->CheckIsBoundToPropertyPath(Path))
 			{
 				PotentialMatches.Add(Property);
 			}
@@ -816,23 +814,7 @@ void SRemoteControlPanel::Unexpose(const TSharedPtr<IPropertyHandle>& Handle)
 
 	for (const TSharedPtr<FRemoteControlProperty>& Property : PotentialMatches)
 	{
-		TArray<UObject*> PropertyOwners = Property->GetBoundObjects();
-		if (PropertyOwners.Num() != OuterObjects.Num())
-		{
-			continue;
-		}
-
-		bool bHasSameObjects = true;
-		for (UObject* Owner : PropertyOwners)
-		{
-			if (!OuterObjects.Contains(Owner))
-			{
-				bHasSameObjects = false;
-				break;
-			}
-		}
-
-		if (bHasSameObjects && OuterObjects.Num())
+		if (Property->ContainsBoundObjects(OuterObjects))
 		{
 			Preset->Unexpose(Property->GetId());
 			break;
