FPrivateToken trick used to allow emplacement-style functions (like MakeShared) to compile without giving up constructor access rights.

This replaces classes' use of friendship with internal TSharedPtr functions.

diff --git a/Engine/Plugins/Experimental/CommonConversation/Source/CommonConversationRuntime/Private/ConversationRegistry.cpp b/Engine/Plugins/Experimental/CommonConversation/Source/CommonConversationRuntime/Private/ConversationRegistry.cpp
--- a/Engine/Plugins/Experimental/CommonConversation/Source/CommonConversationRuntime/Private/ConversationRegistry.cpp
+++ b/Engine/Plugins/Experimental/CommonConversation/Source/CommonConversationRuntime/Private/ConversationRegistry.cpp
@@ -17,13 +17,13 @@
 
 TSharedPtr<FConversationsHandle> FConversationsHandle::Create(const UConversationRegistry* InOwningRegistry, const TSharedPtr<FStreamableHandle>& InStreamableHandle, const TArray<FGameplayTag>& InEntryTags)
 {
-	TSharedPtr<FConversationsHandle> ConversationsHandle = MakeShared<FConversationsHandle>(const_cast<UConversationRegistry*>(InOwningRegistry), InStreamableHandle, InEntryTags);
+	TSharedPtr<FConversationsHandle> ConversationsHandle = MakeShared<FConversationsHandle>(FPrivateToken{}, const_cast<UConversationRegistry*>(InOwningRegistry), InStreamableHandle, InEntryTags);
 	ConversationsHandle->Initialize();
 
 	return ConversationsHandle;
 }
 
-FConversationsHandle::FConversationsHandle(UConversationRegistry* InOwningRegistry, const TSharedPtr<FStreamableHandle>& InStreamableHandle, const TArray<FGameplayTag>& InEntryTags)
+FConversationsHandle::FConversationsHandle(FPrivateToken, UConversationRegistry* InOwningRegistry, const TSharedPtr<FStreamableHandle>& InStreamableHandle, const TArray<FGameplayTag>& InEntryTags)
 	: StreamableHandle(InStreamableHandle)
 	, ConversationEntryTags(InEntryTags)
 	, OwningRegistryPtr(InOwningRegistry)
diff --git a/Engine/Plugins/Experimental/CommonConversation/Source/CommonConversationRuntime/Public/ConversationRegistry.h b/Engine/Plugins/Experimental/CommonConversation/Source/CommonConversationRuntime/Public/ConversationRegistry.h
--- a/Engine/Plugins/Experimental/CommonConversation/Source/CommonConversationRuntime/Public/ConversationRegistry.h
+++ b/Engine/Plugins/Experimental/CommonConversation/Source/CommonConversationRuntime/Public/ConversationRegistry.h
@@ -44,8 +44,13 @@ struct COMMONCONVERSATIONRUNTIME_API FConversationsHandle : public TSharedFromTh
 	static TSharedPtr<FConversationsHandle> Create(const UConversationRegistry* InOwningRegistry, const TSharedPtr<FStreamableHandle>& InStreamableHandle, const TArray<FGameplayTag>& InEntryTags);
 
 private:
-	FConversationsHandle(UConversationRegistry* InOwningRegistry, const TSharedPtr<FStreamableHandle>& InStreamableHandle, const TArray<FGameplayTag>& InEntryTags);
+	// Private token only allows members or friends to call MakeShared
+	struct FPrivateToken { explicit FPrivateToken() = default; };
 
+public:
+	FConversationsHandle(FPrivateToken, UConversationRegistry* InOwningRegistry, const TSharedPtr<FStreamableHandle>& InStreamableHandle, const TArray<FGameplayTag>& InEntryTags);
+
+private:
 	void Initialize();
 	void HandleAvailableConversationsChanged();
 
@@ -53,8 +58,6 @@ private:
 	TSharedPtr<FStreamableHandle> StreamableHandle;
 	TArray<FGameplayTag> ConversationEntryTags;
 	TWeakObjectPtr<UConversationRegistry> OwningRegistryPtr;
-
-	friend class SharedPointerInternals::TIntrusiveReferenceController<FConversationsHandle>;
 };
 
 
diff --git a/Engine/Plugins/VirtualProduction/RemoteControl/Source/RemoteControlProtocolWidgets/Private/ViewModels/ProtocolBindingViewModel.cpp b/Engine/Plugins/VirtualProduction/RemoteControl/Source/RemoteControlProtocolWidgets/Private/ViewModels/ProtocolBindingViewModel.cpp
--- a/Engine/Plugins/VirtualProduction/RemoteControl/Source/RemoteControlProtocolWidgets/Private/ViewModels/ProtocolBindingViewModel.cpp
+++ b/Engine/Plugins/VirtualProduction/RemoteControl/Source/RemoteControlProtocolWidgets/Private/ViewModels/ProtocolBindingViewModel.cpp
@@ -22,13 +22,13 @@ namespace Commands
 	{
 		FGuid EntityId;
 		FGuid BindingId;
-		FGuid RangeId;
+		FGuid RangeId;	
 		TArray<FGuid> RangeIds; // Specify if there are multiple to remove
 	};
-
+	
 	static FGuid AddRangeMappingInternal(URemoteControlPreset* InPreset, const FAddRemoveRangeArgs& InArgs)
 	{
-		if (TSharedPtr<FRemoteControlProperty> RCProperty = InPreset->GetExposedEntity<FRemoteControlProperty>(InArgs.EntityId).Pin())
+		if(TSharedPtr<FRemoteControlProperty> RCProperty = InPreset->GetExposedEntity<FRemoteControlProperty>(InArgs.EntityId).Pin())
 		{
 			FRemoteControlProtocolBinding* ProtocolBinding = RCProperty->ProtocolBindings.FindByHash(GetTypeHash(InArgs.BindingId), InArgs.BindingId);
 			check(ProtocolBinding);
@@ -49,19 +49,19 @@ namespace Commands
 			}
 			else
 			{
-				const FRemoteControlProtocolMapping RangesData(RCProperty->GetProperty(), (*EntityPtr)->GetRangePropertySize());
+			const FRemoteControlProtocolMapping RangesData(RCProperty->GetProperty(), (*EntityPtr)->GetRangePropertySize());
 				ProtocolBinding->AddMapping(RangesData);
 
-				return RangesData.GetId();
-			}
+			return RangesData.GetId();
 		}
-
+		}
+		
 		return FGuid();
 	}
 
 	static bool RemoveRangeMappingInternal(URemoteControlPreset* InPreset, const FAddRemoveRangeArgs& InArgs)
 	{
-		if (TSharedPtr<FRemoteControlProperty> RCProperty = InPreset->GetExposedEntity<FRemoteControlProperty>(InArgs.EntityId).Pin())
+		if(TSharedPtr<FRemoteControlProperty> RCProperty = InPreset->GetExposedEntity<FRemoteControlProperty>(InArgs.EntityId).Pin())
 		{
 			FRemoteControlProtocolBinding* ProtocolBinding = RCProperty->ProtocolBindings.FindByHash(GetTypeHash(InArgs.BindingId), InArgs.BindingId);
 			check(ProtocolBinding);
@@ -79,7 +79,7 @@ namespace Commands
 
 			return NumRemoved > 0;
 		}
-
+		
 		return false;
 	}
 }
@@ -95,13 +95,13 @@ TMap<FProtocolBindingViewModel::EValidity, FText> FProtocolBindingViewModel::Val
 
 TSharedRef<FProtocolBindingViewModel> FProtocolBindingViewModel::Create(const TSharedRef<FProtocolEntityViewModel>& InParentViewModel, const TSharedRef<FRemoteControlProtocolBinding>& InBinding)
 {
-	TSharedRef<FProtocolBindingViewModel> ViewModel = MakeShared<FProtocolBindingViewModel>(InParentViewModel, InBinding);
+	TSharedRef<FProtocolBindingViewModel> ViewModel = MakeShared<FProtocolBindingViewModel>(FPrivateToken{}, InParentViewModel, InBinding);
 	ViewModel->Initialize();
 
 	return ViewModel;
 }
 
-FProtocolBindingViewModel::FProtocolBindingViewModel(const TSharedRef<FProtocolEntityViewModel>& InParentViewModel, const TSharedRef<FRemoteControlProtocolBinding>& InBinding)
+FProtocolBindingViewModel::FProtocolBindingViewModel(FPrivateToken, const TSharedRef<FProtocolEntityViewModel>& InParentViewModel, const TSharedRef<FRemoteControlProtocolBinding>& InBinding)
 	: Preset(InParentViewModel->Preset)
 	, ParentViewModel(InParentViewModel)
 	, Property(InParentViewModel->Property)
@@ -122,7 +122,7 @@ void FProtocolBindingViewModel::Initialize()
 {
 	FRemoteControlProtocolBinding* Binding = GetBinding();
 	// May be stale as a result of an undo deleting it.
-	if (Binding)
+	if(Binding)
 	{
 		if (!RemoteControlTypeUtilities::IsSupportedMappingType(GetProperty().Get()))
 		{
@@ -131,9 +131,9 @@ void FProtocolBindingViewModel::Initialize()
 
 		Ranges.Empty(Ranges.Num());
 		GetBinding()->ForEachMapping([&](FRemoteControlProtocolMapping& InMapping)
-		{
-			Ranges.Emplace(FProtocolRangeViewModel::Create(SharedThis(this), InMapping.GetId()));
-		});
+        { 
+            Ranges.Emplace(FProtocolRangeViewModel::Create(SharedThis(this), InMapping.GetId()));
+        });
 	}
 }
 
@@ -148,20 +148,20 @@ FGuid FProtocolBindingViewModel::AddRangeMapping()
 
 	using FCommandChange = TRemoteControlProtocolCommandChange<Commands::FAddRemoveRangeArgs>;
 	using FOnChange = FCommandChange::FOnUndoRedoDelegate;
-	if (GUndo)
+	if(GUndo)
 	{
 		// What to do on undo and redo
 		GUndo->StoreUndo(Preset.Get(),
-						MakeUnique<FCommandChange>(
-							Preset.Get(),
-							Commands::FAddRemoveRangeArgs{ParentViewModel.Pin()->GetId(), GetId(), NewRangeViewModel->GetId()},
-							FOnChange::CreateLambda([](URemoteControlPreset* InPreset, const Commands::FAddRemoveRangeArgs& InChangeArgs) { Commands::AddRangeMappingInternal(InPreset, InChangeArgs); }),
-							FOnChange::CreateLambda([](URemoteControlPreset* InPreset, const Commands::FAddRemoveRangeArgs& InChangeArgs) { Commands::RemoveRangeMappingInternal(InPreset, InChangeArgs); })
-						)
-		);
+            MakeUnique<FCommandChange>(
+                Preset.Get(),
+                Commands::FAddRemoveRangeArgs{ParentViewModel.Pin()->GetId(), GetId(), NewRangeViewModel->GetId()},
+                FOnChange::CreateLambda([](URemoteControlPreset* InPreset, const Commands::FAddRemoveRangeArgs& InChangeArgs){ Commands::AddRangeMappingInternal(InPreset, InChangeArgs); }),
+                FOnChange::CreateLambda([](URemoteControlPreset* InPreset, const Commands::FAddRemoveRangeArgs& InChangeArgs){ Commands::RemoveRangeMappingInternal(InPreset, InChangeArgs); })
+            )
+        );
 	}
 
-	OnRangeMappingAddedDelegate.Broadcast(NewRangeViewModel.ToSharedRef());
+	OnRangeMappingAddedDelegate.Broadcast(NewRangeViewModel.ToSharedRef());	
 
 	return NewRangeViewModel->GetId();
 }
@@ -191,27 +191,27 @@ void FProtocolBindingViewModel::RemoveRangeMapping(const FGuid& InId)
 
 	using FCommandChange = TRemoteControlProtocolCommandChange<Commands::FAddRemoveRangeArgs>;
 	using FOnChange = FCommandChange::FOnUndoRedoDelegate;
-	if (GUndo)
+	if(GUndo)
 	{
 		// What to do on undo and redo
 		GUndo->StoreUndo(Preset.Get(),
-						MakeUnique<FCommandChange>(
-							Preset.Get(),
-							Commands::FAddRemoveRangeArgs{ParentViewModel.Pin()->GetId(), GetId(), InId},
-							FOnChange::CreateLambda([](URemoteControlPreset* InPreset, const Commands::FAddRemoveRangeArgs& InChangeArgs) { Commands::RemoveRangeMappingInternal(InPreset, InChangeArgs); }),
-							FOnChange::CreateLambda([](URemoteControlPreset* InPreset, const Commands::FAddRemoveRangeArgs& InChangeArgs) { Commands::AddRangeMappingInternal(InPreset, InChangeArgs); })
-						)
-		);
+            MakeUnique<FCommandChange>(
+                Preset.Get(),
+                Commands::FAddRemoveRangeArgs{ParentViewModel.Pin()->GetId(), GetId(), InId},
+                FOnChange::CreateLambda([](URemoteControlPreset* InPreset, const Commands::FAddRemoveRangeArgs& InChangeArgs){ Commands::RemoveRangeMappingInternal(InPreset, InChangeArgs); }),
+                FOnChange::CreateLambda([](URemoteControlPreset* InPreset, const Commands::FAddRemoveRangeArgs& InChangeArgs){ Commands::AddRangeMappingInternal(InPreset, InChangeArgs); })
+            )
+        );
 	}
 
-	if (Commands::RemoveRangeMappingInternal(Preset.Get(), Commands::FAddRemoveRangeArgs{ParentViewModel.Pin()->GetId(), GetId(), InId}))
+	if(Commands::RemoveRangeMappingInternal(Preset.Get(), Commands::FAddRemoveRangeArgs{ParentViewModel.Pin()->GetId(), GetId(), InId}))
 	{
 		const int32 NumItemsRemoved = Ranges.RemoveAll([&](const TSharedPtr<FProtocolRangeViewModel>& InRange)
-		{
-			return InRange->GetId() == InId;
-		});
+        {
+            return InRange->GetId() == InId;
+        });
 
-		if (NumItemsRemoved <= 0)
+		if(NumItemsRemoved <= 0)
 			return;
 
 		OnRangeMappingRemovedDelegate.Broadcast(InId);
@@ -286,56 +286,56 @@ void FProtocolBindingViewModel::AddDefaultRangeMappings()
 				int64 IntMin = RemoteControlTypeUtilities::GetDefaultRangeValueMin<int64>(NumericProperty);
 				int64 IntMax = RemoteControlTypeUtilities::GetDefaultRangeValueMax<int64>(NumericProperty);
 
-				// fixup typename according to typesize, ie. it can be a UInt32 but a typesize of 2 makes its a UInt16
-				if (RangePropertyTypeName == NAME_UInt32Property && RangePropertySize > 0)
-				{
-					if (RangePropertySize == sizeof(uint8))
-					{
+		// fixup typename according to typesize, ie. it can be a UInt32 but a typesize of 2 makes its a UInt16
+		if(RangePropertyTypeName == NAME_UInt32Property && RangePropertySize > 0)
+		{
+			if(RangePropertySize == sizeof(uint8))
+			{
 						IntMin = RemoteControlTypeUtilities::GetDefaultRangeValueMin<uint8>(NumericProperty);
 						IntMax = RemoteControlTypeUtilities::GetDefaultRangeValueMax<uint8>(NumericProperty);
-						RangePropertyTypeName = NAME_ByteProperty;
-					}
-					else if (RangePropertySize == sizeof(uint16))
-					{
+				RangePropertyTypeName = NAME_ByteProperty;
+			}
+			else if(RangePropertySize == sizeof(uint16))
+			{
 						IntMin = RemoteControlTypeUtilities::GetDefaultRangeValueMin<uint16>(NumericProperty);
 						IntMax = RemoteControlTypeUtilities::GetDefaultRangeValueMax<uint16>(NumericProperty);
-						RangePropertyTypeName = NAME_UInt16Property;
-					}
-					else if (RangePropertySize == sizeof(uint64))
-					{
+				RangePropertyTypeName = NAME_UInt16Property; 
+			}
+			else if(RangePropertySize == sizeof(uint64))
+			{
 						IntMin = RemoteControlTypeUtilities::GetDefaultRangeValueMin<uint64>(NumericProperty);
 						IntMax = RemoteControlTypeUtilities::GetDefaultRangeValueMax<uint64>(NumericProperty);
-						RangePropertyTypeName = NAME_UInt64Property;
-					}
-				}
+				RangePropertyTypeName = NAME_UInt64Property;
+			}
+		}
 
-				if (RangePropertyTypeName == NAME_ByteProperty)
+				if(RangePropertyTypeName == NAME_ByteProperty)
 				{
 					MinItem->SetInputValue<uint8>(IntMin);
 					MaxItem->SetInputValue<uint8>(IntMax);
 				}
-				else if (RangePropertyTypeName == NAME_UInt16Property)
+				else if(RangePropertyTypeName == NAME_UInt16Property)
 				{
 					MinItem->SetInputValue<uint16>(IntMin);
 					MaxItem->SetInputValue<uint16>(IntMax);
 				}
-				else if (RangePropertyTypeName == NAME_UInt32Property)
+				else if(RangePropertyTypeName == NAME_UInt32Property)
 				{
 					MinItem->SetInputValue<uint32>(IntMin);
 					MaxItem->SetInputValue<uint32>(IntMax);
 				}
-				else if (RangePropertyTypeName == NAME_UInt64Property)
+				else if(RangePropertyTypeName == NAME_UInt64Property)
 				{
 					MinItem->SetInputValue<uint64>(IntMin);
 					MaxItem->SetInputValue<uint64>(IntMax);
 				}
 			}
-			else if (NumericProperty->IsFloatingPoint())
+			else if(NumericProperty->IsFloatingPoint())
 			{
 				const float FloatMin = RemoteControlTypeUtilities::GetDefaultRangeValueMin<float>(RangeProperty);
 				const float FloatMax = RemoteControlTypeUtilities::GetDefaultRangeValueMax<float>(RangeProperty);
 
-				if (RangePropertyTypeName == NAME_FloatProperty)
+				if(RangePropertyTypeName == NAME_FloatProperty)
 				{
 					MinItem->SetInputValue<float>(FloatMin);
 					MaxItem->SetInputValue<float>(FloatMax);
@@ -346,15 +346,15 @@ void FProtocolBindingViewModel::AddDefaultRangeMappings()
 
 	const FProperty* MappingProperty = GetProperty().Get();
 	FName MappingPropertyTypeName = MappingProperty->GetClass()->GetFName();
-	if (const FStructProperty* StructProperty = CastField<FStructProperty>(MappingProperty))
+	if(const FStructProperty* StructProperty = CastField<FStructProperty>(MappingProperty))
 	{
 		MappingPropertyTypeName = StructProperty->Struct->GetFName();
 	}
 	// Mapping (output)
 	{
-		if (const FNumericProperty* NumericProperty = CastField<FNumericProperty>(MappingProperty))
+		if(const FNumericProperty* NumericProperty = CastField<FNumericProperty>(MappingProperty))
 		{
-			if (NumericProperty->IsInteger())
+			if(NumericProperty->IsInteger())
 			{
 				const int64 IntMin = RemoteControlTypeUtilities::GetDefaultMappingValueMin<int64>(NumericProperty);
 				const int64 IntMax = RemoteControlTypeUtilities::GetDefaultMappingValueMax<int64>(NumericProperty);
@@ -364,32 +364,32 @@ void FProtocolBindingViewModel::AddDefaultRangeMappings()
 					MinItem->SetOutputValue<uint8>(IntMin);
 					MaxItem->SetOutputValue<uint8>(IntMax);
 				}
-				else if (MappingPropertyTypeName == NAME_Int8Property)
+				else if(MappingPropertyTypeName == NAME_Int8Property)
 				{
 					MinItem->SetOutputValue<int8>(IntMin);
 					MaxItem->SetOutputValue<int8>(IntMax);
 				}
-				else if (MappingPropertyTypeName == NAME_Int16Property)
+				else if(MappingPropertyTypeName == NAME_Int16Property)
 				{
 					MinItem->SetOutputValue<int16>(IntMin);
 					MaxItem->SetOutputValue<int16>(IntMax);
 				}
-				else if (RangePropertyTypeName == NAME_UInt16Property)
+				else if(RangePropertyTypeName == NAME_UInt16Property)
 				{
 					MinItem->SetOutputValue<uint16>(IntMin);
 					MaxItem->SetOutputValue<uint16>(IntMax);
 				}
-				else if (MappingPropertyTypeName == NAME_Int32Property)
+				else if(MappingPropertyTypeName == NAME_Int32Property)
 				{
 					MinItem->SetOutputValue<int32>(IntMin);
 					MaxItem->SetOutputValue<int32>(IntMax);
 				}
-				else if (RangePropertyTypeName == NAME_UInt32Property)
+				else if(RangePropertyTypeName == NAME_UInt32Property)
 				{
 					MinItem->SetOutputValue<uint32>(IntMin);
 					MaxItem->SetOutputValue<uint32>(IntMax);
 				}
-				else if (MappingPropertyTypeName == NAME_Int64Property)
+				else if(MappingPropertyTypeName == NAME_Int64Property)
 				{
 					MinItem->SetOutputValue<int64>(IntMin);
 					MaxItem->SetOutputValue<int64>(IntMax);
@@ -400,7 +400,7 @@ void FProtocolBindingViewModel::AddDefaultRangeMappings()
 					MaxItem->SetOutputValue<uint64>(IntMax);
 				}
 			}
-			else if (NumericProperty->IsFloatingPoint())
+			else if(NumericProperty->IsFloatingPoint())
 			{
 				const double FloatMin = RemoteControlTypeUtilities::GetDefaultMappingValueMin<double>(MappingProperty);
 				const double FloatMax = RemoteControlTypeUtilities::GetDefaultMappingValueMax<double>(MappingProperty);
@@ -420,10 +420,10 @@ void FProtocolBindingViewModel::AddDefaultRangeMappings()
 		else
 		{
 			if (MappingPropertyTypeName == NAME_BoolProperty)
-			{
+				{
 				MinItem->SetOutputValue<bool>(RemoteControlTypeUtilities::GetDefaultMappingValueMin<bool>(MappingProperty));
 				MaxItem->SetOutputValue<bool>(RemoteControlTypeUtilities::GetDefaultMappingValueMax<bool>(MappingProperty));
-			}
+				}
 
 			else if (MappingPropertyTypeName == NAME_StrProperty)
 			{
@@ -438,11 +438,11 @@ void FProtocolBindingViewModel::AddDefaultRangeMappings()
 				MaxItem->SetOutputValue<Type>(RemoteControlTypeUtilities::GetDefaultMappingValueMax<Type>(MappingProperty));
 			}
 			else if (MappingPropertyTypeName == NAME_TextProperty)
-			{
+				{
 				using Type = TRCTypeNameToType<NAME_TextProperty>::Type;
 				MinItem->SetOutputValue<Type>(RemoteControlTypeUtilities::GetDefaultMappingValueMin<Type>(MappingProperty));
 				MaxItem->SetOutputValue<Type>(RemoteControlTypeUtilities::GetDefaultMappingValueMax<Type>(MappingProperty));
-			}
+				}
 
 			else if (MappingPropertyTypeName == NAME_Vector)
 			{
@@ -450,17 +450,17 @@ void FProtocolBindingViewModel::AddDefaultRangeMappings()
 				MaxItem->SetOutputValue<FVector>(RemoteControlTypeUtilities::GetDefaultMappingValueMax<FVector>(MappingProperty));
 			}
 			else if (MappingPropertyTypeName == NAME_Vector2D)
-			{
+				{
 				using Type = TRCTypeNameToType<NAME_Vector2D>::Type;
 				MinItem->SetOutputValue<Type>(RemoteControlTypeUtilities::GetDefaultMappingValueMin<Type>(MappingProperty));
 				MaxItem->SetOutputValue<Type>(RemoteControlTypeUtilities::GetDefaultMappingValueMax<Type>(MappingProperty));
-			}
+				}
 			else if (MappingPropertyTypeName == NAME_Vector4)
-			{
+				{
 				using Type = TRCTypeNameToType<NAME_Vector4>::Type;
 				MinItem->SetOutputValue<Type>(RemoteControlTypeUtilities::GetDefaultMappingValueMin<Type>(MappingProperty));
 				MaxItem->SetOutputValue<Type>(RemoteControlTypeUtilities::GetDefaultMappingValueMax<Type>(MappingProperty));
-			}
+				}
 			else if (MappingPropertyTypeName == NAME_Rotator)
 			{
 				using Type = TRCTypeNameToType<NAME_Rotator>::Type;
@@ -472,9 +472,9 @@ void FProtocolBindingViewModel::AddDefaultRangeMappings()
 				using Type = TRCTypeNameToType<NAME_Color>::Type;
 				MinItem->SetOutputValue<Type>(RemoteControlTypeUtilities::GetDefaultMappingValueMin<Type>(MappingProperty));
 				MaxItem->SetOutputValue<Type>(RemoteControlTypeUtilities::GetDefaultMappingValueMax<Type>(MappingProperty));
-			}
+		}
 			else if (MappingPropertyTypeName == NAME_LinearColor)
-			{
+		{
 				using Type = TRCTypeNameToType<NAME_LinearColor>::Type;
 				MinItem->SetOutputValue<Type>(RemoteControlTypeUtilities::GetDefaultMappingValueMin<Type>(MappingProperty));
 				MaxItem->SetOutputValue<Type>(RemoteControlTypeUtilities::GetDefaultMappingValueMax<Type>(MappingProperty));
@@ -500,7 +500,7 @@ void FProtocolBindingViewModel::AddDefaultRangeMappings()
 			else
 			{
 				UE_LOG(LogRemoteControlProtocolWidgets, Warning, TEXT("AddDefaultRangeMappings type: %s for default value resolution."), *MappingPropertyTypeName.ToString());
-			}
+			}	
 		}
 	}
 }
@@ -526,7 +526,7 @@ TSharedPtr<IRemoteControlProtocol> FProtocolBindingViewModel::GetProtocol() cons
 FRemoteControlProtocolMapping* FProtocolBindingViewModel::GetRangesMapping(const FGuid& InRangeId) const
 {
 	check(IsValid());
-
+	
 	FRemoteControlProtocolMapping* Mapping = GetBinding()->FindMapping(InRangeId);
 	check(InRangeId == Mapping->GetId());
 	return Mapping;
diff --git a/Engine/Plugins/VirtualProduction/RemoteControl/Source/RemoteControlProtocolWidgets/Private/ViewModels/ProtocolBindingViewModel.h b/Engine/Plugins/VirtualProduction/RemoteControl/Source/RemoteControlProtocolWidgets/Private/ViewModels/ProtocolBindingViewModel.h
--- a/Engine/Plugins/VirtualProduction/RemoteControl/Source/RemoteControlProtocolWidgets/Private/ViewModels/ProtocolBindingViewModel.h
+++ b/Engine/Plugins/VirtualProduction/RemoteControl/Source/RemoteControlProtocolWidgets/Private/ViewModels/ProtocolBindingViewModel.h
@@ -109,13 +109,17 @@ public:
 	/** Allows the view to notify the ViewModel of a change (ie. from SDetailsView) */
 	void NotifyChanged() const;
 
-protected:
-	template <typename FProtocolBindingViewModel>
-	friend class SharedPointerInternals::TIntrusiveReferenceController;
+private:
+	friend class FProtocolEntityViewModel;
+	
+	// Private token only allows members or friends to call MakeShared
+	struct FPrivateToken { explicit FPrivateToken() = default; };
 
-	FProtocolBindingViewModel() = default;
-	FProtocolBindingViewModel(const TSharedRef<FProtocolEntityViewModel>& InParentViewModel, const TSharedRef<FRemoteControlProtocolBinding>& InBinding);
+public:
+	FProtocolBindingViewModel(FPrivateToken) {}
+	FProtocolBindingViewModel(FPrivateToken, const TSharedRef<FProtocolEntityViewModel>& InParentViewModel, const TSharedRef<FRemoteControlProtocolBinding>& InBinding);
 
+private:
 	void Initialize();
 
 	TSharedPtr<FProtocolRangeViewModel>& AddRangeMappingInternal();
@@ -149,7 +153,7 @@ private:
 
 	/** Range ViewModels for this protocol binding. */
 	TArray<TSharedPtr<FProtocolRangeViewModel>> Ranges;
-
+	
 	/** When single range mapping added. */
 	FOnRangeMappingAdded OnRangeMappingAddedDelegate;
 
diff --git a/Engine/Plugins/VirtualProduction/RemoteControl/Source/RemoteControlProtocolWidgets/Private/ViewModels/ProtocolEntityViewModel.cpp b/Engine/Plugins/VirtualProduction/RemoteControl/Source/RemoteControlProtocolWidgets/Private/ViewModels/ProtocolEntityViewModel.cpp
--- a/Engine/Plugins/VirtualProduction/RemoteControl/Source/RemoteControlProtocolWidgets/Private/ViewModels/ProtocolEntityViewModel.cpp
+++ b/Engine/Plugins/VirtualProduction/RemoteControl/Source/RemoteControlProtocolWidgets/Private/ViewModels/ProtocolEntityViewModel.cpp
@@ -72,13 +72,13 @@ TMap<FProtocolEntityViewModel::EValidity, FText> FProtocolEntityViewModel::Valid
 
 TSharedRef<FProtocolEntityViewModel> FProtocolEntityViewModel::Create(URemoteControlPreset* InPreset, const FGuid& InEntityId)
 {
-	TSharedRef<FProtocolEntityViewModel> ViewModel = MakeShared<FProtocolEntityViewModel>(InPreset, InEntityId);
+	TSharedRef<FProtocolEntityViewModel> ViewModel = MakeShared<FProtocolEntityViewModel>(FPrivateToken{}, InPreset, InEntityId);
 	ViewModel->Initialize();
 
 	return ViewModel;
 }
 
-FProtocolEntityViewModel::FProtocolEntityViewModel(URemoteControlPreset* InPreset, const FGuid& InEntityId)
+FProtocolEntityViewModel::FProtocolEntityViewModel(FPrivateToken, URemoteControlPreset* InPreset, const FGuid& InEntityId)
 	: Preset(InPreset)
 	, PropertyId(InEntityId)
 {
@@ -102,7 +102,7 @@ void FProtocolEntityViewModel::Initialize()
 	if (TSharedPtr<FRemoteControlProperty> RCProperty = Preset->GetExposedEntity<FRemoteControlProperty>(PropertyId).Pin())
 	{
 		Property = RCProperty->GetProperty();
-		for (FRemoteControlProtocolBinding& Binding : RCProperty->ProtocolBindings)
+		for(FRemoteControlProtocolBinding& Binding : RCProperty->ProtocolBindings)
 		{
 			const TSharedPtr<IRemoteControlProtocol> Protocol = IRemoteControlProtocolModule::Get().GetProtocolByName(Binding.GetProtocolName());
 			// Supporting plugin needs to be loaded/protocol available
@@ -165,18 +165,18 @@ TSharedPtr<FProtocolBindingViewModel> FProtocolEntityViewModel::AddBinding(const
 		if (GUndo)
 		{
 			GUndo->StoreUndo(Preset.Get(),
-							MakeUnique<FCommandChange>(
-								Preset.Get(),
-								Commands::FAddRemoveProtocolArgs{GetId(), InProtocolName, ProtocolBinding->GetId()},
-								FOnChange::CreateLambda([](URemoteControlPreset* InPreset, const Commands::FAddRemoveProtocolArgs& InChangeArgs) { Commands::AddProtocolInternal(InPreset, InChangeArgs); }),
-								FOnChange::CreateLambda([](URemoteControlPreset* InPreset, const Commands::FAddRemoveProtocolArgs& InChangeArgs) { Commands::RemoveProtocolInternal(InPreset, InChangeArgs); })
-							)
-			);
+                MakeUnique<FCommandChange>(
+                    Preset.Get(),
+                    Commands::FAddRemoveProtocolArgs{GetId(), InProtocolName, ProtocolBinding->GetId()},
+                    FOnChange::CreateLambda([](URemoteControlPreset* InPreset, const Commands::FAddRemoveProtocolArgs& InChangeArgs){ Commands::AddProtocolInternal(InPreset, InChangeArgs); }),
+                    FOnChange::CreateLambda([](URemoteControlPreset* InPreset, const Commands::FAddRemoveProtocolArgs& InChangeArgs){ Commands::RemoveProtocolInternal(InPreset, InChangeArgs); })
+                )
+            );
 		}
 
 		TSharedPtr<FProtocolBindingViewModel>& NewBindingViewModel = Bindings.Add_GetRef(FProtocolBindingViewModel::Create(AsShared(), ProtocolBinding.ToSharedRef()));
 		NewBindingViewModel->AddDefaultRangeMappings();
-
+		
 		OnBindingAddedDelegate.Broadcast(NewBindingViewModel.ToSharedRef());
 
 		return NewBindingViewModel;
@@ -200,20 +200,20 @@ void FProtocolEntityViewModel::RemoveBinding(const FGuid& InBindingId)
 		if (GUndo)
 		{
 			GUndo->StoreUndo(Preset.Get(),
-							MakeUnique<FCommandChange>(
-								Preset.Get(),
-								Commands::FAddRemoveProtocolArgs{GetId(), RemovedProtocolName, InBindingId},
-								FOnChange::CreateLambda([](URemoteControlPreset* InPreset, const Commands::FAddRemoveProtocolArgs& InChangeArgs) { Commands::RemoveProtocolInternal(InPreset, InChangeArgs); }),
-								FOnChange::CreateLambda([](URemoteControlPreset* InPreset, const Commands::FAddRemoveProtocolArgs& InChangeArgs) { Commands::AddProtocolInternal(InPreset, InChangeArgs); })
-							)
-			);
+                MakeUnique<FCommandChange>(
+                    Preset.Get(),
+                    Commands::FAddRemoveProtocolArgs{GetId(), RemovedProtocolName, InBindingId},
+                    FOnChange::CreateLambda([](URemoteControlPreset* InPreset, const Commands::FAddRemoveProtocolArgs& InChangeArgs){ Commands::RemoveProtocolInternal(InPreset, InChangeArgs); }),
+                    FOnChange::CreateLambda([](URemoteControlPreset* InPreset, const Commands::FAddRemoveProtocolArgs& InChangeArgs){ Commands::AddProtocolInternal(InPreset, InChangeArgs); })
+                )
+            );
 		}
 	}
 
 	const int32 NumItemsRemoved = Bindings.RemoveAll([&](const TSharedPtr<FProtocolBindingViewModel>& InBinding)
-	{
-		return InBinding->GetId() == InBindingId;
-	});
+    {
+        return InBinding->GetId() == InBindingId;
+    });
 
 	if (NumItemsRemoved <= 0)
 	{
@@ -243,7 +243,7 @@ TArray<TSharedPtr<FProtocolBindingViewModel>> FProtocolEntityViewModel::GetFilte
 	}
 
 	TArray<TSharedPtr<FProtocolBindingViewModel>> FilteredBindings;
-	for (const TSharedPtr<FProtocolBindingViewModel>& Binding : Bindings)
+	for(const TSharedPtr<FProtocolBindingViewModel>& Binding : Bindings)
 	{
 		// Make sure this bindings protocol name is NOT in the Hidden list, then add it to FilteredBindings
 		if (Algo::NoneOf(InHiddenProtocolTypeNames, [&Binding](const FName& InHiddenTypeName)
@@ -254,7 +254,7 @@ TArray<TSharedPtr<FProtocolBindingViewModel>> FProtocolEntityViewModel::GetFilte
 			FilteredBindings.Add(Binding);
 		}
 	}
-
+	
 	return FilteredBindings;
 }
 
@@ -329,4 +329,4 @@ void FProtocolEntityViewModel::PostUndo(bool bSuccess)
 	OnChangedDelegate.Broadcast();
 }
 
-#undef LOCTEXT_NAMESPACE
+#undef LOCTEXT_NAMESPACE
\ No newline at end of file
diff --git a/Engine/Plugins/VirtualProduction/RemoteControl/Source/RemoteControlProtocolWidgets/Private/ViewModels/ProtocolEntityViewModel.h b/Engine/Plugins/VirtualProduction/RemoteControl/Source/RemoteControlProtocolWidgets/Private/ViewModels/ProtocolEntityViewModel.h
--- a/Engine/Plugins/VirtualProduction/RemoteControl/Source/RemoteControlProtocolWidgets/Private/ViewModels/ProtocolEntityViewModel.h
+++ b/Engine/Plugins/VirtualProduction/RemoteControl/Source/RemoteControlProtocolWidgets/Private/ViewModels/ProtocolEntityViewModel.h
@@ -55,7 +55,7 @@ public:
 
 	/** Get the Protocol Binding ViewModels, filtered by the provided list of hidden protocol type names */
 	TArray<TSharedPtr<FProtocolBindingViewModel>> GetFilteredBindings(const TSet<FName>& InHiddenProtocolTypeNames);
-
+	
 	/** Checks if this entity is bound to one or more objects. */
 	bool IsBound() const;
 
@@ -79,13 +79,15 @@ public:
 	/** Something has changed within the ViewModel */
 	FOnChanged& OnChanged() { return OnChangedDelegate; }
 
-protected:
-	template <typename FProtocolEntityViewModel>
-	friend class SharedPointerInternals::TIntrusiveReferenceController;
-
-	FProtocolEntityViewModel() = default;
-	FProtocolEntityViewModel(URemoteControlPreset* InPreset, const FGuid& InEntityId);
+private:
+	// Private token only allows members or friends to call MakeShared
+	struct FPrivateToken { explicit FPrivateToken() = default; };
+	
+public:
+	FProtocolEntityViewModel(FPrivateToken) {}
+	FProtocolEntityViewModel(FPrivateToken, URemoteControlPreset* InPreset, const FGuid& InEntityId);
 
+protected:
 	void Initialize();
 
 	/** Respond when entity is unexposed */
diff --git a/Engine/Plugins/VirtualProduction/RemoteControl/Source/RemoteControlProtocolWidgets/Private/ViewModels/ProtocolRangeViewModel.cpp b/Engine/Plugins/VirtualProduction/RemoteControl/Source/RemoteControlProtocolWidgets/Private/ViewModels/ProtocolRangeViewModel.cpp
--- a/Engine/Plugins/VirtualProduction/RemoteControl/Source/RemoteControlProtocolWidgets/Private/ViewModels/ProtocolRangeViewModel.cpp
+++ b/Engine/Plugins/VirtualProduction/RemoteControl/Source/RemoteControlProtocolWidgets/Private/ViewModels/ProtocolRangeViewModel.cpp
@@ -1,4 +1,4 @@
-﻿// Copyright Epic Games, Inc. All Rights Reserved.
+// Copyright Epic Games, Inc. All Rights Reserved.
 
 #include "ProtocolRangeViewModel.h"
 
@@ -18,13 +18,13 @@ TMap<FProtocolRangeViewModel::EValidity, FText> FProtocolRangeViewModel::Validit
 
 TSharedRef<FProtocolRangeViewModel> FProtocolRangeViewModel::Create(const TSharedRef<FProtocolBindingViewModel>& InParentViewModel, const FGuid& InRangeId)
 {
-	TSharedRef<FProtocolRangeViewModel> ViewModel = MakeShared<FProtocolRangeViewModel>(InParentViewModel, InRangeId);
+	TSharedRef<FProtocolRangeViewModel> ViewModel = MakeShared<FProtocolRangeViewModel>(FPrivateToken{}, InParentViewModel, InRangeId);
 	ViewModel->Initialize();
 
 	return ViewModel;
 }
 
-FProtocolRangeViewModel::FProtocolRangeViewModel(const TSharedRef<FProtocolBindingViewModel>& InParentViewModel, const FGuid& InRangeId)
+FProtocolRangeViewModel::FProtocolRangeViewModel(FPrivateToken, const TSharedRef<FProtocolBindingViewModel>& InParentViewModel, const FGuid& InRangeId)
 	: Preset(InParentViewModel->Preset)
 	, ParentViewModel(InParentViewModel)
 	, RangeId(InRangeId)
diff --git a/Engine/Plugins/VirtualProduction/RemoteControl/Source/RemoteControlProtocolWidgets/Private/ViewModels/ProtocolRangeViewModel.h b/Engine/Plugins/VirtualProduction/RemoteControl/Source/RemoteControlProtocolWidgets/Private/ViewModels/ProtocolRangeViewModel.h
--- a/Engine/Plugins/VirtualProduction/RemoteControl/Source/RemoteControlProtocolWidgets/Private/ViewModels/ProtocolRangeViewModel.h
+++ b/Engine/Plugins/VirtualProduction/RemoteControl/Source/RemoteControlProtocolWidgets/Private/ViewModels/ProtocolRangeViewModel.h
@@ -52,16 +52,16 @@ public:
 
 	/** Set input */
 	template <typename ValueType>
-	void SetInputValue(ValueType InValue, bool bWithTransaction = false)
+    void SetInputValue(ValueType InValue, bool bWithTransaction = false)
 	{
 		check(IsValid());
 
-		if (bWithTransaction)
+		if(bWithTransaction)
 		{
 			FScopedTransaction Transaction(LOCTEXT("SetInputValue", "Set Input Value"));
-			Preset->Modify();
+			Preset->Modify();	
 		}
-
+		
 		GetRangesData()->SetRangeValue(InValue);
 	}
 
@@ -80,10 +80,10 @@ public:
 	{
 		check(IsValid());
 
-		if (bWithTransaction)
+		if(bWithTransaction)
 		{
 			FScopedTransaction Transaction(LOCTEXT("SetOutputValue", "Set Output Value"));
-			Preset->Modify();
+        	Preset->Modify();	
 		}
 
 		GetRangesData()->SetMappingValueAsPrimitive(InValue);
@@ -130,17 +130,21 @@ public:
 
 public:
 	DECLARE_MULTICAST_DELEGATE(FOnChanged);
-
+	
 	/** Something has changed within the ViewModel */
 	FOnChanged& OnChanged() { return OnChangedDelegate; }
 
-protected:
-	template <typename FProtocolRangeViewModel>
-	friend class SharedPointerInternals::TIntrusiveReferenceController;
+private:
+	// Private token only allows members or friends to call MakeShared
+	struct FPrivateToken { explicit FPrivateToken() = default; };
 
-	FProtocolRangeViewModel() = default;
-	FProtocolRangeViewModel(const TSharedRef<FProtocolBindingViewModel>& InParentViewModel, const FGuid& InRangeId);
+	friend class FProtocolBindingViewModel;
+	
+public:
+	FProtocolRangeViewModel(FPrivateToken) {}
+	FProtocolRangeViewModel(FPrivateToken, const TSharedRef<FProtocolBindingViewModel>& InParentViewModel, const FGuid& InRangeId);
 
+protected:
 	void Initialize();
 
 	/** Updates the input range value clamping (in case the type/precision has changed). */
@@ -170,7 +174,7 @@ private:
 private:
 	/** Owning Preset */
 	TWeakObjectPtr<URemoteControlPreset> Preset;
-
+	
 	/** Owning ProtocolBinding */
 	TWeakPtr<FProtocolBindingViewModel> ParentViewModel;
 
diff --git a/Engine/Source/Developer/BlueprintCompilerCppBackend/Private/BlueprintCompilerCppBackendGatherDependencies.cpp b/Engine/Source/Developer/BlueprintCompilerCppBackend/Private/BlueprintCompilerCppBackendGatherDependencies.cpp
--- a/Engine/Source/Developer/BlueprintCompilerCppBackend/Private/BlueprintCompilerCppBackendGatherDependencies.cpp
+++ b/Engine/Source/Developer/BlueprintCompilerCppBackend/Private/BlueprintCompilerCppBackendGatherDependencies.cpp
@@ -456,7 +456,7 @@ bool FGatherConvertedClassDependencies::IsFieldFromExcludedPackage(const UField*
 	return false;
 }
 
-FGatherConvertedClassDependencies::FGatherConvertedClassDependencies(UStruct* InStruct, const FCompilerNativizationOptions& InNativizationOptions)
+FGatherConvertedClassDependencies::FGatherConvertedClassDependencies(FPrivateToken, UStruct* InStruct, const FCompilerNativizationOptions& InNativizationOptions)
 	: OriginalStruct(InStruct)
 	, NativizationOptions(InNativizationOptions)
 {
@@ -849,7 +849,7 @@ TSharedPtr<FGatherConvertedClassDependencies> FGatherConvertedClassDependencies:
 		ConvertedClassDependenciesPtr = CachedConvertedClassDependencies.FindOrAdd(InStruct);
 		if (!ConvertedClassDependenciesPtr.IsValid())
 		{
-			ConvertedClassDependenciesPtr = MakeShared<FGatherConvertedClassDependencies>(InStruct, InNativizationOptions);
+			ConvertedClassDependenciesPtr = MakeShared<FGatherConvertedClassDependencies>(FPrivateToken{}, InStruct, InNativizationOptions);
 			check(ConvertedClassDependenciesPtr.IsValid());
 		}
 	}
diff --git a/Engine/Source/Developer/BlueprintCompilerCppBackend/Public/BlueprintCompilerCppBackendGatherDependencies.h b/Engine/Source/Developer/BlueprintCompilerCppBackend/Public/BlueprintCompilerCppBackendGatherDependencies.h
--- a/Engine/Source/Developer/BlueprintCompilerCppBackend/Public/BlueprintCompilerCppBackendGatherDependencies.h
+++ b/Engine/Source/Developer/BlueprintCompilerCppBackend/Public/BlueprintCompilerCppBackendGatherDependencies.h
@@ -59,11 +59,13 @@ public:
 
 	static bool IsFieldFromExcludedPackage(const UField* Field, const TSet<FName>& InExcludedModules);
 
-protected:
-	FGatherConvertedClassDependencies(UStruct* InStruct, const FCompilerNativizationOptions& InNativizationOptions);
+private:
+	// Private token only allows members or friends to call MakeShared
+	struct FPrivateToken { explicit FPrivateToken() = default; };
 
-	/** Friend for access to constructor via MakeShared */
-	friend class SharedPointerInternals::TIntrusiveReferenceController<FGatherConvertedClassDependencies>;
+public:
+	FGatherConvertedClassDependencies(FPrivateToken, UStruct* InStruct, const FCompilerNativizationOptions& InNativizationOptions);
 
+protected:
 	void DependenciesForHeader();
 };
diff --git a/Engine/Source/Runtime/Online/WebSockets/Private/Lws/LwsWebSocket.cpp b/Engine/Source/Runtime/Online/WebSockets/Private/Lws/LwsWebSocket.cpp
--- a/Engine/Source/Runtime/Online/WebSockets/Private/Lws/LwsWebSocket.cpp
+++ b/Engine/Source/Runtime/Online/WebSockets/Private/Lws/LwsWebSocket.cpp
@@ -72,7 +72,7 @@ const TCHAR* FLwsWebSocket::ToString(const EState InState)
 	return TEXT("Unknown");
 }
 
-FLwsWebSocket::FLwsWebSocket(const FString& InUrl, const TArray<FString>& InProtocols, const FString& InUpgradeHeader)
+FLwsWebSocket::FLwsWebSocket(FPrivateToken, const FString& InUrl, const TArray<FString>& InProtocols, const FString& InUpgradeHeader)
 	: State(EState::None)
 	, LastGameThreadState(EState::None)
 	, bWasSendQueueEmpty(true)
diff --git a/Engine/Source/Runtime/Online/WebSockets/Private/Lws/LwsWebSocket.h b/Engine/Source/Runtime/Online/WebSockets/Private/Lws/LwsWebSocket.h
--- a/Engine/Source/Runtime/Online/WebSockets/Private/Lws/LwsWebSocket.h
+++ b/Engine/Source/Runtime/Online/WebSockets/Private/Lws/LwsWebSocket.h
@@ -97,6 +97,8 @@ struct FLwsReceiveBufferText
 
 typedef TUniquePtr<FLwsReceiveBufferText> FLwsReceiveBufferTextPtr;
 
+class FLwsWebSocketsManager;
+
 class FLwsWebSocket
 	: public IWebSocket
 	, public TSharedFromThis<FLwsWebSocket>
@@ -185,12 +187,16 @@ public:
 	void LwsThreadTick();
 
 private:
-	/** Constructor */
-	FLwsWebSocket(const FString& Url, const TArray<FString>& Protocols, const FString& UpgradeHeader);
+	// Private token only allows members or friends to call MakeShared
+	struct FPrivateToken { explicit FPrivateToken() = default; };
+
+	friend class FLwsWebSocketsManager;
 
-	/** Friend for access to constructor via MakeShared */
-	friend class SharedPointerInternals::TIntrusiveReferenceController<FLwsWebSocket>;
+public:
+	/** Constructor */
+	FLwsWebSocket(FPrivateToken, const FString& Url, const TArray<FString>& Protocols, const FString& UpgradeHeader);
 
+private:
 	/**
 	 * Start connecting
 	 * @param LwsContext libwebsockets context
diff --git a/Engine/Source/Runtime/Online/WebSockets/Private/Lws/LwsWebSocketsManager.cpp b/Engine/Source/Runtime/Online/WebSockets/Private/Lws/LwsWebSocketsManager.cpp
--- a/Engine/Source/Runtime/Online/WebSockets/Private/Lws/LwsWebSocketsManager.cpp
+++ b/Engine/Source/Runtime/Online/WebSockets/Private/Lws/LwsWebSocketsManager.cpp
@@ -401,7 +401,7 @@ TSharedRef<IWebSocket> FLwsWebSocketsManager::CreateWebSocket(const FString& Url
 		UpgradeHeaderString += FString::Printf(TEXT("%s: %s\r\n"), *OneHeader.Key, *OneHeader.Value);
 	}
 
-	FLwsWebSocketRef Socket = MakeShared<FLwsWebSocket>(Url, Protocols, UpgradeHeaderString);
+	FLwsWebSocketRef Socket = MakeShared<FLwsWebSocket>(FLwsWebSocket::FPrivateToken{}, Url, Protocols, UpgradeHeaderString);
 	return Socket;
 }
 
diff --git a/Engine/Source/Runtime/Slate/Public/Framework/Text/ShapedTextCache.h b/Engine/Source/Runtime/Slate/Public/Framework/Text/ShapedTextCache.h
--- a/Engine/Source/Runtime/Slate/Public/Framework/Text/ShapedTextCache.h
+++ b/Engine/Source/Runtime/Slate/Public/Framework/Text/ShapedTextCache.h
@@ -57,7 +57,7 @@ public:
 	/** Create a new shaped text cache */
 	static FShapedTextCacheRef Create(const TSharedRef<FSlateFontCache>& InFontCache)
 	{
-		return MakeShared<FShapedTextCache>(InFontCache);
+		return MakeShared<FShapedTextCache>(FPrivateToken{}, InFontCache);
 	}
 
 	/**
@@ -109,14 +109,17 @@ public:
 	}
 
 private:
-	friend class SharedPointerInternals::TIntrusiveReferenceController<FShapedTextCache>;
+	// Private token only allows members or friends to call MakeShared
+	struct FPrivateToken { explicit FPrivateToken() = default; };
 
+public:
 	/** Constructor */
-	FShapedTextCache(const TSharedRef<FSlateFontCache>& InFontCache)
+	FShapedTextCache(FPrivateToken, const TSharedRef<FSlateFontCache>& InFontCache)
 		: FontCachePtr(InFontCache)
 	{
 	}
 
+private:
 	/** Font cache to use when shaping text */
 	TWeakPtr<FSlateFontCache> FontCachePtr;
 
