Bug Fix : See @ UE-130827

diff --git a/Engine/Plugins/VirtualProduction/RemoteControl/Source/RemoteControl/Private/Tests/RemoteControlPropertyHandleTestData.h b/Engine/Plugins/VirtualProduction/RemoteControl/Source/RemoteControl/Private/Tests/RemoteControlPropertyHandleTestData.h
--- a/Engine/Plugins/VirtualProduction/RemoteControl/Source/RemoteControl/Private/Tests/RemoteControlPropertyHandleTestData.h
+++ b/Engine/Plugins/VirtualProduction/RemoteControl/Source/RemoteControl/Private/Tests/RemoteControlPropertyHandleTestData.h
@@ -29,7 +29,7 @@ struct FRemoteControlTestStructInnerSimle
 	GENERATED_BODY()
 
 	UPROPERTY(EditAnywhere, Category = "RC")
-	int32 Int32Value;
+	int32 Int32Value = 0;
 };
 
 USTRUCT()
diff --git a/Engine/Plugins/VirtualProduction/RemoteControl/Source/RemoteControl/Public/RemoteControlInstanceMaterial.h b/Engine/Plugins/VirtualProduction/RemoteControl/Source/RemoteControl/Public/RemoteControlInstanceMaterial.h
--- a/Engine/Plugins/VirtualProduction/RemoteControl/Source/RemoteControl/Public/RemoteControlInstanceMaterial.h
+++ b/Engine/Plugins/VirtualProduction/RemoteControl/Source/RemoteControl/Public/RemoteControlInstanceMaterial.h
@@ -45,7 +45,7 @@ public:
 public:
 	/** Stores original binding class */
 	UPROPERTY()
-	UClass* OriginalClass;
+	UClass* OriginalClass = nullptr;
 
 	/** Stores original material parameter info */
 	UPROPERTY()
@@ -68,7 +68,7 @@ public:
 	//~ Begin IRemoteControlPropertyFactory interface
 	virtual TSharedPtr<FRemoteControlProperty> CreateRemoteControlProperty(URemoteControlPreset* Preset, UObject* Object, FRCFieldPathInfo FieldPath, FRemoteControlPresetExposeArgs Args) override;
 	virtual void PostSetObjectProperties(UObject* Object, bool bInSuccess) const override;
-	virtual bool SupportExposedClass(UClass* Class) const override;
+	virtual bool SupportExposedObject(UObject* InObject, const FRCFieldPathInfo FieldPath) const override;
 	//~ End IRemoteControlPropertyFactory interface
 };
 
