RM, fix onload property factories

diff --git a/Engine/Plugins/VirtualProduction/RemoteControl/Source/RemoteControl/Private/RemoteControlModule.cpp b/Engine/Plugins/VirtualProduction/RemoteControl/Source/RemoteControl/Private/RemoteControlModule.cpp
--- a/Engine/Plugins/VirtualProduction/RemoteControl/Source/RemoteControl/Private/RemoteControlModule.cpp
+++ b/Engine/Plugins/VirtualProduction/RemoteControl/Source/RemoteControl/Private/RemoteControlModule.cpp
@@ -369,8 +369,11 @@ void FRemoteControlModule::ShutdownModule()
 	// Unregister the interceptor feature on module shutdown
 	IModularFeatures::Get().UnregisterModularFeature(IRemoteControlInterceptionFeatureProcessor::GetName(), RCIProcessor.Get());
 
-	// Unregister Property Factories
-	UnregisterEntityFactory(FRemoteControlInstanceMaterial::StaticStruct()->GetFName());
+	if (UObjectInitialized())
+	{
+		// Unregister Property Factories
+		UnregisterEntityFactory(FRemoteControlInstanceMaterial::StaticStruct()->GetFName());
+	}
 }
 
 IRemoteControlModule::FOnPresetRegistered& FRemoteControlModule::OnPresetRegistered()
