Improve performance of the remote control panel

diff --git a/Engine/Plugins/VirtualProduction/RemoteControl/Source/RemoteControl/Private/RemoteControlPreset.cpp b/Engine/Plugins/VirtualProduction/RemoteControl/Source/RemoteControl/Private/RemoteControlPreset.cpp
--- a/Engine/Plugins/VirtualProduction/RemoteControl/Source/RemoteControl/Private/RemoteControlPreset.cpp
+++ b/Engine/Plugins/VirtualProduction/RemoteControl/Source/RemoteControl/Private/RemoteControlPreset.cpp
@@ -660,8 +660,6 @@ void URemoteControlPreset::PostLoad()
 
 	CacheFieldLayoutData();
 
-	RebindingManager->Rebind(this);
-
 	InitializeEntitiesMetadata();
 
 	RegisterEntityDelegates();
@@ -1054,6 +1052,12 @@ TArray<UObject*> URemoteControlPreset::ResolvedBoundObjects(FName FieldLabel)
 	return Objects;
 }
 
+void URemoteControlPreset::RebindUnboundEntities()
+{
+	RebindingManager->Rebind(this);
+	Algo::Transform(Registry->GetExposedEntities(), PerFrameUpdatedEntities, [](const TSharedPtr<FRemoteControlEntity>& Entity) { return Entity->GetId(); });
+}
+
 void URemoteControlPreset::NotifyExposedPropertyChanged(FName PropertyLabel)
 {
 	if (TSharedPtr<FRemoteControlProperty> ExposedProperty = GetExposedEntity<FRemoteControlProperty>(GetExposedEntityId(PropertyLabel)).Pin())
@@ -1690,10 +1694,9 @@ void URemoteControlPreset::OnEndFrame()
 
 void URemoteControlPreset::OnMapChange(uint32)
 {
-	// Delay the rebinding in order for the old actor points to be invalid, and for the new actors to be valid in the current map.
+	// Delay the refresh in order for the old actor points to be invalid.
 	GEditor->GetTimerManager()->SetTimerForNextTick(FTimerDelegate::CreateLambda([this]()
 		{
-			RebindingManager->Rebind(this);
 			Algo::Transform(Registry->GetExposedEntities(), PerFrameUpdatedEntities, [](const TSharedPtr<FRemoteControlEntity>& Entity) { return Entity->GetId(); });
 		}));
 }
diff --git a/Engine/Plugins/VirtualProduction/RemoteControl/Source/RemoteControl/Public/RemoteControlPreset.h b/Engine/Plugins/VirtualProduction/RemoteControl/Source/RemoteControl/Public/RemoteControlPreset.h
--- a/Engine/Plugins/VirtualProduction/RemoteControl/Source/RemoteControl/Public/RemoteControlPreset.h
+++ b/Engine/Plugins/VirtualProduction/RemoteControl/Source/RemoteControl/Public/RemoteControlPreset.h
@@ -691,8 +691,14 @@ public:
 	void CacheLayoutData();
 
 	/** Resolves exposed property/function bounded objects */
+	UE_DEPRECATED(4.27, "ResolvedBoundObjects is deprecated, you can now resolve bound objects using FRemoteControlEntity.")
 	TArray<UObject*> ResolvedBoundObjects(FName FieldLabel);
 
+	/**
+	 * Attempt to rebind all currently unbound properties.
+	 */
+	void RebindUnboundEntities();
+
 	DECLARE_MULTICAST_DELEGATE_TwoParams(FOnPresetEntityEvent, URemoteControlPreset* /*Preset*/, const FGuid& /*EntityId*/);
 	FOnPresetEntityEvent& OnEntityExposed() { return OnEntityExposedDelegate; }
 	FOnPresetEntityEvent& OnEntityUnexposed() { return OnEntityUnexposedDelegate; }
@@ -774,7 +780,6 @@ private:
 	void OnReplaceObjects(const TMap<UObject*, UObject*>& ReplacementObjectMap);
 	void OnEndFrame();
 	void OnMapChange(uint32 /*MapChangeFlags*/);
-	
 #endif
 	
 	/** Get a field ptr using it's id. */
diff --git a/Engine/Plugins/VirtualProduction/RemoteControl/Source/RemoteControlUI/Private/UI/RCPanelWidgetRegistry.cpp b/Engine/Plugins/VirtualProduction/RemoteControl/Source/RemoteControlUI/Private/UI/RCPanelWidgetRegistry.cpp
new file mode 100644
--- /dev/null
+++ b/Engine/Plugins/VirtualProduction/RemoteControl/Source/RemoteControlUI/Private/UI/RCPanelWidgetRegistry.cpp
@@ -0,0 +1,157 @@
+// Copyright Epic Games, Inc. All Rights Reserved.
+
+#include "RCPanelWidgetRegistry.h"
+
+#include "IDetailTreeNode.h"
+#include "IPropertyRowGenerator.h"
+#include "Modules/ModuleManager.h"
+#include "PropertyHandle.h"
+#include "UObject/StructOnScope.h"
+
+namespace WidgetRegistryUtils
+{
+	bool FindPropertyHandleRecursive(const TSharedPtr<IPropertyHandle>& PropertyHandle, const FString& PropertyNameOrPath, ERCFindNodeMethod FindMethod)
+	{
+		if (PropertyHandle && PropertyHandle->IsValidHandle())
+		{
+			uint32 ChildrenCount = 0;
+			PropertyHandle->GetNumChildren(ChildrenCount);
+			for (uint32 Index = 0; Index < ChildrenCount; ++Index)
+			{
+				TSharedPtr<IPropertyHandle> ChildHandle = PropertyHandle->GetChildHandle(Index);
+				if (FindPropertyHandleRecursive(ChildHandle, PropertyNameOrPath, FindMethod))
+				{
+					return true;
+				}
+			}
+
+			if (PropertyHandle->GetProperty())
+			{
+				if (FindMethod == ERCFindNodeMethod::Path)
+				{
+					if (PropertyHandle->GeneratePathToProperty() == PropertyNameOrPath)
+					{
+						return true;
+					}
+				}
+				else if (PropertyHandle->GetProperty()->GetName() == PropertyNameOrPath)
+				{
+					return true;
+				}
+			}
+		}
+
+		return false;
+	}
+
+	TSharedPtr<IDetailTreeNode> FindTreeNodeRecursive(const TSharedRef<IDetailTreeNode>& RootNode, const FString& PropertyNameOrPath, ERCFindNodeMethod FindMethod)
+	{
+		TArray<TSharedRef<IDetailTreeNode>> Children;
+		RootNode->GetChildren(Children);
+		for (TSharedRef<IDetailTreeNode>& Child : Children)
+		{
+			TSharedPtr<IDetailTreeNode> FoundNode = FindTreeNodeRecursive(Child, PropertyNameOrPath, FindMethod);
+			if (FoundNode.IsValid())
+			{
+				return FoundNode;
+			}
+		}
+
+		TSharedPtr<IPropertyHandle> Handle = RootNode->CreatePropertyHandle();
+		if (FindPropertyHandleRecursive(Handle, PropertyNameOrPath, FindMethod))
+		{
+			return RootNode;
+		}
+
+		return nullptr;
+	}
+
+	/** Find a node by its name in a detail tree node hierarchy. */
+	TSharedPtr<IDetailTreeNode> FindNode(const TArray<TSharedRef<IDetailTreeNode>>& RootNodes, const FString& QualifiedPropertyName, ERCFindNodeMethod FindMethod)
+	{
+		for (const TSharedRef<IDetailTreeNode>& CategoryNode : RootNodes)
+		{
+			TSharedPtr<IDetailTreeNode> FoundNode = FindTreeNodeRecursive(CategoryNode, QualifiedPropertyName, FindMethod);
+			if (FoundNode.IsValid())
+			{
+				return FoundNode;
+			}
+		}
+
+		return nullptr;
+	}
+}
+
+TSharedPtr<IDetailTreeNode> FRCPanelWidgetRegistry::GetObjectTreeNode(UObject* InObject, const FString& InField, ERCFindNodeMethod InFindMethod)
+{
+	const TPair<TWeakObjectPtr<UObject>, FString> CacheKey{InObject, InField};
+	if (TWeakPtr<IDetailTreeNode>* Node = TreeNodeCache.Find(CacheKey))
+	{
+		if (Node->IsValid())
+		{
+			return Node->Pin();
+		}
+		else
+		{
+			TreeNodeCache.Remove(CacheKey);
+		}
+	}
+	
+	TSharedPtr<IPropertyRowGenerator> Generator;
+	TWeakObjectPtr<UObject> WeakObject = InObject;
+	
+	if (TSharedPtr<IPropertyRowGenerator>* FoundGenerator = ObjectToRowGenerator.Find(WeakObject))
+	{
+		Generator = *FoundGenerator;
+	}
+	else
+	{
+		Generator = FModuleManager::GetModuleChecked<FPropertyEditorModule>("PropertyEditor").CreatePropertyRowGenerator(FPropertyRowGeneratorArgs()) ;
+		Generator->SetObjects({InObject});
+		ObjectToRowGenerator.Add(WeakObject, Generator);
+	}
+	
+	TSharedPtr<IDetailTreeNode> Node = WidgetRegistryUtils::FindNode(Generator->GetRootTreeNodes(), InField, ERCFindNodeMethod::Path);
+	// Cache the node to avoid having to do the recursive find again.
+	TreeNodeCache.Add(CacheKey, Node);
+	
+	return Node;
+}
+
+TSharedPtr<IDetailTreeNode> FRCPanelWidgetRegistry::GetStructTreeNode(const TSharedPtr<FStructOnScope>& InStruct, const FString& InField, ERCFindNodeMethod InFindMethod)
+{
+	TSharedPtr<IPropertyRowGenerator> Generator;
+	TWeakPtr<FStructOnScope> WeakStructure = InStruct;
+	
+	if (TSharedPtr<IPropertyRowGenerator>* FoundGenerator = StructToRowGenerator.Find(WeakStructure))
+	{
+		Generator = *FoundGenerator;
+	}
+	else
+	{
+		// In RC, struct details node are only used for function arguments, so default to showing hidden properties. 
+		FPropertyRowGeneratorArgs Args;
+		Args.bShouldShowHiddenProperties = true;
+		
+		Generator = FModuleManager::GetModuleChecked<FPropertyEditorModule>("PropertyEditor").CreatePropertyRowGenerator(Args);
+		Generator->SetStructure(InStruct);
+		StructToRowGenerator.Add(WeakStructure, Generator);
+	}
+
+	return WidgetRegistryUtils::FindNode(Generator->GetRootTreeNodes(), InField, InFindMethod);
+}
+
+void FRCPanelWidgetRegistry::Refresh(UObject* InObject)
+{
+	if (TSharedPtr<IPropertyRowGenerator>* Generator = ObjectToRowGenerator.Find({InObject}))
+	{
+		(*Generator)->SetObjects({InObject});
+	}
+}
+
+void FRCPanelWidgetRegistry::Clear()
+{
+	ObjectToRowGenerator.Empty();
+	StructToRowGenerator.Empty();
+	TreeNodeCache.Empty();
+}
diff --git a/Engine/Plugins/VirtualProduction/RemoteControl/Source/RemoteControlUI/Private/UI/RCPanelWidgetRegistry.h b/Engine/Plugins/VirtualProduction/RemoteControl/Source/RemoteControlUI/Private/UI/RCPanelWidgetRegistry.h
new file mode 100644
--- /dev/null
+++ b/Engine/Plugins/VirtualProduction/RemoteControl/Source/RemoteControlUI/Private/UI/RCPanelWidgetRegistry.h
@@ -0,0 +1,61 @@
+// Copyright Epic Games, Inc. All Rights Reserved.
+
+#pragma once
+
+#include "CoreMinimal.h"
+#include "UObject/WeakObjectPtr.h"
+#include "UObject/WeakObjectPtrTemplates.h"
+
+class IDetailTreeNode;
+class IPropertyRowGenerator;
+class FStructOnScope;
+
+enum class ERCFindNodeMethod
+{
+	Name, // Find a node by providing its property name.
+	Path  // Find a node by providing a property path.
+};
+
+class FRCPanelWidgetRegistry
+{
+public:
+	static FRCPanelWidgetRegistry& Get()
+	{
+		static FRCPanelWidgetRegistry Instance;
+		return Instance;
+	}
+
+	/**
+	 * Get a detail tree node for a given object and property.
+	 * @param InObject the object used to generate the detail row.
+	 * @param InField a string specifying which field to generate a row for.
+	 * @param InFindMethod whether the field parameter is a path or a field name.
+	 */
+	TSharedPtr<IDetailTreeNode> GetObjectTreeNode(UObject* InObject, const FString& InField, ERCFindNodeMethod InFindMethod);
+
+	/**
+	 * Get a detail tree node for a given struct and field.
+	 * @param InStruct the struct used to generate the detail row.
+	 * @param InField a string specifying which field to generate a row for.
+	 * @param InFindMethod whether the field parameter is a path or a field name.
+	 */
+	TSharedPtr<IDetailTreeNode> GetStructTreeNode(const TSharedPtr<FStructOnScope>& InStruct, const FString& InField, ERCFindNodeMethod InFindMethod);
+
+	/**
+	 * Update the generator for a given object.
+	 */
+	void Refresh(UObject* InObject);
+
+	/**
+	 * Clear the registry and its cache.
+	 */
+	void Clear();
+
+private:
+	/** Map of objects to row generator, used to have one row generator per object. */
+	TMap<TWeakObjectPtr<UObject>, TSharedPtr<IPropertyRowGenerator>> ObjectToRowGenerator;
+	/** Map of struct on scope to row generator, used to have one row generator per struct ptr. */
+	TMap<TWeakPtr<FStructOnScope>, TSharedPtr<IPropertyRowGenerator>> StructToRowGenerator;
+	/** Cache of Object&Field to tree nodes. */
+	TMap<TPair<TWeakObjectPtr<UObject>, FString>, TWeakPtr<IDetailTreeNode>> TreeNodeCache;
+};
diff --git a/Engine/Plugins/VirtualProduction/RemoteControl/Source/RemoteControlUI/Private/UI/SRCPanelDragHandle.h b/Engine/Plugins/VirtualProduction/RemoteControl/Source/RemoteControlUI/Private/UI/SRCPanelDragHandle.h
--- a/Engine/Plugins/VirtualProduction/RemoteControl/Source/RemoteControlUI/Private/UI/SRCPanelDragHandle.h
+++ b/Engine/Plugins/VirtualProduction/RemoteControl/Source/RemoteControlUI/Private/UI/SRCPanelDragHandle.h
@@ -8,7 +8,7 @@
 #include "Styling/CoreStyle.h"
 #include "Input/DragAndDrop.h"
 
-class SFieldGroup;
+class SRCPanelGroup;
 class FExposedEntityDragDrop;
 class FFieldGroupDragDropOp;
 
diff --git a/Engine/Plugins/VirtualProduction/RemoteControl/Source/RemoteControlUI/Private/UI/SRCPanelExposedEntitiesList.cpp b/Engine/Plugins/VirtualProduction/RemoteControl/Source/RemoteControlUI/Private/UI/SRCPanelExposedEntitiesList.cpp
--- a/Engine/Plugins/VirtualProduction/RemoteControl/Source/RemoteControlUI/Private/UI/SRCPanelExposedEntitiesList.cpp
+++ b/Engine/Plugins/VirtualProduction/RemoteControl/Source/RemoteControlUI/Private/UI/SRCPanelExposedEntitiesList.cpp
@@ -5,6 +5,7 @@
 #include "Algo/ForEach.h"
 #include "Editor.h"
 #include "ISettingsModule.h"
+#include "RCPanelWidgetRegistry.h"
 #include "GameFramework/Actor.h"
 #include "Input/DragAndDrop.h"
 #include "Engine/Selection.h"
@@ -32,6 +33,7 @@ void SRCPanelExposedEntitiesList::Construct(const FArguments& InArgs, URemoteCon
 	bIsInEditMode = InArgs._EditMode;
 	Preset = TStrongObjectPtr<URemoteControlPreset>(InPreset);
 	bDisplayValues = InArgs._DisplayValues;
+	OnEntityListUpdatedDelegate = InArgs._OnEntityListUpdated;
 
 	ColumnSizeData.LeftColumnWidth = TAttribute<float>(this, &SRCPanelExposedEntitiesList::OnGetLeftColumnWidth);
 	ColumnSizeData.RightColumnWidth = TAttribute<float>(this, &SRCPanelExposedEntitiesList::OnGetRightColumnWidth);
@@ -45,6 +47,7 @@ void SRCPanelExposedEntitiesList::Construct(const FArguments& InArgs, URemoteCon
 			.OnGenerateRow(this, &SRCPanelExposedEntitiesList::OnGenerateRow)
 			.OnGetChildren(this, &SRCPanelExposedEntitiesList::OnGetNodeChildren)
 			.OnSelectionChanged(this, &SRCPanelExposedEntitiesList::OnSelectionChanged)
+			.OnContextMenuOpening(this, &SRCPanelExposedEntitiesList::OnContextMenuOpening)
 			.ClearSelectionOnClick(false)
 	];
 
@@ -79,7 +82,7 @@ void SRCPanelExposedEntitiesList::SetSelection(const TSharedPtr<SRCPanelTreeNode
 			return;
 		}
 
-		if (TSharedPtr<SFieldGroup> SRCGroup = FindGroupById(Node->GetId()))
+		if (TSharedPtr<SRCPanelGroup> SRCGroup = FindGroupById(Node->GetId()))
 		{
 			TreeView->SetSelection(SRCGroup);
 		}	
@@ -96,6 +99,7 @@ void SRCPanelExposedEntitiesList::OnObjectPropertyChange(UObject* InObject, FPro
 
 	if ((InChangeEvent.ChangeType & TypesNeedingRefresh) != 0 && InChangeEvent.MemberProperty && IsRelevantProperty(InChangeEvent.MemberProperty->GetClass()))
 	{
+		FRCPanelWidgetRegistry::Get().Refresh(InObject);
 		for (const TPair <FGuid, TSharedPtr<SRCPanelTreeNode>>& Node : FieldWidgetMap)
 		{
 			Node.Value->Refresh();
@@ -110,7 +114,7 @@ void SRCPanelExposedEntitiesList::Refresh()
 	GenerateListWidgets();
 	RefreshGroups();
 
-	for (const TSharedPtr<SFieldGroup>& Group : FieldGroups)
+	for (const TSharedPtr<SRCPanelGroup>& Group : FieldGroups)
 	{
 		TreeView->SetItemExpansion(Group, true);
 	}
@@ -155,7 +159,7 @@ void SRCPanelExposedEntitiesList::RefreshGroups()
 
 	for (const FRemoteControlPresetGroup& RCGroup : Preset->Layout.GetGroups())
 	{
-		TSharedRef<SFieldGroup> FieldGroup = SNew(SFieldGroup, Preset.Get(), ColumnSizeData)
+		TSharedRef<SRCPanelGroup> FieldGroup = SNew(SRCPanelGroup, Preset.Get(), ColumnSizeData)
 			.Id(RCGroup.Id)
 			.Name(RCGroup.Name)
 			.OnFieldDropEvent_Raw(this, &SRCPanelExposedEntitiesList::OnDropOnGroup)
@@ -197,7 +201,7 @@ TSharedRef<ITableRow> SRCPanelExposedEntitiesList::OnGenerateRow(TSharedPtr<SRCP
 		{
 			if (TSharedPtr<FExposedEntityDragDrop> DragDropOp = Event.GetOperationAs<FExposedEntityDragDrop>())
 			{
-				if (TSharedPtr<SFieldGroup> Group = FindGroupById(GetGroupId(Field->GetId())))
+				if (TSharedPtr<SRCPanelGroup> Group = FindGroupById(GetGroupId(Field->GetId())))
 				{
 					if (DragDropOp->IsOfType<FExposedEntityDragDrop>())
 					{
@@ -282,7 +286,7 @@ void SRCPanelExposedEntitiesList::OnSelectionChanged(TSharedPtr<SRCPanelTreeNode
 	OnSelectionChangeDelegate.Broadcast(Node);
 }
 
-FReply SRCPanelExposedEntitiesList::OnDropOnGroup(const TSharedPtr<FDragDropOperation>& DragDropOperation, const TSharedPtr<SRCPanelTreeNode>& TargetEntity, const TSharedPtr<SFieldGroup>& DragTargetGroup)
+FReply SRCPanelExposedEntitiesList::OnDropOnGroup(const TSharedPtr<FDragDropOperation>& DragDropOperation, const TSharedPtr<SRCPanelTreeNode>& TargetEntity, const TSharedPtr<SRCPanelGroup>& DragTargetGroup)
 {
 	checkSlow(DragTargetGroup);
 
@@ -354,7 +358,7 @@ FGuid SRCPanelExposedEntitiesList::GetGroupId(const FGuid& EntityId)
 	return GroupId;
 }
 
-void SRCPanelExposedEntitiesList::OnDeleteGroup(const TSharedPtr<SFieldGroup>& PanelGroup)
+void SRCPanelExposedEntitiesList::OnDeleteGroup(const TSharedPtr<SRCPanelGroup>& PanelGroup)
 {
 	FScopedTransaction Transaction(LOCTEXT("DeleteGroup", "Delete Group"));
 	Preset->Modify();
@@ -398,16 +402,26 @@ void SRCPanelExposedEntitiesList::UnregisterEvents()
 	FCoreUObjectDelegates::OnObjectPropertyChanged.Remove(OnPropertyChangedHandle);
 }
 
-TSharedPtr<SFieldGroup> SRCPanelExposedEntitiesList::FindGroupById(const FGuid& Id)
+TSharedPtr<SRCPanelGroup> SRCPanelExposedEntitiesList::FindGroupById(const FGuid& Id)
 {
-	TSharedPtr<SFieldGroup> TargetGroup;
-	if (TSharedPtr<SFieldGroup>* FoundGroup = FieldGroups.FindByPredicate([Id](const TSharedPtr<SFieldGroup>& InGroup) {return InGroup->GetId() == Id; }))
+	TSharedPtr<SRCPanelGroup> TargetGroup;
+	if (TSharedPtr<SRCPanelGroup>* FoundGroup = FieldGroups.FindByPredicate([Id](const TSharedPtr<SRCPanelGroup>& InGroup) {return InGroup->GetId() == Id; }))
 	{
 		TargetGroup = *FoundGroup;
 	}
 	return TargetGroup;
 }
 
+TSharedPtr<SWidget> SRCPanelExposedEntitiesList::OnContextMenuOpening()
+{
+	if (TSharedPtr<SRCPanelTreeNode> SelectedNode = SRCPanelExposedEntitiesList::GetSelection())
+	{
+		return SelectedNode->GetContextMenu();
+	}
+	
+	return nullptr;
+}
+
 void SRCPanelExposedEntitiesList::RegisterPresetDelegates()
 {
 	FRemoteControlPresetLayout& Layout = Preset->Layout;
@@ -443,7 +457,7 @@ void SRCPanelExposedEntitiesList::OnEntityAdded(const FGuid& InEntityId)
 	{
 		FieldWidgetMap.Add(InEntityId, Node);
 
-		if (TSharedPtr<SFieldGroup> SRCGroup = FindGroupById(GetGroupId(InEntityId)))
+		if (TSharedPtr<SRCPanelGroup> SRCGroup = FindGroupById(GetGroupId(InEntityId)))
 		{
 			SRCGroup->GetNodes().Add(MoveTemp(Node));
 			TreeView->SetItemExpansion(SRCGroup, true);
@@ -467,7 +481,7 @@ void SRCPanelExposedEntitiesList::OnEntityAdded(const FGuid& InEntityId)
 
 void SRCPanelExposedEntitiesList::OnEntityRemoved(const FGuid& InGroupId, const FGuid& InEntityId)
 {
-	if (TSharedPtr<SFieldGroup> PanelGroup = FindGroupById(InGroupId))
+	if (TSharedPtr<SRCPanelGroup> PanelGroup = FindGroupById(InGroupId))
 	{
 		int32 EntityIndex = PanelGroup->GetNodes().IndexOfByPredicate([InEntityId](const TSharedPtr<SRCPanelTreeNode>& Node) {return Node->GetId() == InEntityId; });
 		if (EntityIndex != INDEX_NONE)
@@ -482,7 +496,7 @@ void SRCPanelExposedEntitiesList::OnEntityRemoved(const FGuid& InGroupId, const
 
 void SRCPanelExposedEntitiesList::OnGroupAdded(const FRemoteControlPresetGroup& Group)
 {
-	TSharedRef<SFieldGroup> FieldGroup = SNew(SFieldGroup, Preset.Get(), ColumnSizeData)
+	TSharedRef<SRCPanelGroup> FieldGroup = SNew(SRCPanelGroup, Preset.Get(), ColumnSizeData)
 		.Id(Group.Id)
 		.Name(Group.Name)
 		.OnFieldDropEvent_Raw(this, &SRCPanelExposedEntitiesList::OnDropOnGroup)
@@ -508,7 +522,7 @@ void SRCPanelExposedEntitiesList::OnGroupAdded(const FRemoteControlPresetGroup&
 
 void SRCPanelExposedEntitiesList::OnGroupDeleted(FRemoteControlPresetGroup DeletedGroup)
 {
-	int32 Index = FieldGroups.IndexOfByPredicate([&DeletedGroup](const TSharedPtr<SFieldGroup>& Group) { return Group->GetId() == DeletedGroup.Id; });
+	int32 Index = FieldGroups.IndexOfByPredicate([&DeletedGroup](const TSharedPtr<SRCPanelGroup>& Group) { return Group->GetId() == DeletedGroup.Id; });
 	if (Index != INDEX_NONE)
 	{
 		FieldGroups.RemoveAt(Index);
@@ -526,7 +540,7 @@ void SRCPanelExposedEntitiesList::OnGroupOrderChanged(const TArray<FGuid>& Group
 	}
 
 	auto SortFunc = [&IndicesMap]
-	(const TSharedPtr<SFieldGroup>& A, const TSharedPtr<SFieldGroup>& B)
+	(const TSharedPtr<SRCPanelGroup>& A, const TSharedPtr<SRCPanelGroup>& B)
 	{
 		return IndicesMap.FindChecked(A->GetId()) < IndicesMap.FindChecked(B->GetId());
 	};
@@ -537,7 +551,7 @@ void SRCPanelExposedEntitiesList::OnGroupOrderChanged(const TArray<FGuid>& Group
 
 void SRCPanelExposedEntitiesList::OnGroupRenamed(const FGuid& GroupId, FName NewName)
 {
-	if (TSharedPtr<SFieldGroup> Group = FindGroupById(GroupId))
+	if (TSharedPtr<SRCPanelGroup> Group = FindGroupById(GroupId))
 	{
 		Group->SetName(NewName);
 	}
@@ -555,7 +569,7 @@ void SRCPanelExposedEntitiesList::OnFieldDeleted(const FGuid& GroupId, const FGu
 
 void SRCPanelExposedEntitiesList::OnFieldOrderChanged(const FGuid& GroupId, const TArray<FGuid>& Fields)
 {
-	if (TSharedPtr<SFieldGroup>* Group = FieldGroups.FindByPredicate([GroupId](const TSharedPtr<SFieldGroup>& InGroup) {return InGroup->GetId() == GroupId; }))
+	if (TSharedPtr<SRCPanelGroup>* Group = FieldGroups.FindByPredicate([GroupId](const TSharedPtr<SRCPanelGroup>& InGroup) {return InGroup->GetId() == GroupId; }))
 	{
 		// Sort the group's fields according to the fields array.
 		TMap<FGuid, int32> OrderMap;
@@ -591,6 +605,7 @@ void SRCPanelExposedEntitiesList::OnEntitiesUpdated(URemoteControlPreset*, const
 	}
 
 	TreeView->RequestListRefresh();
+	OnEntityListUpdatedDelegate.ExecuteIfBound();
 }
 
 bool FGroupDragEvent::IsDraggedFromSameGroup() const
diff --git a/Engine/Plugins/VirtualProduction/RemoteControl/Source/RemoteControlUI/Private/UI/SRCPanelExposedEntitiesList.h b/Engine/Plugins/VirtualProduction/RemoteControl/Source/RemoteControlUI/Private/UI/SRCPanelExposedEntitiesList.h
--- a/Engine/Plugins/VirtualProduction/RemoteControl/Source/RemoteControlUI/Private/UI/SRCPanelExposedEntitiesList.h
+++ b/Engine/Plugins/VirtualProduction/RemoteControl/Source/RemoteControlUI/Private/UI/SRCPanelExposedEntitiesList.h
@@ -19,7 +19,7 @@ struct FRemoteControlProperty;
 struct FRemoteControlPresetGroup;
 struct FRemoteControlFunction;
 class ITableRow;
-class SFieldGroup;
+class SRCPanelGroup;
 struct SRCPanelTreeNode;
 class SRemoteControlTarget;
 struct SRCPanelExposedField;
@@ -29,7 +29,7 @@ class URemoteControlPreset;
 /** Holds information about a group drag and drop event  */
 struct FGroupDragEvent
 {
-	FGroupDragEvent(TSharedPtr<SFieldGroup> InDragOriginGroup, TSharedPtr<SFieldGroup> InDragTargetGroup)
+	FGroupDragEvent(TSharedPtr<SRCPanelGroup> InDragOriginGroup, TSharedPtr<SRCPanelGroup> InDragTargetGroup)
 		: DragOriginGroup(MoveTemp(InDragOriginGroup))
 		, DragTargetGroup(MoveTemp(InDragTargetGroup))
 	{
@@ -38,9 +38,9 @@ struct FGroupDragEvent
 	bool IsDraggedFromSameGroup() const;
 
 	/** Group the drag originated in. */
-	TSharedPtr<SFieldGroup> DragOriginGroup;
+	TSharedPtr<SRCPanelGroup> DragOriginGroup;
 	/** Group where the element was dropped. */
-	TSharedPtr<SFieldGroup> DragTargetGroup;
+	TSharedPtr<SRCPanelGroup> DragTargetGroup;
 };
 
 DECLARE_MULTICAST_DELEGATE_OneParam(FOnSelectionChange, const TSharedPtr<SRCPanelTreeNode>&/*SelectedNode*/);
@@ -54,6 +54,7 @@ public:
 	{}
 		SLATE_ATTRIBUTE(bool, EditMode)
 		SLATE_ARGUMENT(bool, DisplayValues)
+		SLATE_EVENT(FSimpleDelegate, OnEntityListUpdated)
 	SLATE_END_ARGS()
 
 	void Construct(const FArguments& InArgs, URemoteControlPreset* InPreset);
@@ -71,6 +72,9 @@ public:
 	/** Returns delegate called on selection change. */
 	FOnSelectionChange& OnSelectionChange() { return OnSelectionChangeDelegate; }
 
+	/** Returns delegate triggered upon a modification to an exposed entity. */
+	FSimpleDelegate OnEntityListUpdated() { return OnEntityListUpdatedDelegate; }
+	
 private:
 	/** Handles object property changes, used to update arrays correctly.  */
 	void OnObjectPropertyChange(UObject* InObject, FPropertyChangedEvent& InChangeEvent);
@@ -85,11 +89,11 @@ private:
 	/** Handle selection changes. */
 	void OnSelectionChanged(TSharedPtr<SRCPanelTreeNode> Node, ESelectInfo::Type SelectInfo);
 	/** Handlers for drag/drop events. */
-	FReply OnDropOnGroup(const TSharedPtr<FDragDropOperation>& DragDropOperation, const TSharedPtr<SRCPanelTreeNode>& TargetEntity, const TSharedPtr<SFieldGroup>& DragTargetGroup);
+	FReply OnDropOnGroup(const TSharedPtr<FDragDropOperation>& DragDropOperation, const TSharedPtr<SRCPanelTreeNode>& TargetEntity, const TSharedPtr<SRCPanelGroup>& DragTargetGroup);
 	/** Get the id of the group that holds a particular widget. */
 	FGuid GetGroupId(const FGuid& EntityId);
 	/** Handles group deletion. */
-	void OnDeleteGroup(const TSharedPtr<SFieldGroup>& PanelGroup);
+	void OnDeleteGroup(const TSharedPtr<SRCPanelGroup>& PanelGroup);
 	/** Select actors in the current level. */
 	void SelectActorsInlevel(const TArray<UObject*>& Objects);
 	//~ Register to engine/editor events in order to correctly update widgets.
@@ -102,7 +106,10 @@ private:
 	void OnSetColumnWidth(float InWidth) { ColumnWidth = InWidth; }
 
 	/** Find a group using its id. */
-	TSharedPtr<SFieldGroup> FindGroupById(const FGuid& Id);
+	TSharedPtr<SRCPanelGroup> FindGroupById(const FGuid& Id);
+
+	/** Handle context menu opening on a row. */
+	TSharedPtr<SWidget> OnContextMenuOpening();
 
 	//~ Register and handle preset delegates.
 	void RegisterPresetDelegates();
@@ -122,7 +129,7 @@ private:
 	/** Holds the fields list view. */
 	TSharedPtr<STreeView<TSharedPtr<SRCPanelTreeNode>>> TreeView;
 	/** Holds all the field groups. */
-	TArray<TSharedPtr<SFieldGroup>> FieldGroups;
+	TArray<TSharedPtr<SRCPanelGroup>> FieldGroups;
 	/** Map of field ids to field widgets. */
 	TMap<FGuid, TSharedPtr<SRCPanelTreeNode>> FieldWidgetMap;
 	/** Whether the panel is in edit mode. */
@@ -139,4 +146,6 @@ private:
 	FRCColumnSizeData ColumnSizeData;
 	/** The actual width of the right column.  The left column is 1-ColumnWidth */
 	float ColumnWidth = 0.65f;
+	/** Event triggered when the entity list is updated. */
+	FSimpleDelegate OnEntityListUpdatedDelegate;
 };
\ No newline at end of file
diff --git a/Engine/Plugins/VirtualProduction/RemoteControl/Source/RemoteControlUI/Private/UI/SRCPanelExposedEntity.cpp b/Engine/Plugins/VirtualProduction/RemoteControl/Source/RemoteControlUI/Private/UI/SRCPanelExposedEntity.cpp
--- a/Engine/Plugins/VirtualProduction/RemoteControl/Source/RemoteControlUI/Private/UI/SRCPanelExposedEntity.cpp
+++ b/Engine/Plugins/VirtualProduction/RemoteControl/Source/RemoteControlUI/Private/UI/SRCPanelExposedEntity.cpp
@@ -2,6 +2,8 @@
 
 #include "SRCPanelExposedEntity.h"
 
+#include "Framework/Application/SlateApplication.h"
+#include "Framework/MultiBox/MultiBoxBuilder.h"
 #include "ActorTreeItem.h"
 #include "GameFramework/Actor.h"
 #include "RemoteControlEntity.h"
@@ -15,16 +17,25 @@
 
 #define LOCTEXT_NAMESPACE "RemoteControlPanel"
 
-TSharedRef<SWidget> SRCPanelExposedEntity::CreateInvalidWidget()
+TSharedPtr<SWidget> SRCPanelExposedEntity::GetContextMenu()
 {
-	FSceneOutlinerModule& SceneOutlinerModule = FModuleManager::Get().LoadModuleChecked<FSceneOutlinerModule>("SceneOutliner");
-	FSceneOutlinerInitializationOptions Options;
-	Options.Filters = MakeShared<FSceneOutlinerFilters>();
-	Options.Filters->AddFilterPredicate<FActorTreeItem>(FActorTreeItem::FFilterPredicate::CreateRaw(this, &SRCPanelExposedEntity::IsActorSelectable));
+	FMenuBuilder MenuBuilder(true, TSharedPtr<const FUICommandList>());
+	MenuBuilder.AddSubMenu(
+		LOCTEXT("EntityRebindSubmenuLabel", "Rebind"),
+		LOCTEXT("EntityRebindSubmenuToolTip", "Pick an object to rebind this exposed entity."),
+		FNewMenuDelegate::CreateLambda([this](FMenuBuilder& SubMenuBuilder)
+		{
+			constexpr bool bNoIndent = true;
+			SubMenuBuilder.AddWidget(CreateRebindMenuContent(), FText::GetEmpty(), bNoIndent);
+		}));
+	return MenuBuilder.MakeWidget();
+}
 
+TSharedRef<SWidget> SRCPanelExposedEntity::CreateInvalidWidget()
+{
 	return SNew(SHorizontalBox)
 		+ SHorizontalBox::Slot()
-	    .AutoWidth()
+		.AutoWidth()
 		[
 			SNew(SComboButton)
 			.ButtonContent()
@@ -38,18 +49,31 @@ TSharedRef<SWidget> SRCPanelExposedEntity::CreateInvalidWidget()
 				.MaxDesiredHeight(400.0f)
 				.WidthOverride(300.0f)
 				[
-					SceneOutlinerModule.CreateActorPicker(Options, FOnActorPicked::CreateRaw(this, &SRCPanelExposedEntity::OnActorSelected), nullptr)
+					CreateRebindMenuContent()
 				]
 			]
 		];
 }
 
+TSharedRef<SWidget> SRCPanelExposedEntity::CreateRebindMenuContent()
+{
+	FSceneOutlinerModule& SceneOutlinerModule = FModuleManager::Get().LoadModuleChecked<FSceneOutlinerModule>("SceneOutliner");
+	return SNew(SBox)
+	.MaxDesiredHeight(400.0f)
+	.WidthOverride(300.0f)
+	[
+		SceneOutlinerModule.CreateActorPicker(Options, FOnActorPicked::CreateRaw(this, &SRCPanelExposedEntity::OnActorSelected), nullptr)
+	];
+}
+
 void SRCPanelExposedEntity::OnActorSelected(AActor* InActor) const
 {
 	if (TSharedPtr<FRemoteControlEntity> Entity = GetEntity())
 	{
 		Entity->BindObject(InActor);
 	}
+
+	FSlateApplication::Get().DismissAllMenus();
 }
 
 bool SRCPanelExposedEntity::IsActorSelectable(const AActor* Actor) const
diff --git a/Engine/Plugins/VirtualProduction/RemoteControl/Source/RemoteControlUI/Private/UI/SRCPanelExposedEntity.h b/Engine/Plugins/VirtualProduction/RemoteControl/Source/RemoteControlUI/Private/UI/SRCPanelExposedEntity.h
--- a/Engine/Plugins/VirtualProduction/RemoteControl/Source/RemoteControlUI/Private/UI/SRCPanelExposedEntity.h
+++ b/Engine/Plugins/VirtualProduction/RemoteControl/Source/RemoteControlUI/Private/UI/SRCPanelExposedEntity.h
@@ -15,15 +15,22 @@ struct SRCPanelExposedEntity : public SRCPanelTreeNode
 	 */
 	virtual TSharedPtr<FRemoteControlEntity> GetEntity() const = 0;
 
-protected:
 	/**
-	 * Create a widget that displays the rebind button.
+	 * Get the context menu for this item.
 	 */
+	virtual TSharedPtr<SWidget> GetContextMenu() override;
+
+protected:
+	/** Create a widget that displays the rebind button. */
 	TSharedRef<SWidget> CreateInvalidWidget();
 
 private:
 	/** Handles changing the object this entity is bound to upon selecting an actor in the rebinding dropdown. */
 	void OnActorSelected(AActor* InActor) const;
+	
+	/** Create the content of the rebind button.  */
+	TSharedRef<SWidget> CreateRebindMenuContent();
+
 	/** Returns whether or not the actor is selectable for a binding replacement. */
 	bool IsActorSelectable(const AActor* Parent) const;
 };
\ No newline at end of file
diff --git a/Engine/Plugins/VirtualProduction/RemoteControl/Source/RemoteControlUI/Private/UI/SRCPanelExposedField.cpp b/Engine/Plugins/VirtualProduction/RemoteControl/Source/RemoteControlUI/Private/UI/SRCPanelExposedField.cpp
--- a/Engine/Plugins/VirtualProduction/RemoteControl/Source/RemoteControlUI/Private/UI/SRCPanelExposedField.cpp
+++ b/Engine/Plugins/VirtualProduction/RemoteControl/Source/RemoteControlUI/Private/UI/SRCPanelExposedField.cpp
@@ -6,8 +6,8 @@
 #include "EditorFontGlyphs.h"
 #include "EditorStyleSet.h"
 #include "Framework/SlateDelegates.h"
+#include "RCPanelWidgetRegistry.h"
 #include "IDetailTreeNode.h"
-#include "IPropertyRowGenerator.h"
 #include "Layout/Visibility.h"
 #include "Modules/ModuleManager.h"
 #include "PropertyHandle.h"
@@ -29,77 +29,6 @@
 
 namespace ExposedFieldUtils
 {
-	bool FindPropertyHandleRecursive(const TSharedPtr<IPropertyHandle>& PropertyHandle, const FString& QualifiedPropertyName, bool bRequiresMatchingPath)
-	{
-		if (PropertyHandle && PropertyHandle->IsValidHandle())
-		{
-			uint32 ChildrenCount = 0;
-			PropertyHandle->GetNumChildren(ChildrenCount);
-			for (uint32 Index = 0; Index < ChildrenCount; ++Index)
-			{
-				TSharedPtr<IPropertyHandle> ChildHandle = PropertyHandle->GetChildHandle(Index);
-				if (FindPropertyHandleRecursive(ChildHandle, QualifiedPropertyName, bRequiresMatchingPath))
-				{
-					return true;
-				}
-			}
-
-			if (PropertyHandle->GetProperty())
-			{
-				if (bRequiresMatchingPath)
-				{
-					if (PropertyHandle->GeneratePathToProperty() == QualifiedPropertyName)
-					{
-						return true;
-					}
-				}
-				else if (PropertyHandle->GetProperty()->GetName() == QualifiedPropertyName)
-				{
-					return true;
-				}
-			}
-		}
-
-		return false;
-	}
-
-	TSharedPtr<IDetailTreeNode> FindTreeNodeRecursive(const TSharedRef<IDetailTreeNode>& RootNode, const FString& QualifiedPropertyName, bool bRequiresMatchingPath)
-	{
-		TArray<TSharedRef<IDetailTreeNode>> Children;
-		RootNode->GetChildren(Children);
-		for (TSharedRef<IDetailTreeNode>& Child : Children)
-		{
-			TSharedPtr<IDetailTreeNode> FoundNode = FindTreeNodeRecursive(Child, QualifiedPropertyName, bRequiresMatchingPath);
-			if (FoundNode.IsValid())
-			{
-				return FoundNode;
-			}
-		}
-
-		TSharedPtr<IPropertyHandle> Handle = RootNode->CreatePropertyHandle();
-		if (FindPropertyHandleRecursive(Handle, QualifiedPropertyName, bRequiresMatchingPath))
-		{
-			return RootNode;
-		}
-
-		return nullptr;
-	}
-
-	/** Find a node by its name in a detail tree node hierarchy. */
-	TSharedPtr<IDetailTreeNode> FindNode(const TArray<TSharedRef<IDetailTreeNode>>& RootNodes, const FString& QualifiedPropertyName, bool bRequiresMatchingPath)
-	{
-		for (const TSharedRef<IDetailTreeNode>& CategoryNode : RootNodes)
-		{
-			TSharedPtr<IDetailTreeNode> FoundNode = FindTreeNodeRecursive(CategoryNode, QualifiedPropertyName, bRequiresMatchingPath);
-			if (FoundNode.IsValid())
-			{
-				return FoundNode;
-			}
-		}
-
-		return nullptr;
-	}
-
 	TSharedRef<SWidget> CreateNodeValueWidget(const TSharedPtr<IDetailTreeNode>& Node)
 	{
 		FNodeWidgets NodeWidgets = Node->CreateNodeWidgets();
@@ -217,7 +146,6 @@ void SRCPanelExposedField::Refresh()
 	if (TSharedPtr<FRemoteControlField> Field = WeakField.Pin())
 	{
 		CachedLabel = Field->GetLabel();
-		RowGenerator->SetObjects(Field->GetBoundObjects());
 
 		// Don't update the function row needlessly.
 		if (Field->FieldType == EExposedFieldType::Property)
@@ -241,9 +169,11 @@ TSharedRef<SWidget> SRCPanelExposedField::ConstructWidget()
 	{
 		if (TSharedPtr<FRemoteControlField> Field = WeakField.Pin())
 		{
-			if (GetFieldType() == EExposedFieldType::Property && RowGenerator->GetSelectedObjects().Num())
+			// For the moment, just use the first object.
+			TArray<UObject*> Objects = Field->GetBoundObjects();
+			if (GetFieldType() == EExposedFieldType::Property && Objects.Num() > 0)
 			{
-				if (TSharedPtr<IDetailTreeNode> Node = ExposedFieldUtils::FindNode(RowGenerator->GetRootTreeNodes(), Field->FieldPathInfo.ToPathPropertyString(), true))
+				if (TSharedPtr<IDetailTreeNode> Node = FRCPanelWidgetRegistry::Get().GetObjectTreeNode(Objects[0], Field->FieldPathInfo.ToPathPropertyString(), ERCFindNodeMethod::Path))
 				{
 					TArray<TSharedRef<IDetailTreeNode>> ChildNodes;
 					Node->GetChildren(ChildNodes);
@@ -317,6 +247,7 @@ TSharedRef<SWidget> SRCPanelExposedField::MakeFieldWidget(const TSharedRef<SWidg
 		];
 
 	Args.NameWidget = SNew(SHorizontalBox)
+		.Clipping(EWidgetClipping::OnDemand)
 		+ SHorizontalBox::Slot()
 		.VAlign(VAlign_Center)
 		.Padding(0.f, 0.f, 2.0f, 0.f)
@@ -330,6 +261,7 @@ TSharedRef<SWidget> SRCPanelExposedField::MakeFieldWidget(const TSharedRef<SWidg
             .Text(FEditorFontGlyphs::Exclamation_Triangle)
 		]
 		+ SHorizontalBox::Slot()
+		.AutoWidth()
 		[
 			SAssignNew(NameTextBox, SInlineEditableTextBlock)
 			.Text(FText::FromName(CachedLabel))
@@ -430,25 +362,12 @@ void SRCPanelExposedField::OnLabelCommitted(const FText& InLabel, ETextCommit::T
 
 void SRCPanelExposedField::ConstructPropertyWidget()
 {
-	RowGenerator = FModuleManager::GetModuleChecked<FPropertyEditorModule>("PropertyEditor").CreatePropertyRowGenerator(FPropertyRowGeneratorArgs());
-
-	if (URemoteControlPreset* RCPreset = Preset.Get())
-	{
-		if (TSharedPtr<FRemoteControlProperty> Property = RCPreset->GetExposedEntity<FRemoteControlProperty>(FieldId).Pin())
-		{
-			RowGenerator->SetObjects(Property->GetBoundObjects());
-		}
-	}
-
 	ChildSlot.AttachWidget(ConstructWidget());
 }
 
 void SRCPanelExposedField::ConstructFunctionWidget()
 {
 	TSharedPtr<SRCPanelExposedField> ExposedFieldWidget;
-	FPropertyRowGeneratorArgs GeneratorArgs;
-	GeneratorArgs.bShouldShowHiddenProperties = true;
-	RowGenerator = FModuleManager::GetModuleChecked<FPropertyEditorModule>("PropertyEditor").CreatePropertyRowGenerator(GeneratorArgs);
 
 	URemoteControlPreset* RCPreset = Preset.Get();
 	if (!RCPreset)
@@ -460,8 +379,6 @@ void SRCPanelExposedField::ConstructFunctionWidget()
 	{
 		if (RCFunction->GetFunction() && RCFunction->GetBoundObjects().Num())
 		{
-			RowGenerator->SetStructure(RCFunction->FunctionArguments);
-			
 			if (bDisplayValues)
 			{
 				TArray<TSharedPtr<SRCPanelFieldChildNode>> ChildNodes;
@@ -477,7 +394,7 @@ void SRCPanelExposedField::ConstructFunctionWidget()
 						continue;
 					}
 
-					if (TSharedPtr<IDetailTreeNode> PropertyNode = ExposedFieldUtils::FindNode(RowGenerator->GetRootTreeNodes(), It->GetFName().ToString(), false))
+					if (TSharedPtr<IDetailTreeNode> PropertyNode = FRCPanelWidgetRegistry::Get().GetStructTreeNode(RCFunction->FunctionArguments, It->GetFName().ToString(), ERCFindNodeMethod::Name))
 					{
 						ChildNodes.Add(SNew(SRCPanelFieldChildNode, PropertyNode.ToSharedRef(), ColumnSizeData));
 					}
diff --git a/Engine/Plugins/VirtualProduction/RemoteControl/Source/RemoteControlUI/Private/UI/SRCPanelExposedField.h b/Engine/Plugins/VirtualProduction/RemoteControl/Source/RemoteControlUI/Private/UI/SRCPanelExposedField.h
--- a/Engine/Plugins/VirtualProduction/RemoteControl/Source/RemoteControlUI/Private/UI/SRCPanelExposedField.h
+++ b/Engine/Plugins/VirtualProduction/RemoteControl/Source/RemoteControlUI/Private/UI/SRCPanelExposedField.h
@@ -17,7 +17,6 @@ enum class EExposedFieldType : uint8;
 struct FSlateBrush;
 struct FRemoteControlField;
 struct FGuid;
-class IPropertyRowGenerator;
 class IDetailTreeNode;
 class SInlineEditableTextBlock;
 struct SRCPanelFieldChildNode;
@@ -103,6 +102,7 @@ private:
 	/** Handles calling an exposed function.*/
 	FReply OnClickFunctionButton();
 private:
+	/** Weak pointer to the underlying RC Field. */
 	TWeakPtr<FRemoteControlField> WeakField;
 	/** Display name of the field. */
 	FName CachedLabel;
@@ -118,8 +118,6 @@ private:
 	TSharedPtr<SWidget> OptionsWidget;
 	/** This exposed field's child widgets (ie. An array's rows) */
 	TArray<TSharedPtr<SRCPanelFieldChildNode>> ChildWidgets;
-	/** Holds the generator that creates the widgets. */
-	TSharedPtr<IPropertyRowGenerator> RowGenerator;
 	/** Whether the panel is in edit mode or not. */
 	TAttribute<bool> bEditMode;
 	/** The underlying preset. */
diff --git a/Engine/Plugins/VirtualProduction/RemoteControl/Source/RemoteControlUI/Private/UI/SRCPanelFieldGroup.cpp b/Engine/Plugins/VirtualProduction/RemoteControl/Source/RemoteControlUI/Private/UI/SRCPanelFieldGroup.cpp
--- a/Engine/Plugins/VirtualProduction/RemoteControl/Source/RemoteControlUI/Private/UI/SRCPanelFieldGroup.cpp
+++ b/Engine/Plugins/VirtualProduction/RemoteControl/Source/RemoteControlUI/Private/UI/SRCPanelFieldGroup.cpp
@@ -21,7 +21,7 @@
 
 #define LOCTEXT_NAMESPACE "RemoteControlPanel"
 
-void SFieldGroup::Tick(const FGeometry&, const double, const float)
+void SRCPanelGroup::Tick(const FGeometry&, const double, const float)
 {
 	if (bNeedsRename)
 	{
@@ -33,7 +33,7 @@ void SFieldGroup::Tick(const FGeometry&, const double, const float)
 	}
 }
 
-void SFieldGroup::Construct(const FArguments& InArgs, URemoteControlPreset* InPreset, FRCColumnSizeData InColumnSizeData)
+void SRCPanelGroup::Construct(const FArguments& InArgs, URemoteControlPreset* InPreset, FRCColumnSizeData InColumnSizeData)
 {
 	Preset = InPreset;
 	
@@ -58,7 +58,7 @@ void SFieldGroup::Construct(const FArguments& InArgs, URemoteControlPreset* InPr
 		[
 			SNew(SBox)
 			.Padding(FMargin(0.0f, 2.0f) )
-			.Visibility(this, &SFieldGroup::GetVisibilityAccordingToEditMode, EVisibility::Collapsed)
+			.Visibility(this, &SRCPanelGroup::GetVisibilityAccordingToEditMode, EVisibility::Collapsed)
 			[
 				SNew(SRCPanelDragHandle<FFieldGroupDragDropOp>, Id)
 				.Widget(SharedThis(this))
@@ -72,11 +72,11 @@ void SFieldGroup::Construct(const FArguments& InArgs, URemoteControlPreset* InPr
 		.AutoWidth()
 		[
 			SAssignNew(NameTextBox, SInlineEditableTextBlock)
-			.ColorAndOpacity(this, &SFieldGroup::GetGroupNameTextColor)
+			.ColorAndOpacity(this, &SRCPanelGroup::GetGroupNameTextColor)
 			.Font(FEditorStyle::GetFontStyle("DetailsView.CategoryFontStyle"))
 			.Text(FText::FromName(Name))
-			.OnTextCommitted(this, &SFieldGroup::OnLabelCommitted)
-			.OnVerifyTextChanged(this, &SFieldGroup::OnVerifyItemLabelChanged)
+			.OnTextCommitted(this, &SRCPanelGroup::OnLabelCommitted)
+			.OnVerifyTextChanged(this, &SRCPanelGroup::OnVerifyItemLabelChanged)
 			.IsReadOnly_Lambda([this]() { return !bEditMode.Get(); })
 		]
 		+ SHorizontalBox::Slot()
@@ -86,7 +86,7 @@ void SFieldGroup::Construct(const FArguments& InArgs, URemoteControlPreset* InPr
 		.HAlign(HAlign_Left)
 		[
 			SNew(SButton)
-			.Visibility(this, &SFieldGroup::GetVisibilityAccordingToEditMode, EVisibility::Collapsed)
+			.Visibility(this, &SRCPanelGroup::GetVisibilityAccordingToEditMode, EVisibility::Collapsed)
 			.ButtonStyle(FEditorStyle::Get(), "FlatButton")
 			.OnClicked_Lambda([this] () 
 				{
@@ -110,8 +110,8 @@ void SFieldGroup::Construct(const FArguments& InArgs, URemoteControlPreset* InPr
 		.FillWidth(1.f)
 		[
 			SNew(SButton)
-			.Visibility_Raw(this, &SFieldGroup::GetVisibilityAccordingToEditMode, EVisibility::Hidden)
-			.OnClicked(this, &SFieldGroup::HandleDeleteGroup)
+			.Visibility_Raw(this, &SRCPanelGroup::GetVisibilityAccordingToEditMode, EVisibility::Hidden)
+			.OnClicked(this, &SRCPanelGroup::HandleDeleteGroup)
 			.ButtonStyle(FRemoteControlPanelStyle::Get(), "RemoteControlPanel.UnexposeButton")
 			[
 				SNew(STextBlock)
@@ -125,15 +125,15 @@ void SFieldGroup::Construct(const FArguments& InArgs, URemoteControlPreset* InPr
 	[
 		SNew(SBorder)
 		.Padding(0.f)
-		.BorderImage(this, &SFieldGroup::GetBorderImage)
+		.BorderImage(this, &SRCPanelGroup::GetBorderImage)
 		.VAlign(VAlign_Fill)
 		[
 			SNew(SDropTarget)
 			.VerticalImage(FRemoteControlPanelStyle::Get()->GetBrush("RemoteControlPanel.VerticalDash"))
 			.HorizontalImage(FRemoteControlPanelStyle::Get()->GetBrush("RemoteControlPanel.HorizontalDash"))
 			.OnDrop_Lambda([this](TSharedPtr<FDragDropOperation> DragDropOperation) { return OnFieldDropGroup(DragDropOperation, nullptr); })
-			.OnAllowDrop(this, &SFieldGroup::OnAllowDropFromOtherGroup)
-			.OnIsRecognized(this, &SFieldGroup::OnAllowDropFromOtherGroup)
+			.OnAllowDrop(this, &SRCPanelGroup::OnAllowDropFromOtherGroup)
+			.OnIsRecognized(this, &SRCPanelGroup::OnAllowDropFromOtherGroup)
 			[
 				MakeSplitRow(LeftColumn, RightColumn)
 			]
@@ -141,38 +141,38 @@ void SFieldGroup::Construct(const FArguments& InArgs, URemoteControlPreset* InPr
 	];
 }
 
-FName SFieldGroup::GetGroupName() const
+FName SRCPanelGroup::GetGroupName() const
 {
 	return Name;
 }
 
-void SFieldGroup::SetName(FName InName)
+void SRCPanelGroup::SetName(FName InName)
 {
 	Name = InName;
 	NameTextBox->SetText(FText::FromName(Name));
 }
 
-void SFieldGroup::GetNodeChildren(TArray<TSharedPtr<SRCPanelTreeNode>>& OutChildren) const
+void SRCPanelGroup::GetNodeChildren(TArray<TSharedPtr<SRCPanelTreeNode>>& OutChildren) const
 {
 	OutChildren.Append(Nodes);
 }
 
-FGuid SFieldGroup::GetId() const
+FGuid SRCPanelGroup::GetId() const
 {
 	return Id;
 }
 
-SRCPanelTreeNode::ENodeType SFieldGroup::GetType() const
+SRCPanelTreeNode::ENodeType SRCPanelGroup::GetType() const
 {
 	return SRCPanelTreeNode::Group;
 }
 
-TSharedPtr<SFieldGroup> SFieldGroup::AsGroup()
+TSharedPtr<SRCPanelGroup> SRCPanelGroup::AsGroup()
 {
-	return StaticCastSharedRef<SFieldGroup>(AsShared());
+	return StaticCastSharedRef<SRCPanelGroup>(AsShared());
 }
 
-FReply SFieldGroup::OnFieldDropGroup(const FDragDropEvent& Event, TSharedPtr<SRCPanelTreeNode> TargetField)
+FReply SRCPanelGroup::OnFieldDropGroup(const FDragDropEvent& Event, TSharedPtr<SRCPanelTreeNode> TargetField)
 {
 	if (TSharedPtr<FExposedEntityDragDrop> DragDropOp = Event.GetOperationAs<FExposedEntityDragDrop>())
 	{
@@ -181,7 +181,7 @@ FReply SFieldGroup::OnFieldDropGroup(const FDragDropEvent& Event, TSharedPtr<SRC
 	return FReply::Unhandled();
 }
 
-FReply SFieldGroup::OnFieldDropGroup(TSharedPtr<FDragDropOperation> DragDropOperation, TSharedPtr<SRCPanelTreeNode> TargetField)
+FReply SRCPanelGroup::OnFieldDropGroup(TSharedPtr<FDragDropOperation> DragDropOperation, TSharedPtr<SRCPanelTreeNode> TargetField)
 {
 	if (DragDropOperation)
 	{
@@ -198,7 +198,7 @@ FReply SFieldGroup::OnFieldDropGroup(TSharedPtr<FDragDropOperation> DragDropOper
 	return FReply::Unhandled();
 }
 
-bool SFieldGroup::OnAllowDropFromOtherGroup(TSharedPtr<FDragDropOperation> DragDropOperation)
+bool SRCPanelGroup::OnAllowDropFromOtherGroup(TSharedPtr<FDragDropOperation> DragDropOperation)
 {
 	if (DragDropOperation->IsOfType<FExposedEntityDragDrop>())
 	{
@@ -225,28 +225,28 @@ bool SFieldGroup::OnAllowDropFromOtherGroup(TSharedPtr<FDragDropOperation> DragD
 	return false;
 }
 
-FReply SFieldGroup::HandleDeleteGroup()
+FReply SRCPanelGroup::HandleDeleteGroup()
 {
 	OnDeleteGroup.ExecuteIfBound(AsGroup());
 	return FReply::Handled();
 }
 
-FSlateColor SFieldGroup::GetGroupNameTextColor() const
+FSlateColor SRCPanelGroup::GetGroupNameTextColor() const
 {
 	return FLinearColor(1, 1, 1, 0.7);
 }
 
-const FSlateBrush* SFieldGroup::GetBorderImage() const
+const FSlateBrush* SRCPanelGroup::GetBorderImage() const
 {
 	return FRemoteControlPanelStyle::Get()->GetBrush("RemoteControlPanel.GroupBorder");
 }
 
-EVisibility SFieldGroup::GetVisibilityAccordingToEditMode(EVisibility NonEditModeVisibility) const
+EVisibility SRCPanelGroup::GetVisibilityAccordingToEditMode(EVisibility NonEditModeVisibility) const
 {
 	return bEditMode.Get() ? EVisibility::Visible : NonEditModeVisibility;
 }
 
-bool SFieldGroup::OnVerifyItemLabelChanged(const FText& InLabel, FText& OutErrorMessage)
+bool SRCPanelGroup::OnVerifyItemLabelChanged(const FText& InLabel, FText& OutErrorMessage)
 {
 	check(Preset.IsValid());
 	FName TentativeName = FName(*InLabel.ToString());
@@ -259,7 +259,7 @@ bool SFieldGroup::OnVerifyItemLabelChanged(const FText& InLabel, FText& OutError
 	return true;
 }
 
-void SFieldGroup::OnLabelCommitted(const FText& InLabel, ETextCommit::Type InCommitInfo)
+void SRCPanelGroup::OnLabelCommitted(const FText& InLabel, ETextCommit::Type InCommitInfo)
 {
 	check(Preset.IsValid());
 	FScopedTransaction Transaction(LOCTEXT("RenameGroup", "Rename Group"));
diff --git a/Engine/Plugins/VirtualProduction/RemoteControl/Source/RemoteControlUI/Private/UI/SRCPanelFieldGroup.h b/Engine/Plugins/VirtualProduction/RemoteControl/Source/RemoteControlUI/Private/UI/SRCPanelFieldGroup.h
--- a/Engine/Plugins/VirtualProduction/RemoteControl/Source/RemoteControlUI/Private/UI/SRCPanelFieldGroup.h
+++ b/Engine/Plugins/VirtualProduction/RemoteControl/Source/RemoteControlUI/Private/UI/SRCPanelFieldGroup.h
@@ -20,18 +20,18 @@ class SInlineEditableTextBlock;
 class URemoteControlPreset;
 
 /** Widget representing a group. */
-class SFieldGroup : public SRCPanelTreeNode, public SCompoundWidget
+class SRCPanelGroup : public SRCPanelTreeNode, public SCompoundWidget
 {
 public:
 	using SWidget::SharedThis;
 	using SWidget::AsShared;
 
 public:
-	DECLARE_DELEGATE_RetVal_ThreeParams(FReply, FOnFieldDropEvent, const TSharedPtr<FDragDropOperation>& /* Event */, const TSharedPtr<SRCPanelTreeNode>& /* TargetField */, const TSharedPtr<SFieldGroup>& /* DragTargetGroup */);
+	DECLARE_DELEGATE_RetVal_ThreeParams(FReply, FOnFieldDropEvent, const TSharedPtr<FDragDropOperation>& /* Event */, const TSharedPtr<SRCPanelTreeNode>& /* TargetField */, const TSharedPtr<SRCPanelGroup>& /* DragTargetGroup */);
 	DECLARE_DELEGATE_RetVal_OneParam(FGuid, FOnGetGroupId, const FGuid& /* EntityId */);
-	DECLARE_DELEGATE_OneParam(FOnDeleteGroup, const TSharedPtr<SFieldGroup>&);
+	DECLARE_DELEGATE_OneParam(FOnDeleteGroup, const TSharedPtr<SRCPanelGroup>&);
 
-	SLATE_BEGIN_ARGS(SFieldGroup)
+	SLATE_BEGIN_ARGS(SRCPanelGroup)
 		: _EditMode(true)
 	{}
 		SLATE_ARGUMENT(FGuid, Id)
@@ -60,7 +60,7 @@ public:
 	virtual void GetNodeChildren(TArray<TSharedPtr<SRCPanelTreeNode>>& OutChildren) const override;
 	virtual FGuid GetId() const override;
 	virtual ENodeType GetType() const override;
-	virtual TSharedPtr<SFieldGroup> AsGroup() override;
+	virtual TSharedPtr<SRCPanelGroup> AsGroup() override;
 
 private:
 	//~ Handle drag/drop events
@@ -110,9 +110,9 @@ class FFieldGroupDragDropOp : public FDecoratedDragDropOp
 public:
 	DRAG_DROP_OPERATOR_TYPE(FFieldGroupDragDropOp, FDragDropOperation)
 
-	using WidgetType = SFieldGroup;
+	using WidgetType = SRCPanelGroup;
 
-	FFieldGroupDragDropOp(TSharedPtr<SFieldGroup> InWidget, FGuid InId)
+	FFieldGroupDragDropOp(TSharedPtr<SRCPanelGroup> InWidget, FGuid InId)
 		: Id(MoveTemp(InId))
 	{
 		DecoratorWidget = SNew(SBorder)
diff --git a/Engine/Plugins/VirtualProduction/RemoteControl/Source/RemoteControlUI/Private/UI/SRCPanelTreeNode.cpp b/Engine/Plugins/VirtualProduction/RemoteControl/Source/RemoteControlUI/Private/UI/SRCPanelTreeNode.cpp
--- a/Engine/Plugins/VirtualProduction/RemoteControl/Source/RemoteControlUI/Private/UI/SRCPanelTreeNode.cpp
+++ b/Engine/Plugins/VirtualProduction/RemoteControl/Source/RemoteControlUI/Private/UI/SRCPanelTreeNode.cpp
@@ -2,7 +2,9 @@
 
 #include "SRCPanelTreeNode.h"
 
+#include "Widgets/Layout/SBox.h"
 #include "Widgets/SBoxPanel.h"
+#include "Widgets/SOverlay.h"
 #include "Widgets/Text/STextBlock.h"
 
 #define LOCTEXT_NAMESPACE "RemoteControlPanelNode"
@@ -38,6 +40,7 @@ TSharedRef<SWidget> SRCPanelTreeNode::MakeNodeWidget(const FMakeNodeWidgetArgs&
 
 	TSharedRef<SWidget> LeftColumn = 
 		SNew(SHorizontalBox)
+		.Clipping(EWidgetClipping::OnDemand)
 		// Drag and drop handle
 		+ SHorizontalBox::Slot()
 		.HAlign(HAlign_Left)
@@ -62,23 +65,29 @@ TSharedRef<SWidget> SRCPanelTreeNode::MakeNodeWidget(const FMakeNodeWidgetArgs&
 		];
 
 	TSharedRef<SWidget> RightColumn =
-		SNew(SHorizontalBox)
-		// Field value
-		+ SHorizontalBox::Slot()
-		.VAlign(VAlign_Center)
-		.HAlign(HAlign_Left)
-		.FillWidth(1.0f)
+		SNew(SOverlay)
+		+ SOverlay::Slot()
+		.HAlign(HAlign_Fill)
 		[
-			WidgetOrNull(Args.ValueWidget)
+			SNew(SBox)
+			.HAlign(HAlign_Left)
+			[
+				WidgetOrNull(Args.ValueWidget)			
+			]
 		]
-		// Unexpose button
-		+ SHorizontalBox::Slot()
+
+		+ SOverlay::Slot()
+		.HAlign(HAlign_Right)
 		.VAlign(VAlign_Center)
-		.AutoWidth()
 		[
-			WidgetOrNull(Args.UnexposeButton)
+			SNew(SBorder)
+			.BorderImage(FEditorStyle::GetBrush("ToolPanel.GroupBorder"))
+			.Padding(0.f)
+			[
+				WidgetOrNull(Args.UnexposeButton)
+			]
 		];
-	
+		
 	return MakeSplitRow(LeftColumn, RightColumn);
 }
 
diff --git a/Engine/Plugins/VirtualProduction/RemoteControl/Source/RemoteControlUI/Private/UI/SRCPanelTreeNode.h b/Engine/Plugins/VirtualProduction/RemoteControl/Source/RemoteControlUI/Private/UI/SRCPanelTreeNode.h
--- a/Engine/Plugins/VirtualProduction/RemoteControl/Source/RemoteControlUI/Private/UI/SRCPanelTreeNode.h
+++ b/Engine/Plugins/VirtualProduction/RemoteControl/Source/RemoteControlUI/Private/UI/SRCPanelTreeNode.h
@@ -9,7 +9,7 @@
 #include "Widgets/Layout/SSplitter.h"
 
 class SWidget;
-class SFieldGroup;
+class SRCPanelGroup;
 struct SRCPanelExposedField;
 struct SRCPanelExposedActor;
 
@@ -47,11 +47,13 @@ struct SRCPanelTreeNode
 	virtual ENodeType GetType() const { return Invalid; };
 	/** Refresh the node. */
 	virtual void Refresh() {};
+	/** Get the context menu for this node. */
+	virtual TSharedPtr<SWidget> GetContextMenu() { return nullptr; }
 
 	//~ Utiliy methods for not having to downcast 
 	virtual TSharedPtr<SRCPanelExposedField> AsField() { return nullptr; }
 	virtual TSharedPtr<SWidget> AsFieldChild() { return nullptr; }
-	virtual TSharedPtr<SFieldGroup> AsGroup() { return nullptr; }
+	virtual TSharedPtr<SRCPanelGroup> AsGroup() { return nullptr; }
 	virtual TSharedPtr<SRCPanelExposedActor> AsActor() { return nullptr; }
 
 protected:
diff --git a/Engine/Plugins/VirtualProduction/RemoteControl/Source/RemoteControlUI/Private/UI/SRemoteControlPanel.cpp b/Engine/Plugins/VirtualProduction/RemoteControl/Source/RemoteControlUI/Private/UI/SRemoteControlPanel.cpp
--- a/Engine/Plugins/VirtualProduction/RemoteControl/Source/RemoteControlUI/Private/UI/SRemoteControlPanel.cpp
+++ b/Engine/Plugins/VirtualProduction/RemoteControl/Source/RemoteControlUI/Private/UI/SRemoteControlPanel.cpp
@@ -21,6 +21,7 @@
 #include "PropertyCustomizationHelpers.h"
 #include "PropertyHandle.h"
 #include "PropertyEditorModule.h"
+#include "RCPanelWidgetRegistry.h"
 #include "RemoteControlActor.h"
 #include "RemoteControlEntity.h"
 #include "RemoteControlField.h"
@@ -86,7 +87,8 @@ void SRemoteControlPanel::Construct(const FArguments& InArgs, URemoteControlPres
 {
 	OnEditModeChange = InArgs._OnEditModeChange;
 	Preset = TStrongObjectPtr<URemoteControlPreset>(InPreset);
-
+	UpdateRebindButtonVisibility();
+	
 	TArray<TSharedRef<SWidget>> ExtensionWidgets;
 	FRemoteControlUIModule::Get().GetExtensionGenerators().Broadcast(ExtensionWidgets);
 
@@ -96,6 +98,7 @@ void SRemoteControlPanel::Construct(const FArguments& InArgs, URemoteControlPres
 	
 	EntityList = SNew(SRCPanelExposedEntitiesList, Preset.Get())
 		.DisplayValues(true)
+		.OnEntityListUpdated_Lambda([this] () { UpdateRebindButtonVisibility(); })
 		.EditMode_Lambda([this](){ return bIsInEditMode; });
 	
 	EntityList->OnSelectionChange().AddSP(this, &SRemoteControlPanel::UpdateEntityDetailsView);
@@ -137,7 +140,7 @@ void SRemoteControlPanel::Construct(const FArguments& InArgs, URemoteControlPres
 			[
 				CreateExposeButton()
 			]
-
+			// Right aligned widgets
 			+ SHorizontalBox::Slot()
 			.VAlign(VAlign_Center)
 			.HAlign(HAlign_Right)
@@ -145,6 +148,20 @@ void SRemoteControlPanel::Construct(const FArguments& InArgs, URemoteControlPres
 			.Padding(0, 7.0f)
 			[
 				SAssignNew(TopExtensions, SHorizontalBox)
+				// Rebind button
+				+ SHorizontalBox::Slot()
+				.VAlign(VAlign_Center)
+				.AutoWidth()
+				[
+					SNew(SButton)
+					.Visibility_Lambda([this]() { return bShowRebindButton ? EVisibility::Visible : EVisibility::Collapsed; })
+					.OnClicked_Raw(this, &SRemoteControlPanel::OnClickRebindAllButton)
+					[
+						SNew(STextBlock)
+						.ToolTipText(LOCTEXT("RebindButtonToolTip", "Attempt to rebind all unbound entites of the preset."))
+						.Text(LOCTEXT("RebindButtonText", "Rebind All"))
+					]
+				]
 				+ SHorizontalBox::Slot()
 				.VAlign(VAlign_Center)
 				.Padding(4.0f, 0)
@@ -183,6 +200,7 @@ void SRemoteControlPanel::Construct(const FArguments& InArgs, URemoteControlPres
 				+ SWidgetSwitcher::Slot()
 				[
 					SNew(SSplitter)
+					.Orientation(EOrientation::Orient_Vertical)
 					+ SSplitter::Slot()
 					.Value(0.7f)
 					[
@@ -238,6 +256,7 @@ void SRemoteControlPanel::Construct(const FArguments& InArgs, URemoteControlPres
 SRemoteControlPanel::~SRemoteControlPanel()
 {
 	UnregisterEvents();
+	FRCPanelWidgetRegistry::Get().Clear();
 }
 
 void SRemoteControlPanel::PostUndo(bool bSuccess)
@@ -252,6 +271,11 @@ void SRemoteControlPanel::PostRedo(bool bSuccess)
 
 bool SRemoteControlPanel::IsExposed(const TSharedPtr<IPropertyHandle>& PropertyHandle)
 {
+	if (CachedExposedProperties.Contains(TWeakPtr<IPropertyHandle>{PropertyHandle}))
+	{
+		return true;
+	}
+	
 	TArray<UObject*> OuterObjects;
 	PropertyHandle->GetOuterObjects(OuterObjects);
 
@@ -287,6 +311,11 @@ bool SRemoteControlPanel::IsExposed(const TSharedPtr<IPropertyHandle>& PropertyH
 		bAllObjectsExposed &= bFoundPropForObject;
 	}
 
+	if (bAllObjectsExposed)
+	{
+		CachedExposedProperties.Emplace(PropertyHandle);
+	}
+	
 	return bAllObjectsExposed;
 }
 
@@ -307,7 +336,6 @@ void SRemoteControlPanel::ToggleProperty(const TSharedPtr<IPropertyHandle>& Prop
 	{
 		FScopedTransaction Transaction(LOCTEXT("ExposeProperty", "Expose Property"));
 		Preset->Modify();
-
 		for (UObject* Object : OuterObjects)
 		{
 			if (Object)
@@ -316,6 +344,8 @@ void SRemoteControlPanel::ToggleProperty(const TSharedPtr<IPropertyHandle>& Prop
 				ExposeProperty(Object, FRCFieldPathInfo{PropertyHandle->GeneratePathToProperty(), bCleanDuplicates});
 			}
 		}
+		
+		CachedExposedProperties.Emplace(PropertyHandle);
 	}
 }
 
@@ -364,6 +394,8 @@ TSharedRef<SWidget> SRemoteControlPanel::CreateExposeButton()
 {	
 	FMenuBuilder MenuBuilder(true, nullptr);
 
+	
+	
 	SAssignNew(BlueprintPicker, SRCPanelFunctionPicker)
 		.AllowDefaultObjects(true)
 		.Label(LOCTEXT("FunctionLibrariesLabel", "Function Libraries"))
@@ -437,6 +469,7 @@ TSharedRef<SWidget> SRemoteControlPanel::CreateExposeButton()
 	MenuBuilder.EndSection();
 	
 	return SAssignNew(ExposeComboButton, SComboButton)
+	.Visibility_Lambda([this]() { return this->bIsInEditMode ? EVisibility::Visible : EVisibility::Collapsed; })
 	.ButtonStyle(FEditorStyle::Get(), "PropertyEditor.AssetComboStyle")
 	.ForegroundColor(FSlateColor::UseForeground())
 	.CollapseMenuOnParentFocus(true)
@@ -569,6 +602,8 @@ void SRemoteControlPanel::OnMapChange(uint32)
 	{
 		ClassPicker->Refresh();	
 	}
+
+	UpdateRebindButtonVisibility();
 }
 
 void SRemoteControlPanel::RegisterEvents()
@@ -584,10 +619,16 @@ void SRemoteControlPanel::RegisterEvents()
 	{
 		GEngine->OnLevelActorAdded().AddSP(this, &SRemoteControlPanel::OnActorAddedToLevel);
 	}
+
+	Preset->OnEntityExposed().AddSP(this, &SRemoteControlPanel::OnEntityExposed);
+	Preset->OnEntityUnexposed().AddSP(this, &SRemoteControlPanel::OnEntityUnexposed);
 }
 
 void SRemoteControlPanel::UnregisterEvents()
 {
+	Preset->OnEntityExposed().RemoveAll(this);
+	Preset->OnEntityUnexposed().RemoveAll(this);
+	
 	if (GEngine)
 	{
 		GEngine->OnLevelActorAdded().RemoveAll(this);
@@ -787,4 +828,45 @@ void SRemoteControlPanel::UpdateEntityDetailsView(const TSharedPtr<SRCPanelTreeN
 	}
 }
 
+void SRemoteControlPanel::UpdateRebindButtonVisibility()
+{
+	if (URemoteControlPreset* PresetPtr = Preset.Get())
+	{
+		for (TWeakPtr<FRemoteControlEntity> WeakEntity : PresetPtr->GetExposedEntities<FRemoteControlEntity>())
+		{
+			if (TSharedPtr<FRemoteControlEntity> Entity = WeakEntity.Pin())
+			{
+				if (!Entity->IsBound())
+				{
+					bShowRebindButton = true;
+					return;
+				}
+			}
+		}
+	}
+
+	bShowRebindButton = false;
+}
+
+FReply SRemoteControlPanel::OnClickRebindAllButton()
+{
+	if (URemoteControlPreset* PresetPtr = Preset.Get())
+	{
+		PresetPtr->RebindUnboundEntities();
+
+		UpdateRebindButtonVisibility();
+	}
+	return FReply::Handled();
+}
+
+void SRemoteControlPanel::OnEntityExposed(URemoteControlPreset* InPreset, const FGuid& InEntityId)
+{
+	CachedExposedProperties.Empty();
+}
+
+void SRemoteControlPanel::OnEntityUnexposed(URemoteControlPreset* InPreset, const FGuid& InEntityId)
+{
+	CachedExposedProperties.Empty();
+}
+
 #undef LOCTEXT_NAMESPACE /*RemoteControlPanel*/
diff --git a/Engine/Plugins/VirtualProduction/RemoteControl/Source/RemoteControlUI/Private/UI/SRemoteControlPanel.h b/Engine/Plugins/VirtualProduction/RemoteControl/Source/RemoteControlUI/Private/UI/SRemoteControlPanel.h
--- a/Engine/Plugins/VirtualProduction/RemoteControl/Source/RemoteControlUI/Private/UI/SRemoteControlPanel.h
+++ b/Engine/Plugins/VirtualProduction/RemoteControl/Source/RemoteControlUI/Private/UI/SRemoteControlPanel.h
@@ -153,6 +153,16 @@ private:
 	
 	/** Update the details view following entity selection change.  */
 	void UpdateEntityDetailsView(const TSharedPtr<SRCPanelTreeNode>& SelectedNode);
+
+	/** Returns whether the preset has any unbound property or function. */
+	void UpdateRebindButtonVisibility();
+
+	/** Handle user clicking on the rebind all button. */
+	FReply OnClickRebindAllButton();
+
+	//~ Handlers called in order to clear the exposed property cache.
+	void OnEntityExposed(URemoteControlPreset* InPreset, const FGuid& InEntityId);
+	void OnEntityUnexposed(URemoteControlPreset* InPreset, const FGuid& InEntityId);
 private:
 	/** Holds the preset asset. */
 	TStrongObjectPtr<URemoteControlPreset> Preset;
@@ -185,4 +195,8 @@ private:
 	 * Used in order to ensure that the generated details view keeps a valid pointer to the selected entity.
 	 */
 	TSharedPtr<FRemoteControlEntity> SelectedEntity;
+	/** Whether to show the rebind all button. */
+	bool bShowRebindButton = false;
+	/** Cache of exposed properties. */
+	TSet<TWeakPtr<IPropertyHandle>> CachedExposedProperties;
 };
