Fix merge conflict error

diff --git a/Engine/Plugins/VirtualProduction/RemoteControl/Source/RemoteControlUI/Private/UI/SRCPanelExposedEntity.cpp b/Engine/Plugins/VirtualProduction/RemoteControl/Source/RemoteControlUI/Private/UI/SRCPanelExposedEntity.cpp
--- a/Engine/Plugins/VirtualProduction/RemoteControl/Source/RemoteControlUI/Private/UI/SRCPanelExposedEntity.cpp
+++ b/Engine/Plugins/VirtualProduction/RemoteControl/Source/RemoteControlUI/Private/UI/SRCPanelExposedEntity.cpp
@@ -58,6 +58,10 @@ TSharedRef<SWidget> SRCPanelExposedEntity::CreateInvalidWidget()
 TSharedRef<SWidget> SRCPanelExposedEntity::CreateRebindMenuContent()
 {
 	FSceneOutlinerModule& SceneOutlinerModule = FModuleManager::Get().LoadModuleChecked<FSceneOutlinerModule>("SceneOutliner");
+	FSceneOutlinerInitializationOptions Options;
+	Options.Filters = MakeShared<FSceneOutlinerFilters>();
+	Options.Filters->AddFilterPredicate<FActorTreeItem>(FActorTreeItem::FFilterPredicate::CreateRaw(this, &SRCPanelExposedEntity::IsActorSelectable));
+
 	return SNew(SBox)
 	.MaxDesiredHeight(400.0f)
 	.WidthOverride(300.0f)
