RC Add logging for protocols

diff --git a/Engine/Plugins/VirtualProduction/RemoteControl/Source/RemoteControl/Private/RemoteControlLogger.cpp b/Engine/Plugins/VirtualProduction/RemoteControl/Source/RemoteControl/Private/RemoteControlLogger.cpp
new file mode 100644
--- /dev/null
+++ b/Engine/Plugins/VirtualProduction/RemoteControl/Source/RemoteControl/Private/RemoteControlLogger.cpp
@@ -0,0 +1,76 @@
+﻿// Copyright Epic Games, Inc. All Rights Reserved.
+
+#include "RemoteControlLogger.h"
+
+#include "IMessageLogListing.h"
+#include "MessageLogModule.h"
+#include "Modules/ModuleManager.h"
+
+#define LOCTEXT_NAMESPACE "RemoteControlLogger"
+
+FRemoteControlLogger::FRemoteControlLogger()
+{
+	FMessageLogModule& MessageLogModule = FModuleManager::LoadModuleChecked<FMessageLogModule>("MessageLog");
+
+	FMessageLogInitializationOptions LogOptions;
+
+	// Don't show Pages so that user is never allowed to clear log messages
+	LogOptions.bShowPages   = false;
+	LogOptions.bShowFilters = false;
+	LogOptions.bAllowClear  = true;
+	LogOptions.MaxPageCount = 1;
+
+	MessageLogListing = MessageLogModule.CreateLogListing("Remote control logging", LogOptions);
+
+	// Create widget
+	LogListingWidget = MessageLogModule.CreateLogListingWidget(MessageLogListing.ToSharedRef());
+}
+
+TSharedRef<SWidget> FRemoteControlLogger::GetWidget() const
+{
+	return LogListingWidget.ToSharedRef();
+}
+
+void FRemoteControlLogger::Log(const FName& InputType, FLogCallback InLogTextCallback, EVerbosityLevel Verbosity)
+{
+	if (!bIsEnabled)
+	{
+		return;
+	}
+	
+	if (!ensure(MessageLogListing.IsValid()))
+	{
+		return;
+	}
+		
+	TArray<TSharedRef<FTokenizedMessage>> Messages;
+	TSharedRef<FTokenizedMessage> Line = FTokenizedMessage::Create(
+		Verbosity == EVerbosityLevel::Log ? EMessageSeverity::Info : 
+			(Verbosity == EVerbosityLevel::Warning ? EMessageSeverity::Warning : EMessageSeverity::Error));
+
+	const FText Message = InLogTextCallback();
+	Line->AddToken(FTextToken::Create(
+		FText::Format(LOCTEXT("null", "{0}"), FText::FromString(FDateTime::Now().ToString(TEXT("%H:%M:%S - "))))));
+	Line->AddToken(FTextToken::Create(FText::Format(LOCTEXT("null", "{0}"), FText::FromName(InputType))));
+	Line->AddToken(FTextToken::Create(Message));
+	Messages.Add(Line);
+
+	constexpr bool bMirrorToOutputLog = false;
+	MessageLogListing->AddMessages(MoveTemp(Messages), bMirrorToOutputLog);
+
+	// Always select last message, that keep the UI widget scrolling
+	constexpr bool bSelected = true;
+	MessageLogListing->SelectMessage(Line, bSelected);
+}
+
+void FRemoteControlLogger::EnableLog(const bool bEnable)
+{
+	bIsEnabled = bEnable;
+}
+
+void FRemoteControlLogger::ClearLog() const
+{
+	MessageLogListing->ClearMessages();
+}
+
+#undef LOCTEXT_NAMESPACE
diff --git a/Engine/Plugins/VirtualProduction/RemoteControl/Source/RemoteControl/Public/RemoteControlLogger.h b/Engine/Plugins/VirtualProduction/RemoteControl/Source/RemoteControl/Public/RemoteControlLogger.h
new file mode 100644
--- /dev/null
+++ b/Engine/Plugins/VirtualProduction/RemoteControl/Source/RemoteControl/Public/RemoteControlLogger.h
@@ -0,0 +1,79 @@
+﻿// Copyright Epic Games, Inc. All Rights Reserved.
+
+#pragma once
+
+#include "CoreMinimal.h"
+
+class IMessageLogListing;
+class SWidget;
+
+/**
+ * Log verbosity level
+ */
+enum class EVerbosityLevel
+{
+	Log,
+	Warning,
+	Error
+};
+
+/**
+ * Remote control logger.
+ * Generated Log UI and collect logs from APIs and protocols
+ */
+class REMOTECONTROL_API FRemoteControlLogger : public TSharedFromThis<FRemoteControlLogger>
+{
+public:
+	/** Get singleton logger instance */
+	static FRemoteControlLogger& Get()
+	{
+		static FRemoteControlLogger* Instance = nullptr;
+		if (Instance == nullptr)
+		{
+			Instance = new FRemoteControlLogger;
+		}
+		return *Instance;
+	}
+
+	/** Callback for logging, it should return FText to log */
+	using FLogCallback = TFunctionRef<FText(void)>;
+	
+	FRemoteControlLogger();
+	virtual ~FRemoteControlLogger() = default;
+
+	/**
+	 * Get Logger SWidget
+	 */
+	TSharedRef<SWidget> GetWidget() const;
+
+	/**
+	 * Log the message
+	 *
+	 * @param InputType Type of the log, protocol, api, webapi
+	 * @param InLogTextCallback	Log callback which should return the Text to log
+	 * @param Verbosity			Log level
+	 */
+	void Log(const FName& InputType, FLogCallback InLogTextCallback, EVerbosityLevel Verbosity = EVerbosityLevel::Log);
+
+	/** 
+	 * Enable/Disable the log for Remote Control
+	 * @param bEnable enable flag
+	 */
+	void EnableLog(const bool bEnable);
+
+	/** Return the current state of the log */
+	bool IsEnabled() const { return bIsEnabled; }
+
+	/** Removes all messages from log */
+	void ClearLog() const;
+
+private:
+	/** Log listening interface */
+	TSharedPtr<IMessageLogListing> MessageLogListing;
+
+	/** Pointer to log widget */
+	TSharedPtr<SWidget> LogListingWidget;
+
+	/** Is the logger enabled */
+	bool bIsEnabled = false;
+};
diff --git a/Engine/Plugins/VirtualProduction/RemoteControl/Source/RemoteControl/RemoteControl.Build.cs b/Engine/Plugins/VirtualProduction/RemoteControl/Source/RemoteControl/RemoteControl.Build.cs
--- a/Engine/Plugins/VirtualProduction/RemoteControl/Source/RemoteControl/RemoteControl.Build.cs
+++ b/Engine/Plugins/VirtualProduction/RemoteControl/Source/RemoteControl/RemoteControl.Build.cs
@@ -17,6 +17,7 @@ public class RemoteControl : ModuleRules
 			new string[] {
 				"Cbor",
 				"Engine",
+				"MessageLog",
 				"RemoteControlInterception",
 				"Serialization",
 			}
diff --git a/Engine/Plugins/VirtualProduction/RemoteControl/Source/RemoteControlUI/Private/UI/SRemoteControlPanel.cpp b/Engine/Plugins/VirtualProduction/RemoteControl/Source/RemoteControlUI/Private/UI/SRemoteControlPanel.cpp
--- a/Engine/Plugins/VirtualProduction/RemoteControl/Source/RemoteControlUI/Private/UI/SRemoteControlPanel.cpp
+++ b/Engine/Plugins/VirtualProduction/RemoteControl/Source/RemoteControlUI/Private/UI/SRemoteControlPanel.cpp
@@ -25,6 +25,7 @@
 #include "RemoteControlActor.h"
 #include "RemoteControlEntity.h"
 #include "RemoteControlField.h"
+#include "RemoteControlLogger.h"
 #include "RemoteControlPanelStyle.h"
 #include "RemoteControlPreset.h"
 #include "RemoteControlUIModule.h"
@@ -178,6 +179,22 @@ void SRemoteControlPanel::Construct(const FArguments& InArgs, URemoteControlPres
 					.IsChecked_Lambda([this]() { return this->bIsInEditMode ? ECheckBoxState::Checked : ECheckBoxState::Unchecked; })
 					.OnCheckStateChanged(this, &SRemoteControlPanel::OnEditModeCheckboxToggle)
 				]
+				+ SHorizontalBox::Slot()
+				.VAlign(VAlign_Center)
+				.Padding(4.0f, 0)
+				.AutoWidth()
+				[
+					SNew(STextBlock)
+					.Text(LOCTEXT("EnableLogLabel", "Enable Log: "))
+				]
+				+ SHorizontalBox::Slot()
+				.VAlign(VAlign_Center)
+				.AutoWidth()
+				[
+					SNew(SCheckBox)
+					.IsChecked_Lambda([]() { return FRemoteControlLogger::Get().IsEnabled() ? ECheckBoxState::Checked : ECheckBoxState::Unchecked; })
+					.OnCheckStateChanged(this, &SRemoteControlPanel::OnLogCheckboxToggle)
+				]
 			]
 		]
 		+ SVerticalBox::Slot()
@@ -186,53 +203,75 @@ void SRemoteControlPanel::Construct(const FArguments& InArgs, URemoteControlPres
 			.Padding(FMargin(0.f, 5.f, 0.f, 0.f))
 			.BorderImage(FEditorStyle::GetBrush("ToolPanel.DarkGroupBorder"))
 			[
-				SNew(SWidgetSwitcher)
-				.WidgetIndex_Lambda([this](){ return !bIsInEditMode ? 0 : 1; })
-				+ SWidgetSwitcher::Slot()
+				SNew(SSplitter)
+				.Orientation(Orient_Vertical)
+				+ SSplitter::Slot()
+				.Value(.8f)
 				[
-					// Exposed entities List
 					SNew(SBorder)
-					.BorderImage(FEditorStyle::GetBrush("ToolPanel.GroupBorder"))
-					[
-						EntityList.ToSharedRef()
-					]
-				]
-				+ SWidgetSwitcher::Slot()
-				[
-					SNew(SSplitter)
-					.Orientation(EOrientation::Orient_Vertical)
-					+ SSplitter::Slot()
-					.Value(0.7f)
+					.Padding(FMargin(0.f, 5.f, 0.f, 0.f))
+					.BorderImage(FEditorStyle::GetBrush("ToolPanel.DarkGroupBorder"))
 					[
-						// Exposed entities List
-						SNew(SBorder)
-						.BorderImage(FEditorStyle::GetBrush("ToolPanel.GroupBorder"))
+						SNew(SWidgetSwitcher)
+						.WidgetIndex_Lambda([this](){ return !bIsInEditMode ? 0 : 1; })
+						+ SWidgetSwitcher::Slot()
 						[
-							EntityList.ToSharedRef()
+							// Exposed entities List
+							SNew(SBorder)
+							.BorderImage(FEditorStyle::GetBrush("ToolPanel.GroupBorder"))
+							[
+								EntityList.ToSharedRef()
+							]
 						]
-					]
-					+ SSplitter::Slot()
-					.Value(0.3f)
-					[
-						SNew(SBorder)
-						.BorderImage(FEditorStyle::GetBrush("ToolPanel.GroupBorder"))
-						.Padding(5.f)
+						+ SWidgetSwitcher::Slot()
 						[
 							SNew(SSplitter)
-							.Orientation(Orient_Vertical)
+							.Orientation(EOrientation::Orient_Vertical)
 							+ SSplitter::Slot()
-							.Value(0.4f)
+							.Value(0.7f)
 							[
-								CreateEntityDetailsView()
+								// Exposed entities List
+								SNew(SBorder)
+								.BorderImage(FEditorStyle::GetBrush("ToolPanel.GroupBorder"))
+								[
+									EntityList.ToSharedRef()
+								]
 							]
 							+ SSplitter::Slot()
-							.Value(0.6f)
+							.Value(0.3f)
 							[
-								EntityProtocolDetails.ToSharedRef()
+								SNew(SBorder)
+								.BorderImage(FEditorStyle::GetBrush("ToolPanel.GroupBorder"))
+								.Padding(5.f)
+								[
+									SNew(SSplitter)
+									.Orientation(Orient_Vertical)
+									+ SSplitter::Slot()
+									.Value(0.4f)
+									[
+										CreateEntityDetailsView()
+									]
+									+ SSplitter::Slot()
+									.Value(0.6f)
+									[
+										EntityProtocolDetails.ToSharedRef()
+									]
+								]
 							]
 						]
 					]
 				]
+				+ SSplitter::Slot()
+				.Value(.2f)
+				[
+					SNew(SBorder)
+					.BorderImage(FEditorStyle::GetBrush("Menu.Background"))
+					.Visibility_Lambda([](){ return FRemoteControlLogger::Get().IsEnabled() ? EVisibility::Visible : EVisibility::Collapsed; })
+					.Padding(2.f)
+					[
+						FRemoteControlLogger::Get().GetWidget()
+					]
+				]
 			]
 		]
 	];
@@ -240,7 +279,7 @@ void SRemoteControlPanel::Construct(const FArguments& InArgs, URemoteControlPres
 	for (const TSharedRef<SWidget>& Widget : ExtensionWidgets)
 	{
 		// We want to insert the widgets before the edit mode buttons.
-		TopExtensions->InsertSlot(TopExtensions->NumSlots()-2)
+		TopExtensions->InsertSlot(TopExtensions->NumSlots()-4)
 		.VAlign(VAlign_Center)
 		.AutoWidth()
 		[
@@ -257,6 +296,9 @@ SRemoteControlPanel::~SRemoteControlPanel()
 {
 	UnregisterEvents();
 	FRCPanelWidgetRegistry::Get().Clear();
+
+	// Clear the log
+	FRemoteControlLogger::Get().ClearLog();
 }
 
 void SRemoteControlPanel::PostUndo(bool bSuccess)
@@ -697,9 +739,15 @@ void SRemoteControlPanel::Unexpose(const TSharedPtr<IPropertyHandle>& Handle)
 
 void SRemoteControlPanel::OnEditModeCheckboxToggle(ECheckBoxState State)
 {
-	bIsInEditMode = State == ECheckBoxState::Checked ? true : false;
+	bIsInEditMode = (State == ECheckBoxState::Checked) ? true : false;
 	OnEditModeChange.ExecuteIfBound(SharedThis(this), bIsInEditMode);
 }
+
+void SRemoteControlPanel::OnLogCheckboxToggle(ECheckBoxState State)
+{
+	const bool bIsLogEnabled = (State == ECheckBoxState::Checked) ? true : false;
+	FRemoteControlLogger::Get().EnableLog(bIsLogEnabled);
+}
  
 void SRemoteControlPanel::OnBlueprintReinstanced()
 {
diff --git a/Engine/Plugins/VirtualProduction/RemoteControl/Source/RemoteControlUI/Private/UI/SRemoteControlPanel.h b/Engine/Plugins/VirtualProduction/RemoteControl/Source/RemoteControlUI/Private/UI/SRemoteControlPanel.h
--- a/Engine/Plugins/VirtualProduction/RemoteControl/Source/RemoteControlUI/Private/UI/SRemoteControlPanel.h
+++ b/Engine/Plugins/VirtualProduction/RemoteControl/Source/RemoteControlUI/Private/UI/SRemoteControlPanel.h
@@ -163,6 +163,8 @@ private:
 	//~ Handlers called in order to clear the exposed property cache.
 	void OnEntityExposed(URemoteControlPreset* InPreset, const FGuid& InEntityId);
 	void OnEntityUnexposed(URemoteControlPreset* InPreset, const FGuid& InEntityId);
+	/** Toggles the logging part of UI */
+	void OnLogCheckboxToggle(ECheckBoxState State);
 private:
 	/** Holds the preset asset. */
 	TStrongObjectPtr<URemoteControlPreset> Preset;
