Generate Transaction in Protocol Settings

diff --git a/Engine/Plugins/VirtualProduction/RemoteControl/Source/RemoteControl/Private/RemoteControlProtocolBinding.cpp b/Engine/Plugins/VirtualProduction/RemoteControl/Source/RemoteControl/Private/RemoteControlProtocolBinding.cpp
--- a/Engine/Plugins/VirtualProduction/RemoteControl/Source/RemoteControl/Private/RemoteControlProtocolBinding.cpp
+++ b/Engine/Plugins/VirtualProduction/RemoteControl/Source/RemoteControl/Private/RemoteControlProtocolBinding.cpp
@@ -4,19 +4,15 @@
 
 #include "IRemoteControlModule.h"
 #include "RemoteControlPreset.h"
+#include "RemoteControlSettings.h"
 
+#include "CborWriter.h"
 #include "Algo/Sort.h"
 #include "Backends/CborStructDeserializerBackend.h"
-#include "CborWriter.h"
-#include "Engine/Engine.h"
-#include "Serialization/MemoryWriter.h"
 #include "Serialization/MemoryReader.h"
+#include "Serialization/MemoryWriter.h"
 #include "UObject/StructOnScope.h"
 
-#if WITH_EDITOR
-#include "ScopedTransaction.h"
-#endif
-
 #define LOCTEXT_NAMESPACE "RemoteControl"
 
 namespace EntityInterpolation
@@ -474,7 +470,14 @@ bool FRemoteControlProtocolEntity::ApplyProtocolValueToProperty(double InProtoco
 
 	FRCObjectReference ObjectRef;
 	ObjectRef.Property = Property;
-	ObjectRef.Access = bGenerateTransaction ? ERCAccess::WRITE_TRANSACTION_ACCESS : ERCAccess::WRITE_ACCESS;
+	ObjectRef.Access = ERCAccess::WRITE_ACCESS;
+
+	const URemoteControlSettings* RemoteControlSettings = GetDefault<URemoteControlSettings>();
+	if (RemoteControlSettings->bProtocolsGenerateTransactions)
+	{
+		ObjectRef.Access = ERCAccess::WRITE_TRANSACTION_ACCESS;
+	}
+	
 	ObjectRef.PropertyPathInfo = RemoteControlProperty->FieldPathInfo.ToString();
 
 	bool bSuccess = true;
diff --git a/Engine/Plugins/VirtualProduction/RemoteControl/Source/RemoteControl/Public/RemoteControlProtocolBinding.h b/Engine/Plugins/VirtualProduction/RemoteControl/Source/RemoteControl/Public/RemoteControlProtocolBinding.h
--- a/Engine/Plugins/VirtualProduction/RemoteControl/Source/RemoteControl/Public/RemoteControlProtocolBinding.h
+++ b/Engine/Plugins/VirtualProduction/RemoteControl/Source/RemoteControl/Public/RemoteControlProtocolBinding.h
@@ -377,10 +377,6 @@ protected:
 	UPROPERTY()
 	FGuid PropertyId;
 
-	/** Should property generate transaction events? */
-	UPROPERTY(EditAnywhere, Category = Mapping)
-	bool bGenerateTransaction = true;
-
 protected:
 	/** 
 	 * Property mapping ranges set
diff --git a/Engine/Plugins/VirtualProduction/RemoteControl/Source/RemoteControl/Public/RemoteControlSettings.h b/Engine/Plugins/VirtualProduction/RemoteControl/Source/RemoteControl/Public/RemoteControlSettings.h
new file mode 100644
--- /dev/null
+++ b/Engine/Plugins/VirtualProduction/RemoteControl/Source/RemoteControl/Public/RemoteControlSettings.h
@@ -0,0 +1,26 @@
+﻿// Copyright Epic Games, Inc. All Rights Reserved.
+
+#pragma once
+
+#include "CoreMinimal.h"
+#include "Engine/DeveloperSettings.h"
+
+#include "RemoteControlSettings.generated.h"
+
+
+/**
+ * Global remote control editor settings
+ */
+UCLASS(config=Editor)
+class REMOTECONTROL_API URemoteControlSettings : public UDeveloperSettings
+{
+	GENERATED_BODY()
+
+public:
+	/**
+	 * Should generate transaction for protocol events
+	 * It makes the transaction for protocls optionals and improve the performance in Editor
+	 */
+	UPROPERTY(config, EditAnywhere, Category = RemoteControl)
+	bool bProtocolsGenerateTransactions = true;
+};
diff --git a/Engine/Plugins/VirtualProduction/RemoteControl/Source/RemoteControl/RemoteControl.Build.cs b/Engine/Plugins/VirtualProduction/RemoteControl/Source/RemoteControl/RemoteControl.Build.cs
--- a/Engine/Plugins/VirtualProduction/RemoteControl/Source/RemoteControl/RemoteControl.Build.cs
+++ b/Engine/Plugins/VirtualProduction/RemoteControl/Source/RemoteControl/RemoteControl.Build.cs
@@ -27,6 +27,7 @@ public class RemoteControl : ModuleRules
 		{
 			PrivateDependencyModuleNames.AddRange(
 				new string[] {
+					"DeveloperSettings",
 					"EditorFramework",
 					"UnrealEd",
 				}
