fix for MessageLog referenced from Runtime plugin

diff --git a/Engine/Plugins/VirtualProduction/RemoteControl/Source/RemoteControl/Private/RemoteControlLogger.cpp b/Engine/Plugins/VirtualProduction/RemoteControl/Source/RemoteControl/Private/RemoteControlLogger.cpp
--- a/Engine/Plugins/VirtualProduction/RemoteControl/Source/RemoteControl/Private/RemoteControlLogger.cpp
+++ b/Engine/Plugins/VirtualProduction/RemoteControl/Source/RemoteControl/Private/RemoteControlLogger.cpp
@@ -2,6 +2,8 @@
 
 #include "RemoteControlLogger.h"
 
+#if WITH_EDITOR
+
 #include "IMessageLogListing.h"
 #include "MessageLogModule.h"
 #include "Modules/ModuleManager.h"
@@ -74,3 +76,5 @@ void FRemoteControlLogger::ClearLog() const
 }
 
 #undef LOCTEXT_NAMESPACE
+
+#endif //WITH_EDITOR
diff --git a/Engine/Plugins/VirtualProduction/RemoteControl/Source/RemoteControl/Public/RemoteControlLogger.h b/Engine/Plugins/VirtualProduction/RemoteControl/Source/RemoteControl/Public/RemoteControlLogger.h
--- a/Engine/Plugins/VirtualProduction/RemoteControl/Source/RemoteControl/Public/RemoteControlLogger.h
+++ b/Engine/Plugins/VirtualProduction/RemoteControl/Source/RemoteControl/Public/RemoteControlLogger.h
@@ -2,6 +2,8 @@
 
 #pragma once
 
+#if WITH_EDITOR
+
 #include "CoreMinimal.h"
 
 class IMessageLogListing;
@@ -77,3 +79,5 @@ private:
 	/** Is the logger enabled */
 	bool bIsEnabled = false;
 };
+
+#endif // WITH_EDITOR
