- Fixing CIS related RemoteControl logger breaking shipping build

diff --git a/Engine/Plugins/VirtualProduction/RemoteControl/Source/RemoteControl/Private/RemoteControlLogger.cpp b/Engine/Plugins/VirtualProduction/RemoteControl/Source/RemoteControl/Private/RemoteControlLogger.cpp
--- a/Engine/Plugins/VirtualProduction/RemoteControl/Source/RemoteControl/Private/RemoteControlLogger.cpp
+++ b/Engine/Plugins/VirtualProduction/RemoteControl/Source/RemoteControl/Private/RemoteControlLogger.cpp
@@ -1,17 +1,19 @@
-﻿// Copyright Epic Games, Inc. All Rights Reserved.
+// Copyright Epic Games, Inc. All Rights Reserved.
 
 #include "RemoteControlLogger.h"
 
-#if WITH_EDITOR
+#include "Modules/ModuleManager.h"
 
+#if WITH_EDITOR
 #include "IMessageLogListing.h"
 #include "MessageLogModule.h"
-#include "Modules/ModuleManager.h"
+#endif
 
 #define LOCTEXT_NAMESPACE "RemoteControlLogger"
 
 FRemoteControlLogger::FRemoteControlLogger()
 {
+#if WITH_EDITOR
 	FMessageLogModule& MessageLogModule = FModuleManager::LoadModuleChecked<FMessageLogModule>("MessageLog");
 
 	FMessageLogInitializationOptions LogOptions;
@@ -26,15 +28,21 @@ FRemoteControlLogger::FRemoteControlLogger()
 
 	// Create widget
 	LogListingWidget = MessageLogModule.CreateLogListingWidget(MessageLogListing.ToSharedRef());
+#endif
 }
 
 TSharedRef<SWidget> FRemoteControlLogger::GetWidget() const
 {
+#if WITH_EDITOR
 	return LogListingWidget.ToSharedRef();
+#else
+	return SNullWidget::NullWidget;
+#endif
 }
 
 void FRemoteControlLogger::Log(const FName& InputType, FLogCallback InLogTextCallback, EVerbosityLevel Verbosity)
 {
+#if WITH_EDITOR
 	if (!bIsEnabled)
 	{
 		return;
@@ -63,6 +71,7 @@ void FRemoteControlLogger::Log(const FName& InputType, FLogCallback InLogTextCal
 	// Always select last message, that keep the UI widget scrolling
 	constexpr bool bSelected = true;
 	MessageLogListing->SelectMessage(Line, bSelected);
+#endif
 }
 
 void FRemoteControlLogger::EnableLog(const bool bEnable)
@@ -72,9 +81,9 @@ void FRemoteControlLogger::EnableLog(const bool bEnable)
 
 void FRemoteControlLogger::ClearLog() const
 {
+#if WITH_EDITOR
 	MessageLogListing->ClearMessages();
+#endif
 }
 
 #undef LOCTEXT_NAMESPACE
-
-#endif //WITH_EDITOR
diff --git a/Engine/Plugins/VirtualProduction/RemoteControl/Source/RemoteControl/Public/RemoteControlLogger.h b/Engine/Plugins/VirtualProduction/RemoteControl/Source/RemoteControl/Public/RemoteControlLogger.h
--- a/Engine/Plugins/VirtualProduction/RemoteControl/Source/RemoteControl/Public/RemoteControlLogger.h
+++ b/Engine/Plugins/VirtualProduction/RemoteControl/Source/RemoteControl/Public/RemoteControlLogger.h
@@ -1,9 +1,7 @@
-﻿// Copyright Epic Games, Inc. All Rights Reserved.
+// Copyright Epic Games, Inc. All Rights Reserved.
 
 #pragma once
 
-#if WITH_EDITOR
-
 #include "CoreMinimal.h"
 
 class IMessageLogListing;
@@ -70,8 +68,11 @@ public:
 	void ClearLog() const;
 
 private:
+
+#if WITH_EDITOR
 	/** Log listening interface */
 	TSharedPtr<IMessageLogListing> MessageLogListing;
+#endif
 
 	/** Pointer to log widget */
 	TSharedPtr<SWidget> LogListingWidget;
@@ -79,5 +80,3 @@ private:
 	/** Is the logger enabled */
 	bool bIsEnabled = false;
 };
-
-#endif // WITH_EDITOR
diff --git a/Engine/Plugins/VirtualProduction/RemoteControl/Source/RemoteControl/RemoteControl.Build.cs b/Engine/Plugins/VirtualProduction/RemoteControl/Source/RemoteControl/RemoteControl.Build.cs
--- a/Engine/Plugins/VirtualProduction/RemoteControl/Source/RemoteControl/RemoteControl.Build.cs
+++ b/Engine/Plugins/VirtualProduction/RemoteControl/Source/RemoteControl/RemoteControl.Build.cs
@@ -27,7 +27,6 @@ public class RemoteControl : ModuleRules
 			PrivateDependencyModuleNames.AddRange(
 				new string[] {
 					"DeveloperSettings",
-					"EditorFramework",
 					"MessageLog",
 					"UnrealEd",
 				}
