Fix rebind not working after deleting the object twice

Fix crash when property is invalid
and fix rc styling panel issue

Fix crash when saving preset with exposed function

Fix splitter being invisible in the bottom part

diff --git a/Engine/Plugins/VirtualProduction/RemoteControl/Source/RemoteControl/Private/RemoteControlEntity.cpp b/Engine/Plugins/VirtualProduction/RemoteControl/Source/RemoteControl/Private/RemoteControlEntity.cpp
--- a/Engine/Plugins/VirtualProduction/RemoteControl/Source/RemoteControl/Private/RemoteControlEntity.cpp
+++ b/Engine/Plugins/VirtualProduction/RemoteControl/Source/RemoteControl/Private/RemoteControlEntity.cpp
@@ -41,14 +41,26 @@ void FRemoteControlEntity::SetMetadataValue(FName Key, FString Value)
 
 void FRemoteControlEntity::BindObject(UObject* InObjectToBind)
 {
-	if (Bindings.Num() && InObjectToBind)
+	if (!InObjectToBind)
 	{
-		if (URemoteControlBinding* Binding = Bindings[0].Get())
-		{
-			Binding->Modify();
-			Binding->SetBoundObject(InObjectToBind);
-			OnEntityModifiedDelegate.ExecuteIfBound(Id);
-		}
+		return;
+	}
+
+	URemoteControlBinding* Binding = nullptr;
+
+	if (Bindings.Num() == 0)
+	{
+		Binding = Owner->FindOrAddBinding(InObjectToBind);
+		Bindings.Emplace(Binding);
+	}
+
+	Binding = Bindings[0].Get();
+
+	if (Binding)
+	{
+		Binding->Modify();
+		Binding->SetBoundObject(InObjectToBind);
+		OnEntityModifiedDelegate.ExecuteIfBound(Id);
 	}
 }
 
diff --git a/Engine/Plugins/VirtualProduction/RemoteControl/Source/RemoteControl/Private/RemoteControlField.cpp b/Engine/Plugins/VirtualProduction/RemoteControl/Source/RemoteControl/Private/RemoteControlField.cpp
--- a/Engine/Plugins/VirtualProduction/RemoteControl/Source/RemoteControl/Private/RemoteControlField.cpp
+++ b/Engine/Plugins/VirtualProduction/RemoteControl/Source/RemoteControl/Private/RemoteControlField.cpp
@@ -385,14 +385,20 @@ bool FRemoteControlFunction::IsBound() const
 
 UFunction* FRemoteControlFunction::GetFunction() const
 {
-	UObject* ResolvedFunction = FunctionPath.ResolveObject();
+	UObject* ResolvedObject = FunctionPath.ResolveObject();
 	
-	if (!ResolvedFunction)
+	if (!ResolvedObject)
 	{
-		ResolvedFunction = FunctionPath.TryLoad();
+		ResolvedObject = FunctionPath.TryLoad();
+	}
+
+	UFunction* ResolvedFunction = Cast<UFunction>(ResolvedObject);
+	if (ResolvedFunction)
+	{
+		CachedFunction = ResolvedFunction;
 	}
 	
-	return Cast<UFunction>(ResolvedFunction);
+	return ResolvedFunction;
 }
 
 void FRemoteControlFunction::AssignDefaultFunctionArguments()
@@ -434,13 +440,13 @@ FArchive& operator<<(FArchive& Ar, FRemoteControlFunction& RCFunction)
 		if (UFunction* Function = RCFunction.GetFunction())
 		{
 			RCFunction.FunctionArguments = MakeShared<FStructOnScope>(Function);
+			Function->SerializeTaggedProperties(Ar, RCFunction.FunctionArguments->GetStructMemory(), Function, nullptr);
 		}
 	}
-	
-	if (UFunction* Function = RCFunction.GetFunction())
+	else if (RCFunction.CachedFunction.IsValid())
 	{
-		Function->SerializeTaggedProperties(Ar, RCFunction.FunctionArguments->GetStructMemory(), Function, nullptr);
+		RCFunction.CachedFunction->SerializeTaggedProperties(Ar, RCFunction.FunctionArguments->GetStructMemory(), RCFunction.CachedFunction.Get(), nullptr);	
 	}
-
+		
 	return Ar;
 }
\ No newline at end of file
diff --git a/Engine/Plugins/VirtualProduction/RemoteControl/Source/RemoteControl/Private/RemoteControlPreset.cpp b/Engine/Plugins/VirtualProduction/RemoteControl/Source/RemoteControl/Private/RemoteControlPreset.cpp
--- a/Engine/Plugins/VirtualProduction/RemoteControl/Source/RemoteControl/Private/RemoteControlPreset.cpp
+++ b/Engine/Plugins/VirtualProduction/RemoteControl/Source/RemoteControl/Private/RemoteControlPreset.cpp
@@ -1626,11 +1626,12 @@ void URemoteControlPreset::OnActorDeleted(AActor* Actor)
 	{
 		if (Entity)
 		{
-			for (TWeakObjectPtr<URemoteControlBinding> Binding : Entity->Bindings)
+			for (auto It = Entity->Bindings.CreateIterator(); It; ++It)
 			{
-				if (ModifiedBindings.Contains(Binding.Get()))
+				if (ModifiedBindings.Contains(It->Get()))
 				{
 					PerFrameUpdatedEntities.Add(Entity->GetId());
+					It.RemoveCurrent();
 					break;
 				}
 			}
diff --git a/Engine/Plugins/VirtualProduction/RemoteControl/Source/RemoteControl/Public/RemoteControlField.h b/Engine/Plugins/VirtualProduction/RemoteControl/Source/RemoteControl/Public/RemoteControlField.h
--- a/Engine/Plugins/VirtualProduction/RemoteControl/Source/RemoteControl/Public/RemoteControlField.h
+++ b/Engine/Plugins/VirtualProduction/RemoteControl/Source/RemoteControl/Public/RemoteControlField.h
@@ -216,6 +216,9 @@ private:
 	UPROPERTY()
 	FSoftObjectPath FunctionPath;
 
+	/** Cached resolved underlying function used to avoid doing a findobject while serializing. */
+	mutable TWeakObjectPtr<UFunction> CachedFunction;
+
 private:
 	/** Parse function metadata to get the function`s default parameters */
 	void AssignDefaultFunctionArguments();
@@ -236,5 +239,5 @@ template<> struct TStructOpsTypeTraits<FRemoteControlProperty> : public TStructO
 	{
 		WithSerializer = true,
 		WithPostSerialize = true
-    };
+	};
 };
diff --git a/Engine/Plugins/VirtualProduction/RemoteControl/Source/RemoteControl/Public/RemoteControlPreset.h b/Engine/Plugins/VirtualProduction/RemoteControl/Source/RemoteControl/Public/RemoteControlPreset.h
--- a/Engine/Plugins/VirtualProduction/RemoteControl/Source/RemoteControl/Public/RemoteControlPreset.h
+++ b/Engine/Plugins/VirtualProduction/RemoteControl/Source/RemoteControl/Public/RemoteControlPreset.h
@@ -546,7 +546,7 @@ public:
 	 * @param ExposedEntityId The id of the entity to get.
 	 * @note ExposableEntityType must derive from FRemoteControlEntity.
 	 */
-	template <typename ExposableEntityType>
+	template <typename ExposableEntityType = FRemoteControlEntity>
 	TWeakPtr<const ExposableEntityType> GetExposedEntity(const FGuid& ExposedEntityId) const
 	{
 		static_assert(TIsDerivedFrom<ExposableEntityType, FRemoteControlEntity>::Value, "ExposableEntityType must derive from FRemoteControlEntity.");
@@ -558,7 +558,7 @@ public:
 	 * @param ExposedEntityId The id of the entity to get.
 	 * @note ExposableEntityType must derive from FRemoteControlEntity.
 	 */
-	template <typename ExposableEntityType>
+	template <typename ExposableEntityType = FRemoteControlEntity>
 	TWeakPtr<ExposableEntityType> GetExposedEntity(const FGuid& ExposedEntityId)
 	{
 		static_assert(TIsDerivedFrom<ExposableEntityType, FRemoteControlEntity>::Value, "ExposableEntityType must derive from FRemoteControlEntity.");
@@ -875,5 +875,6 @@ private:
 
 	friend FRemoteControlTarget;
 	friend FRemoteControlPresetLayout;
+	friend FRemoteControlEntity;
 };
 
diff --git a/Engine/Plugins/VirtualProduction/RemoteControl/Source/RemoteControlProtocolWidgets/Private/ViewModels/ProtocolEntityViewModel.cpp b/Engine/Plugins/VirtualProduction/RemoteControl/Source/RemoteControlProtocolWidgets/Private/ViewModels/ProtocolEntityViewModel.cpp
--- a/Engine/Plugins/VirtualProduction/RemoteControl/Source/RemoteControlProtocolWidgets/Private/ViewModels/ProtocolEntityViewModel.cpp
+++ b/Engine/Plugins/VirtualProduction/RemoteControl/Source/RemoteControlProtocolWidgets/Private/ViewModels/ProtocolEntityViewModel.cpp
@@ -31,7 +31,7 @@ namespace Commands
 		const TSharedPtr<IRemoteControlProtocol> Protocol = IRemoteControlProtocolModule::Get().GetProtocolByName(InArgs.ProtocolName);
 		if (Protocol.IsValid())
 		{
-			if(TSharedPtr<FRemoteControlProperty> RCProperty = InPreset->GetExposedEntity<FRemoteControlProperty>(InArgs.EntityId).Pin())
+			if (TSharedPtr<FRemoteControlProperty> RCProperty = InPreset->GetExposedEntity<FRemoteControlProperty>(InArgs.EntityId).Pin())
 			{
 				const FRemoteControlProtocolEntityPtr RemoteControlProtocolEntityPtr = Protocol->CreateNewProtocolEntity(RCProperty->GetProperty(), InPreset, InArgs.EntityId);
 
@@ -48,11 +48,11 @@ namespace Commands
 
 	static FName RemoveProtocolInternal(URemoteControlPreset* InPreset, const FAddRemoveProtocolArgs& InArgs)
 	{
-		if(TSharedPtr<FRemoteControlProperty> RCProperty = InPreset->GetExposedEntity<FRemoteControlProperty>(InArgs.EntityId).Pin())
+		if (TSharedPtr<FRemoteControlProperty> RCProperty = InPreset->GetExposedEntity<FRemoteControlProperty>(InArgs.EntityId).Pin())
 		{
 			TSet<FRemoteControlProtocolBinding>& Test = RCProperty->ProtocolBinding;
 			FRemoteControlProtocolBinding* FoundElement = Test.FindByHash(GetTypeHash(InArgs.BindingId), InArgs.BindingId);
-			if(FoundElement)
+			if (FoundElement)
 			{
 				Test.RemoveByHash(GetTypeHash(InArgs.BindingId), InArgs.BindingId);
 				return FoundElement->GetProtocolName();
@@ -89,14 +89,14 @@ FProtocolEntityViewModel::~FProtocolEntityViewModel()
 void FProtocolEntityViewModel::Initialize()
 {
 	Bindings.Empty(Bindings.Num());
-	if(TSharedPtr<FRemoteControlProperty> RCProperty = Preset->GetExposedEntity<FRemoteControlProperty>(PropertyId).Pin())
+	if (TSharedPtr<FRemoteControlProperty> RCProperty = Preset->GetExposedEntity<FRemoteControlProperty>(PropertyId).Pin())
 	{
 		Property = RCProperty->GetProperty();
 		for(FRemoteControlProtocolBinding& Binding : RCProperty->ProtocolBinding)
 		{
 			const TSharedPtr<IRemoteControlProtocol> Protocol = IRemoteControlProtocolModule::Get().GetProtocolByName(Binding.GetProtocolName());
 			// Supporting plugin needs to be loaded/protocol available
-			if(Protocol.IsValid())
+			if (Protocol.IsValid())
 			{
 				Bindings.Emplace(FProtocolBindingViewModel::Create(SharedThis(this), MakeShared<FRemoteControlProtocolBinding>(Binding)));
 				Protocol->Bind(Binding.GetRemoteControlProtocolEntityPtr());
@@ -109,28 +109,28 @@ void FProtocolEntityViewModel::Initialize()
 bool FProtocolEntityViewModel::CanAddBinding(const FName& InProtocolName, FText& OutMessage) const
 {
 	// If no protocols registered
-	if(InProtocolName == NAME_None)
+	if (InProtocolName == NAME_None || !GetProperty().IsValid())
 	{
 		return false;
 	}
 
 	FFieldClass* PropertyClass = GetProperty()->GetClass();
-	if(PropertyClass == FBoolProperty::StaticClass())
+	if (PropertyClass == FBoolProperty::StaticClass())
 	{
 		return FRemoteControlProtocolBinding::IsRangeTypeSupported<bool>();
 	}
 	
-	if(PropertyClass->IsChildOf(FNumericProperty::StaticClass()))
+	if (PropertyClass->IsChildOf(FNumericProperty::StaticClass()))
 	{
 		return true;
 	}
 	
-	if(PropertyClass->IsChildOf(FStructProperty::StaticClass()))
+	if (PropertyClass->IsChildOf(FStructProperty::StaticClass()))
 	{
 		return true;
 	}
 
-	if(PropertyClass->IsChildOf(FArrayProperty::StaticClass())
+	if (PropertyClass->IsChildOf(FArrayProperty::StaticClass())
 		|| PropertyClass->IsChildOf(FSetProperty::StaticClass())
 		|| PropertyClass->IsChildOf(FMapProperty::StaticClass()))
 	{
@@ -149,11 +149,11 @@ TSharedPtr<FProtocolBindingViewModel> FProtocolEntityViewModel::AddBinding(const
 	FScopedTransaction Transaction(LOCTEXT("AddProtocolBinding", "Add Protocol Binding"));
 
 	const TSharedPtr<FRemoteControlProtocolBinding> ProtocolBinding = Commands::AddProtocolInternal(Preset.Get(), Commands::FAddRemoveProtocolArgs{GetId(), InProtocolName});
-	if(ProtocolBinding.IsValid())
+	if (ProtocolBinding.IsValid())
 	{
 		using FCommandChange = TRemoteControlProtocolCommandChange<Commands::FAddRemoveProtocolArgs>;
 		using FOnChange = FCommandChange::FOnUndoRedoDelegate;
-		if(GUndo)
+		if (GUndo)
 		{
 			GUndo->StoreUndo(Preset.Get(),
                 MakeUnique<FCommandChange>(
@@ -183,11 +183,11 @@ void FProtocolEntityViewModel::RemoveBinding(const FGuid& InBindingId)
 	FScopedTransaction Transaction(LOCTEXT("RemoveProtocolBinding", "Remove Protocol Binding"));
 
 	const FName RemovedProtocolName = Commands::RemoveProtocolInternal(Preset.Get(), Commands::FAddRemoveProtocolArgs{GetId(), NAME_None, InBindingId});
-	if(RemovedProtocolName != NAME_None)
+	if (RemovedProtocolName != NAME_None)
 	{
 		using FCommandChange = TRemoteControlProtocolCommandChange<Commands::FAddRemoveProtocolArgs>;
 		using FOnChange = FCommandChange::FOnUndoRedoDelegate;
-		if(GUndo)
+		if (GUndo)
 		{
 			GUndo->StoreUndo(Preset.Get(),
                 MakeUnique<FCommandChange>(
@@ -205,7 +205,7 @@ void FProtocolEntityViewModel::RemoveBinding(const FGuid& InBindingId)
         return InBinding->GetId() == InBindingId;
     });
 
-	if(NumItemsRemoved <= 0)
+	if (NumItemsRemoved <= 0)
 	{
 		return;
 	}
@@ -215,7 +215,7 @@ void FProtocolEntityViewModel::RemoveBinding(const FGuid& InBindingId)
 
 TArray<TSharedPtr<FProtocolBindingViewModel>> FProtocolEntityViewModel::GetFilteredBindings(const TSet<FName>& InHiddenProtocolTypeNames)
 {
-	if(InHiddenProtocolTypeNames.Num() == 0)
+	if (InHiddenProtocolTypeNames.Num() == 0)
 	{
 		return GetBindings();
 	}
@@ -224,7 +224,7 @@ TArray<TSharedPtr<FProtocolBindingViewModel>> FProtocolEntityViewModel::GetFilte
 	for(const TSharedPtr<FProtocolBindingViewModel>& Binding : Bindings)
 	{
 		// make sure this bindings protocol name is NOT in the Hidden list, then add it to FilteredBindings
-		if(Algo::NoneOf(InHiddenProtocolTypeNames, [&Binding](const FName& InHiddenTypeName)
+		if (Algo::NoneOf(InHiddenProtocolTypeNames, [&Binding](const FName& InHiddenTypeName)
 		{
 			return InHiddenTypeName == Binding->GetBinding()->GetProtocolName();
 		}))
diff --git a/Engine/Plugins/VirtualProduction/RemoteControl/Source/RemoteControlUI/Private/AssetEditor/RemoteControlPresetEditorToolkit.cpp b/Engine/Plugins/VirtualProduction/RemoteControl/Source/RemoteControlUI/Private/AssetEditor/RemoteControlPresetEditorToolkit.cpp
--- a/Engine/Plugins/VirtualProduction/RemoteControl/Source/RemoteControlUI/Private/AssetEditor/RemoteControlPresetEditorToolkit.cpp
+++ b/Engine/Plugins/VirtualProduction/RemoteControl/Source/RemoteControlUI/Private/AssetEditor/RemoteControlPresetEditorToolkit.cpp
@@ -6,6 +6,7 @@
 #include "Widgets/Docking/SDockTab.h"
 #include "Widgets/SWidget.h"
 #include "Framework/Docking/TabManager.h"
+#include "Subsystems/AssetEditorSubsystem.h"
 #include "UI/SRCPanelExposedEntitiesList.h"
 #include "UI/SRCPanelTreeNode.h"
 #include "UI/SRemoteControlPanel.h"
@@ -29,32 +30,7 @@ void FRemoteControlPresetEditorToolkit::InitRemoteControlPresetEditor(const EToo
 	Preset = InPreset;
 
 	PanelTab = FRemoteControlUIModule::Get().CreateRemoteControlPanel(InPreset);
-	
-	struct Local
-	{
-		static void OnRemoteControlPresetClosed(TSharedRef<SDockTab> DockTab, TWeakPtr<IAssetEditorInstance> InRemoteControlPreset)
-		{
-			TSharedPtr<IAssetEditorInstance> AssetEditorInstance = InRemoteControlPreset.Pin();
-
-			if (AssetEditorInstance.IsValid())
-			{
-				InRemoteControlPreset.Pin()->CloseWindow();
-			}
-		}
-	};
-	
-	// Create a new DockTab and add the RemoteControlPreset widget to it.
-	FLevelEditorModule& LevelEditorModule = FModuleManager::GetModuleChecked<FLevelEditorModule>("LevelEditor");
-	TSharedPtr<FTabManager> EditorTabManager = LevelEditorModule.GetLevelEditorTabManager();
 
-	// Required, will cause the previous toolkit to close bringing down the RemoteControlPreset and unsubscribing the
-	// tab spawner. Without this, the InitAssetEditor call below will trigger an ensure as the RemoteControlPreset
-	// tab ID will already be registered within EditorTabManager
-	if (EditorTabManager->FindExistingLiveTab(PanelTabId).IsValid())
-	{
-		EditorTabManager->TryInvokeTab(PanelTabId)->RequestCloseTab();
-	}
-	
 	const TSharedRef<FTabManager::FLayout> StandaloneDefaultLayout = FTabManager::NewLayout("Standalone_RemoteControlPresetEditor")
 		->AddArea
 		(
@@ -69,13 +45,19 @@ void FRemoteControlPresetEditorToolkit::InitRemoteControlPresetEditor(const EToo
 	constexpr bool bCreateDefaultStandaloneMenu = true;
 	constexpr bool bCreateDefaultToolbar = true;
 
-	FAssetEditorToolkit::InitAssetEditor(Mode, InitToolkitHost, RemoteControlPanelAppIdentifier, StandaloneDefaultLayout, bCreateDefaultStandaloneMenu, bCreateDefaultStandaloneMenu, InPreset);
-
-	if (TSharedPtr<SDockTab> Tab = EditorTabManager->TryInvokeTab(PanelTabId))
+	// Required, will cause the previous toolkit to close bringing down the RemoteControlPreset and unsubscribing the
+	// tab spawner. Without this, the InitAssetEditor call below will trigger an ensure as the RemoteControlPreset
+	// tab ID will already be registered within EditorTabManager
+	FLevelEditorModule& LevelEditorModule = FModuleManager::GetModuleChecked<FLevelEditorModule>("LevelEditor");
+	TSharedPtr<FTabManager> EditorTabManager = LevelEditorModule.GetLevelEditorTabManager();
+	if (EditorTabManager->FindExistingLiveTab(PanelTabId).IsValid())
 	{
-		Tab->SetContent(PanelTab.ToSharedRef());
-		Tab->SetOnTabClosed(SDockTab::FOnTabClosedCallback::CreateStatic(&Local::OnRemoteControlPresetClosed, TWeakPtr<IAssetEditorInstance>(SharedThis(this))));
+		EditorTabManager->TryInvokeTab(PanelTabId)->RequestCloseTab();
 	}
+	
+	FAssetEditorToolkit::InitAssetEditor(Mode, InitToolkitHost, RemoteControlPanelAppIdentifier, StandaloneDefaultLayout, bCreateDefaultStandaloneMenu, bCreateDefaultStandaloneMenu, InPreset);
+
+	InvokePanelTab();
 }
 
 void FRemoteControlPresetEditorToolkit::RegisterTabSpawners(const TSharedRef<class FTabManager>& InTabManager)
@@ -98,6 +80,12 @@ bool FRemoteControlPresetEditorToolkit::OnRequestClose()
 	return true;
 }
 
+void FRemoteControlPresetEditorToolkit::FocusWindow(UObject* ObjectToFocusOn)
+{
+	InvokePanelTab();
+	BringToolkitToFront();
+}
+
 FText FRemoteControlPresetEditorToolkit::GetBaseToolkitName() const
 {
 	return LOCTEXT("PanelToolkitName", "Remote Control Panel");
@@ -131,4 +119,30 @@ TSharedRef<SDockTab> FRemoteControlPresetEditorToolkit::HandleTabManagerSpawnPan
 		];
 }
 
+void FRemoteControlPresetEditorToolkit::InvokePanelTab()
+{
+	struct Local
+	{
+		static void OnRemoteControlPresetClosed(TSharedRef<SDockTab> DockTab, TWeakPtr<IAssetEditorInstance> InRemoteControlPreset)
+		{
+			TSharedPtr<IAssetEditorInstance> AssetEditorInstance = InRemoteControlPreset.Pin();
+
+			if (AssetEditorInstance.IsValid())
+			{
+				InRemoteControlPreset.Pin()->CloseWindow();
+			}
+		}
+	};
+
+	// Create a new DockTab and add the RemoteControlPreset widget to it.
+	FLevelEditorModule& LevelEditorModule = FModuleManager::GetModuleChecked<FLevelEditorModule>("LevelEditor");
+	TSharedPtr<FTabManager> EditorTabManager = LevelEditorModule.GetLevelEditorTabManager();
+
+	if (TSharedPtr<SDockTab> Tab = EditorTabManager->TryInvokeTab(PanelTabId))
+	{
+		Tab->SetContent(PanelTab.ToSharedRef());
+		Tab->SetOnTabClosed(SDockTab::FOnTabClosedCallback::CreateStatic(&Local::OnRemoteControlPresetClosed, TWeakPtr<IAssetEditorInstance>(SharedThis(this))));
+	}
+}
+
 #undef LOCTEXT_NAMESPACE /*RemoteControlPresetEditorToolkit*/
\ No newline at end of file
diff --git a/Engine/Plugins/VirtualProduction/RemoteControl/Source/RemoteControlUI/Private/AssetEditor/RemoteControlPresetEditorToolkit.h b/Engine/Plugins/VirtualProduction/RemoteControl/Source/RemoteControlUI/Private/AssetEditor/RemoteControlPresetEditorToolkit.h
--- a/Engine/Plugins/VirtualProduction/RemoteControl/Source/RemoteControlUI/Private/AssetEditor/RemoteControlPresetEditorToolkit.h
+++ b/Engine/Plugins/VirtualProduction/RemoteControl/Source/RemoteControlUI/Private/AssetEditor/RemoteControlPresetEditorToolkit.h
@@ -34,12 +34,20 @@ public:
 	virtual bool OnRequestClose() override;
 	//~ End IToolkit interface
 
+
+	/** Begin IAssetEditorInstance interface */
+	virtual void FocusWindow(UObject* ObjectToFocusOn = nullptr) override;
+	/** End IAssetEditorInstance interface */
+
 private:
 	/** Handle spawning the tab that holds the remote control panel tab. */
 	TSharedRef<SDockTab> HandleTabManagerSpawnPanelTab(const FSpawnTabArgs& Args);
 
-	//~ Handle assigning group selection in either tab.
+	/** Handle assigning group selection in either tab. */
 	void OnPanelSelectionChange(const TSharedPtr<SRCPanelTreeNode>& Node);
+
+	/** Handle invoking the remote control tab. */
+	void InvokePanelTab();
 private:
 	/** Holds the remote control panel tab id. */
 	static const FName PanelTabId;
diff --git a/Engine/Plugins/VirtualProduction/RemoteControl/Source/RemoteControlUI/Private/UI/SRemoteControlPanel.cpp b/Engine/Plugins/VirtualProduction/RemoteControl/Source/RemoteControlUI/Private/UI/SRemoteControlPanel.cpp
--- a/Engine/Plugins/VirtualProduction/RemoteControl/Source/RemoteControlUI/Private/UI/SRemoteControlPanel.cpp
+++ b/Engine/Plugins/VirtualProduction/RemoteControl/Source/RemoteControlUI/Private/UI/SRemoteControlPanel.cpp
@@ -48,6 +48,7 @@
 #include "Widgets/Layout/SBorder.h"
 #include "Widgets/Layout/SBox.h"
 #include "Widgets/Layout/SSplitter.h"
+#include "Widgets/Layout/SSeparator.h"
 #include "Widgets/Layout/SWidgetSwitcher.h"
 #include "Widgets/SBoxPanel.h"
 #include "Widgets/Text/STextBlock.h"
@@ -199,22 +200,32 @@ void SRemoteControlPanel::Construct(const FArguments& InArgs, URemoteControlPres
 		]
 		+ SVerticalBox::Slot()
 		[
-			SNew(SBorder)
-			.Padding(FMargin(0.f, 5.f, 0.f, 0.f))
-			.BorderImage(FEditorStyle::GetBrush("ToolPanel.DarkGroupBorder"))
+			SNew(SSplitter)
+			.Orientation(Orient_Vertical)
+			+ SSplitter::Slot()
+			.Value(.8f)
 			[
-				SNew(SSplitter)
-				.Orientation(Orient_Vertical)
-				+ SSplitter::Slot()
-				.Value(.8f)
+				SNew(SBorder)
+				.Padding(FMargin(0.f, 5.f, 0.f, 0.f))
+				.BorderImage(FEditorStyle::GetBrush("ToolPanel.DarkGroupBorder"))
 				[
-					SNew(SBorder)
-					.Padding(FMargin(0.f, 5.f, 0.f, 0.f))
-					.BorderImage(FEditorStyle::GetBrush("ToolPanel.DarkGroupBorder"))
+					SNew(SWidgetSwitcher)
+					.WidgetIndex_Lambda([this](){ return !bIsInEditMode ? 0 : 1; })
+					+ SWidgetSwitcher::Slot()
 					[
-						SNew(SWidgetSwitcher)
-						.WidgetIndex_Lambda([this](){ return !bIsInEditMode ? 0 : 1; })
-						+ SWidgetSwitcher::Slot()
+						// Exposed entities List
+						SNew(SBorder)
+						.BorderImage(FEditorStyle::GetBrush("ToolPanel.GroupBorder"))
+						[
+							EntityList.ToSharedRef()
+						]
+					]
+					+ SWidgetSwitcher::Slot()
+					[
+						SNew(SSplitter)
+						.Orientation(EOrientation::Orient_Vertical)
+						+ SSplitter::Slot()
+						.Value(0.7f)
 						[
 							// Exposed entities List
 							SNew(SBorder)
@@ -223,54 +234,42 @@ void SRemoteControlPanel::Construct(const FArguments& InArgs, URemoteControlPres
 								EntityList.ToSharedRef()
 							]
 						]
-						+ SWidgetSwitcher::Slot()
+						+ SSplitter::Slot()
+						.Value(0.3f)
 						[
 							SNew(SSplitter)
-							.Orientation(EOrientation::Orient_Vertical)
+							.Orientation(Orient_Vertical)
 							+ SSplitter::Slot()
-							.Value(0.7f)
+							.Value(0.4f)
 							[
-								// Exposed entities List
 								SNew(SBorder)
 								.BorderImage(FEditorStyle::GetBrush("ToolPanel.GroupBorder"))
 								[
-									EntityList.ToSharedRef()
+									CreateEntityDetailsView()
 								]
 							]
 							+ SSplitter::Slot()
-							.Value(0.3f)
+							.Value(0.6f)
 							[
 								SNew(SBorder)
 								.BorderImage(FEditorStyle::GetBrush("ToolPanel.GroupBorder"))
-								.Padding(5.f)
 								[
-									SNew(SSplitter)
-									.Orientation(Orient_Vertical)
-									+ SSplitter::Slot()
-									.Value(0.4f)
-									[
-										CreateEntityDetailsView()
-									]
-									+ SSplitter::Slot()
-									.Value(0.6f)
-									[
-										EntityProtocolDetails.ToSharedRef()
-									]
+									EntityProtocolDetails.ToSharedRef()
 								]
 							]
 						]
 					]
 				]
-				+ SSplitter::Slot()
-				.Value(.2f)
+			]
+			+ SSplitter::Slot()
+			.Value(.2f)
+			[
+				SNew(SBorder)
+				.BorderImage(FEditorStyle::GetBrush("Menu.Background"))
+				.Visibility_Lambda([](){ return FRemoteControlLogger::Get().IsEnabled() ? EVisibility::Visible : EVisibility::Collapsed; })
+				.Padding(2.f)
 				[
-					SNew(SBorder)
-					.BorderImage(FEditorStyle::GetBrush("Menu.Background"))
-					.Visibility_Lambda([](){ return FRemoteControlLogger::Get().IsEnabled() ? EVisibility::Visible : EVisibility::Collapsed; })
-					.Padding(2.f)
-					[
-						FRemoteControlLogger::Get().GetWidget()
-					]
+					FRemoteControlLogger::Get().GetWidget()
 				]
 			]
 		]
@@ -689,6 +688,7 @@ void SRemoteControlPanel::Refresh()
 	BlueprintPicker->Refresh();
 	ActorFunctionPicker->Refresh();
 	SubsystemFunctionPicker->Refresh();
+	EntityList->Refresh();
 }
 
 void SRemoteControlPanel::Unexpose(const TSharedPtr<IPropertyHandle>& Handle)
diff --git a/Engine/Plugins/VirtualProduction/RemoteControl/Source/WebRemoteControl/Private/RemoteControlModels.h b/Engine/Plugins/VirtualProduction/RemoteControl/Source/WebRemoteControl/Private/RemoteControlModels.h
--- a/Engine/Plugins/VirtualProduction/RemoteControl/Source/WebRemoteControl/Private/RemoteControlModels.h
+++ b/Engine/Plugins/VirtualProduction/RemoteControl/Source/WebRemoteControl/Private/RemoteControlModels.h
@@ -337,19 +337,19 @@ struct FRCPresetLayoutGroupDescription
 		: Name(Group.Name)
 	{
 		checkSlow(Preset);
-		for (const FGuid& FieldId : Group.GetFields())
+		for (const FGuid& EntityId : Group.GetFields())
 		{
-			AddExposedField(Preset, FieldId);
+			AddExposedField(Preset, EntityId);
 		}
 	}
 
-	FRCPresetLayoutGroupDescription(const URemoteControlPreset* Preset, const FRemoteControlPresetGroup& Group, const TArray<FName>& FieldLabels)
+	FRCPresetLayoutGroupDescription(const URemoteControlPreset* Preset, const FRemoteControlPresetGroup& Group, const TArray<FGuid>& EntityIds)
 		: Name(Group.Name)
 	{
 		checkSlow(Preset);
-		for (FName FieldLabel : FieldLabels)
+		for (const FGuid& Id : EntityIds)
 		{
-			AddExposedField(Preset, Preset->GetExposedEntityId(FieldLabel));
+			AddExposedField(Preset, Id);
 		}
 	}
 
@@ -368,23 +368,23 @@ public:
 
 private:
 	/** Add an exposed field to this group description. */
-	void AddExposedField(const URemoteControlPreset* Preset, const FGuid& FieldId)
+	void AddExposedField(const URemoteControlPreset* Preset, const FGuid& EntityId)
 	{
-		if (TSharedPtr<const FRemoteControlProperty> RCProperty = Preset->GetExposedEntity<FRemoteControlProperty>(FieldId).Pin())
+		if (TSharedPtr<const FRemoteControlProperty> RCProperty = Preset->GetExposedEntity<FRemoteControlProperty>(EntityId).Pin())
 		{
 			if (FProperty* Property = RCProperty->GetProperty())
 			{
 				ExposedProperties.Emplace(*RCProperty);	
 			}
 		}
-		else if (TSharedPtr<const FRemoteControlFunction> RCFunction = Preset->GetExposedEntity<FRemoteControlFunction>(FieldId).Pin())
+		else if (TSharedPtr<const FRemoteControlFunction> RCFunction = Preset->GetExposedEntity<FRemoteControlFunction>(EntityId).Pin())
 		{
 			if (RCFunction->GetFunction())
 			{
 				ExposedFunctions.Emplace(*RCFunction);
 			}
 		}
-		else if (TSharedPtr<const FRemoteControlActor> RCActor = Preset->GetExposedEntity<FRemoteControlActor>(FieldId).Pin())
+		else if (TSharedPtr<const FRemoteControlActor> RCActor = Preset->GetExposedEntity<FRemoteControlActor>(EntityId).Pin())
 		{
 			if (RCActor->GetActor())
 			{
@@ -441,7 +441,7 @@ USTRUCT()
 struct FRCShortPresetDescription
 {
 	GENERATED_BODY()
-
+ 
 	FRCShortPresetDescription() = default;
 
 	FRCShortPresetDescription(const FAssetData& PresetAsset)
diff --git a/Engine/Plugins/VirtualProduction/RemoteControl/Source/WebRemoteControl/Private/RemoteControlResponse.h b/Engine/Plugins/VirtualProduction/RemoteControl/Source/WebRemoteControl/Private/RemoteControlResponse.h
--- a/Engine/Plugins/VirtualProduction/RemoteControl/Source/WebRemoteControl/Private/RemoteControlResponse.h
+++ b/Engine/Plugins/VirtualProduction/RemoteControl/Source/WebRemoteControl/Private/RemoteControlResponse.h
@@ -253,11 +253,12 @@ struct FRCPresetFieldsRemovedEvent
 
 	FRCPresetFieldsRemovedEvent() = default;
 
-	FRCPresetFieldsRemovedEvent(FName InPresetName, TArray<FName> InRemovedFields)
+	FRCPresetFieldsRemovedEvent(FName InPresetName, TArray<FName> InRemovedFields, const TArray<FGuid>& InRemovedFieldIDs)
 		: Type(TEXT("PresetFieldsRemoved"))
 		, PresetName(InPresetName)
 		, RemovedFields(MoveTemp(InRemovedFields))
 	{
+		Algo::Transform(InRemovedFieldIDs, RemovedFieldIds, [](const FGuid& Id){ return Id.ToString(); });
 	}
 
 	UPROPERTY()
@@ -268,6 +269,9 @@ struct FRCPresetFieldsRemovedEvent
 
 	UPROPERTY()
 	TArray<FName> RemovedFields;
+
+	UPROPERTY()
+	TArray<FString> RemovedFieldIds;
 };
 
 USTRUCT()
diff --git a/Engine/Plugins/VirtualProduction/RemoteControl/Source/WebRemoteControl/Private/WebSocketMessageHandler.cpp b/Engine/Plugins/VirtualProduction/RemoteControl/Source/WebRemoteControl/Private/WebSocketMessageHandler.cpp
--- a/Engine/Plugins/VirtualProduction/RemoteControl/Source/WebRemoteControl/Private/WebSocketMessageHandler.cpp
+++ b/Engine/Plugins/VirtualProduction/RemoteControl/Source/WebRemoteControl/Private/WebSocketMessageHandler.cpp
@@ -83,8 +83,8 @@ void FWebSocketMessageHandler::HandleWebSocketPresetRegister(const FRemoteContro
 
 		//Register to any useful callback for the given preset
 		Preset->OnExposedPropertyChanged().AddRaw(this, &FWebSocketMessageHandler::OnPresetExposedPropertyChanged);
-		Preset->OnPropertyExposed().AddRaw(this, &FWebSocketMessageHandler::OnPropertyExposed);
-		Preset->OnPropertyUnexposed().AddRaw(this, &FWebSocketMessageHandler::OnPropertyUnexposed);
+		Preset->OnEntityExposed().AddRaw(this, &FWebSocketMessageHandler::OnPropertyExposed);
+		Preset->OnEntityUnexposed().AddRaw(this, &FWebSocketMessageHandler::OnPropertyUnexposed);
 		Preset->OnFieldRenamed().AddRaw(this, &FWebSocketMessageHandler::OnFieldRenamed);
 		Preset->OnMetadataModified().AddRaw(this, &FWebSocketMessageHandler::OnMetadataModified);
 		Preset->OnActorPropertyModified().AddRaw(this, &FWebSocketMessageHandler::OnActorPropertyChanged);
@@ -175,7 +175,7 @@ void FWebSocketMessageHandler::ProcessChangedActorProperties()
 	PerFrameActorPropertyChanged.Empty();
 }
 
-void FWebSocketMessageHandler::OnPropertyExposed(URemoteControlPreset* Owner, FName PropertyLabel)
+void FWebSocketMessageHandler::OnPropertyExposed(URemoteControlPreset* Owner, const FGuid& EntityId)
 {
 	if (Owner == nullptr)
 	{
@@ -188,7 +188,7 @@ void FWebSocketMessageHandler::OnPropertyExposed(URemoteControlPreset* Owner, FN
 	}
 
 	//Cache the property field that was removed for end of frame notification
-	PerFrameAddedProperties.FindOrAdd(Owner->GetFName()).AddUnique(PropertyLabel);
+	PerFrameAddedProperties.FindOrAdd(Owner->GetFName()).AddUnique(EntityId);
 }
 
 void FWebSocketMessageHandler::OnPresetExposedPropertyChanged(URemoteControlPreset* Owner, const FRemoteControlProperty& PropertyChanged)
@@ -218,7 +218,7 @@ void FWebSocketMessageHandler::OnPresetExposedPropertyChanged(URemoteControlPres
 	}
 }
 
-void FWebSocketMessageHandler::OnPropertyUnexposed(URemoteControlPreset* Owner, FName PropertyLabel)
+void FWebSocketMessageHandler::OnPropertyUnexposed(URemoteControlPreset* Owner, const FGuid& EntityId)
 {
 	if (Owner == nullptr)
 	{
@@ -231,7 +231,7 @@ void FWebSocketMessageHandler::OnPropertyUnexposed(URemoteControlPreset* Owner,
 	}
 
 	//Cache the property field that was removed for end of frame notification
-	PerFrameRemovedProperties.FindOrAdd(Owner->GetFName()).AddUnique(PropertyLabel);
+	PerFrameRemovedProperties.FindOrAdd(Owner->GetFName()).AddUnique(EntityId);
 }
 
 void FWebSocketMessageHandler::OnFieldRenamed(URemoteControlPreset* Owner, FName OldFieldLabel, FName NewFieldLabel)
@@ -351,7 +351,7 @@ void FWebSocketMessageHandler::OnEndFrame()
 
 void FWebSocketMessageHandler::ProcessAddedProperties()
 {
-	for (const TPair<FName, TArray<FName>>& Entry : PerFrameAddedProperties)
+	for (const TPair<FName, TArray<FGuid>>& Entry : PerFrameAddedProperties)
 	{
 		if (Entry.Value.Num() <= 0 || !ShouldProcessEventForPreset(Entry.Key))
 		{
@@ -367,18 +367,19 @@ void FWebSocketMessageHandler::ProcessAddedProperties()
 		FRCPresetDescription AddedPropertiesDescription;
 		AddedPropertiesDescription.Name = Preset->GetName();
 		AddedPropertiesDescription.Path = Preset->GetPathName();
+		AddedPropertiesDescription.Id = Preset->GetPresetId().ToString();
 
-		TMap<FRemoteControlPresetGroup*, TArray<FName>> GroupedNewFields;
+		TMap<FRemoteControlPresetGroup*, TArray<FGuid>> GroupedNewFields;
 
-		for (const FName& Label : Entry.Value)
+		for (const FGuid& Id : Entry.Value)
 		{
-			if (FRemoteControlPresetGroup* Group = Preset->Layout.FindGroupFromField(Preset->GetExposedEntityId(Label)))
+			if (FRemoteControlPresetGroup* Group = Preset->Layout.FindGroupFromField(Id))
 			{
-				GroupedNewFields.FindOrAdd(Group).Add(Label);
+				GroupedNewFields.FindOrAdd(Group).Add(Id);
 			}
 		}
 
-		for (const TTuple<FRemoteControlPresetGroup*, TArray<FName>>& Tuple : GroupedNewFields)
+		for (const TTuple<FRemoteControlPresetGroup*, TArray<FGuid>>& Tuple : GroupedNewFields)
 		{
 			AddedPropertiesDescription.Groups.Emplace(Preset, *Tuple.Key, Tuple.Value);
 		}
@@ -393,18 +394,33 @@ void FWebSocketMessageHandler::ProcessAddedProperties()
 
 void FWebSocketMessageHandler::ProcessRemovedProperties()
 {
-	for (const TPair<FName, TArray<FName>>& Entry : PerFrameRemovedProperties)
+	for (const TPair<FName, TArray<FGuid>>& Entry : PerFrameRemovedProperties)
 	{
 		if (Entry.Value.Num() <= 0 || !ShouldProcessEventForPreset(Entry.Key))
 		{
 			continue;
 		}
 
+		TArray<FName> Labels;
+		Labels.Reserve(Entry.Value.Num());
+		if (URemoteControlPreset* Preset = IRemoteControlModule::Get().ResolvePreset(Entry.Key))
+		{
+			for (const FGuid& Id : Entry.Value)
+			{
+				if (TSharedPtr<FRemoteControlEntity> Entity = Preset->GetExposedEntity(Id).Pin())
+				{
+					Labels.Add(Entity->GetLabel());
+				}
+			}
+		}
+
+		ensure(Labels.Num() == Entry.Value.Num());
+		
 		TArray<uint8> Payload;
-		WebRemoteControlUtils::SerializeResponse(FRCPresetFieldsRemovedEvent{ Entry.Key, Entry.Value }, Payload);
+		WebRemoteControlUtils::SerializeResponse(FRCPresetFieldsRemovedEvent{ Entry.Key, Labels, Entry.Value}, Payload);
 		BroadcastToListeners(Entry.Key, Payload);
 	}
-
+	
 	PerFrameRemovedProperties.Empty();
 }
 
@@ -487,6 +503,7 @@ bool FWebSocketMessageHandler::WritePropertyChangeEventPayload(URemoteControlPre
 				JsonWriter->WriteObjectStart();
 				{
 					JsonWriter->WriteValue(TEXT("PropertyLabel"), *Property.GetLabel().ToString());
+					JsonWriter->WriteValue(TEXT("Id"), *Property.GetId().ToString());
 
 					for (UObject* Object : Property.GetBoundObjects())
 					{
@@ -547,6 +564,8 @@ bool FWebSocketMessageHandler::WriteActorPropertyChangePayload(URemoteControlPre
 			}
 
 			FString RCActorName = Pair.Key.GetLabel().ToString();
+			
+			JsonWriter->WriteValue(TEXT("Id"), *Pair.Key.GetId().ToString());
 			JsonWriter->WriteValue(TEXT("DisplayName"), RCActorName);
 			JsonWriter->WriteValue(TEXT("Path"), Pair.Key.Path.ToString());
 
diff --git a/Engine/Plugins/VirtualProduction/RemoteControl/Source/WebRemoteControl/Private/WebSocketMessageHandler.h b/Engine/Plugins/VirtualProduction/RemoteControl/Source/WebRemoteControl/Private/WebSocketMessageHandler.h
--- a/Engine/Plugins/VirtualProduction/RemoteControl/Source/WebRemoteControl/Private/WebSocketMessageHandler.h
+++ b/Engine/Plugins/VirtualProduction/RemoteControl/Source/WebRemoteControl/Private/WebSocketMessageHandler.h
@@ -41,8 +41,8 @@ private:
 
 	//Preset callbacks
 	void OnPresetExposedPropertyChanged(URemoteControlPreset* Owner, const FRemoteControlProperty& PropertyChanged);
-	void OnPropertyExposed(URemoteControlPreset* Owner, FName PropertyLabel);
-	void OnPropertyUnexposed(URemoteControlPreset* Owner, FName PropertyLabel);
+	void OnPropertyExposed(URemoteControlPreset* Owner,  const FGuid& EntityId);
+	void OnPropertyUnexposed(URemoteControlPreset* Owner, const FGuid& EntityId);
 	void OnFieldRenamed(URemoteControlPreset* Owner, FName OldFieldLabel, FName NewFieldLabel);
 	void OnMetadataModified(URemoteControlPreset* Owner);
 	void OnActorPropertyChanged(URemoteControlPreset* Owner, FRemoteControlActor& Actor, UObject* ModifiedObject, FProperty* ModifiedProperty);
@@ -119,10 +119,10 @@ private:
 	TMap<FName, TMap<FGuid, TMap<FRemoteControlActor, TArray<FRCObjectReference>>>> PerFrameActorPropertyChanged;
 
 	/** Properties that were exposed for a frame, per preset */
-	TMap<FName, TArray<FName>> PerFrameAddedProperties;
+	TMap<FName, TArray<FGuid>> PerFrameAddedProperties;
 
 	/** Properties that were unexposed for a frame, per preset */
-	TMap<FName, TArray<FName>> PerFrameRemovedProperties;
+	TMap<FName, TArray<FGuid>> PerFrameRemovedProperties;
 
 	/** Fields that were renamed for a frame, per preset */
 	TMap<FName, TArray<TTuple<FName, FName>>> PerFrameRenamedFields;
