Fixed PVS static analysis ternary operator warnings.

diff --git a/Engine/Plugins/VirtualProduction/RemoteControl/Source/RemoteControlCommon/Private/RCPropertyContainer.cpp b/Engine/Plugins/VirtualProduction/RemoteControl/Source/RemoteControlCommon/Private/RCPropertyContainer.cpp
--- a/Engine/Plugins/VirtualProduction/RemoteControl/Source/RemoteControlCommon/Private/RCPropertyContainer.cpp
+++ b/Engine/Plugins/VirtualProduction/RemoteControl/Source/RemoteControlCommon/Private/RCPropertyContainer.cpp
@@ -295,7 +295,7 @@ namespace PropertyContainers
 		FFieldVariant PropertyScope = InParent;
 
 		// A name is usually required further down the chain, so set one if not already.
-		const FName PropertyName = InPropertyName == NAME_None ? TEXT("Value") : InPropertyName;
+		const FName PropertyName = (InPropertyName == NAME_None) ? TEXT("Value") : InPropertyName;
 		
 		FProperty* NewProperty = nullptr;
 
