Fix function picker not getting focused upon being open in the context menu.
Fix Old presets not having a valid ID
Fix not being able to rebind exposed actor
Fix undo not working when unexposing property
Fix NSLOCTEXT not being serialized correctly
Fix FieldsRemovedEvent not working because entity IDs no longer resolve on the frame after they are unexposed.
Fix this by keeping both the ID and the labels for the unexposed entities to return to the web client.

diff --git a/Engine/Plugins/VirtualProduction/RemoteControl/Source/RemoteControl/Private/RemoteControlPreset.cpp b/Engine/Plugins/VirtualProduction/RemoteControl/Source/RemoteControl/Private/RemoteControlPreset.cpp
--- a/Engine/Plugins/VirtualProduction/RemoteControl/Source/RemoteControl/Private/RemoteControlPreset.cpp
+++ b/Engine/Plugins/VirtualProduction/RemoteControl/Source/RemoteControl/Private/RemoteControlPreset.cpp
@@ -782,6 +782,8 @@ TWeakPtr<FRemoteControlFunction> URemoteControlPreset::ExposeFunction(UObject* O
 
 TSharedPtr<FRemoteControlEntity> URemoteControlPreset::Expose(FRemoteControlEntity&& Entity, UScriptStruct* EntityType, const FGuid& GroupId)
 {
+	Registry->Modify();
+	
 	TSharedPtr<FRemoteControlEntity> RCEntity = Registry->AddExposedEntity(MoveTemp(Entity), EntityType);
 	InitializeEntityMetadata(RCEntity);
 	
@@ -981,6 +983,7 @@ void URemoteControlPreset::Unexpose(const FGuid& EntityId)
 	if (EntityId.IsValid() && Registry->GetExposedEntity(EntityId).IsValid())
 	{
 		OnEntityUnexposedDelegate.Broadcast(this, EntityId);
+		Registry->Modify();
 		Registry->RemoveExposedEntity(EntityId);
 		FRCCachedFieldData CachedData = FieldCache.FindChecked(EntityId);
 		Layout.RemoveField(CachedData.LayoutGroupId, EntityId);
@@ -1090,6 +1093,11 @@ void URemoteControlPreset::Serialize(FArchive& Ar)
 		{
 			ConvertTargetsToBindings();
 		}
+
+		if (!PresetId.IsValid())
+		{
+			PresetId = FGuid::NewGuid();
+		}
 	}
 }
 
diff --git a/Engine/Plugins/VirtualProduction/RemoteControl/Source/RemoteControl/Public/RemoteControlActor.h b/Engine/Plugins/VirtualProduction/RemoteControl/Source/RemoteControl/Public/RemoteControlActor.h
--- a/Engine/Plugins/VirtualProduction/RemoteControl/Source/RemoteControl/Public/RemoteControlActor.h
+++ b/Engine/Plugins/VirtualProduction/RemoteControl/Source/RemoteControl/Public/RemoteControlActor.h
@@ -37,7 +37,7 @@ struct REMOTECONTROL_API FRemoteControlActor : public FRemoteControlEntity
 	{
 		if (InObjectToBind && InObjectToBind->IsA<AActor>())
 		{
-			SetActor(CastChecked<AActor>(InObjectToBind));
+			FRemoteControlEntity::BindObject(InObjectToBind);
 		}
 	}
 
@@ -60,15 +60,7 @@ struct REMOTECONTROL_API FRemoteControlActor : public FRemoteControlEntity
 
 	void SetActor(AActor* InActor)
 	{
-		if (Owner.IsValid())
-		{
-			Owner->Modify();
-		}
-
-		if (ensure(Bindings.Num()) && Bindings[0].IsValid())
-		{
-			Bindings[0]->SetBoundObject(InActor);
-		}
+		BindObject(InActor);
 	}
 
 public:
diff --git a/Engine/Plugins/VirtualProduction/RemoteControl/Source/RemoteControlUI/Private/UI/SRCPanelFunctionPicker.cpp b/Engine/Plugins/VirtualProduction/RemoteControl/Source/RemoteControlUI/Private/UI/SRCPanelFunctionPicker.cpp
--- a/Engine/Plugins/VirtualProduction/RemoteControl/Source/RemoteControlUI/Private/UI/SRCPanelFunctionPicker.cpp
+++ b/Engine/Plugins/VirtualProduction/RemoteControl/Source/RemoteControlUI/Private/UI/SRCPanelFunctionPicker.cpp
@@ -234,6 +234,20 @@ void SRCPanelFunctionPicker::Construct(const FArguments& InArgs)
 	];
 }
 
+void SRCPanelFunctionPicker::Tick(const FGeometry& AllottedGeometry, const double InCurrentTime, const float InDeltaTime)
+{
+	SWidget::Tick(AllottedGeometry, InCurrentTime, InDeltaTime);
+
+	constexpr double MinimumTimeToSwitchFocus = 0.2;
+	if (LastTimeSinceTick != 0.0 && InCurrentTime - LastTimeSinceTick > MinimumTimeToSwitchFocus)
+	{
+		ObjectsTreeView->ClearSearchBox();
+		ObjectsTreeView->Focus();
+	}
+
+	LastTimeSinceTick = InCurrentTime;
+}
+
 FReply SRCPanelFunctionPicker::OnFocusReceived(const FGeometry& MyGeometry, const FFocusEvent& InFocusEvent)
 {
 	if (InFocusEvent.GetCause() == EFocusCause::Navigation)
diff --git a/Engine/Plugins/VirtualProduction/RemoteControl/Source/RemoteControlUI/Private/UI/SRCPanelFunctionPicker.h b/Engine/Plugins/VirtualProduction/RemoteControl/Source/RemoteControlUI/Private/UI/SRCPanelFunctionPicker.h
--- a/Engine/Plugins/VirtualProduction/RemoteControl/Source/RemoteControlUI/Private/UI/SRCPanelFunctionPicker.h
+++ b/Engine/Plugins/VirtualProduction/RemoteControl/Source/RemoteControlUI/Private/UI/SRCPanelFunctionPicker.h
@@ -25,8 +25,11 @@ public:
 
 	void Construct(const FArguments& InArgs);
 
-	virtual FReply OnFocusReceived(const FGeometry& MyGeometry, const FFocusEvent& InFocusEvent) override;
+	//~ Begin SWidget interface
+	virtual void Tick(const FGeometry& AllottedGeometry, const double InCurrentTime, const float InDeltaTime);
+	//~ End SWidget interface
 
+	virtual FReply OnFocusReceived(const FGeometry& MyGeometry, const FFocusEvent& InFocusEvent) override;
 	virtual bool SupportsKeyboardFocus() const override
 	{
 		return true;
@@ -48,4 +51,6 @@ private:
 	FText Label;
 	/** Allow default objects when refreshing. */
 	bool bAllowDefaultObjects = false;
+	/** Keep track of the last time this widget tick in order to focus it if it's been more than a frame. */
+	double LastTimeSinceTick = 0.0;
 };
\ No newline at end of file
diff --git a/Engine/Plugins/VirtualProduction/RemoteControl/Source/WebRemoteControl/Private/Serialization/RCJsonStructSerializerBackend.h b/Engine/Plugins/VirtualProduction/RemoteControl/Source/WebRemoteControl/Private/Serialization/RCJsonStructSerializerBackend.h
--- a/Engine/Plugins/VirtualProduction/RemoteControl/Source/WebRemoteControl/Private/Serialization/RCJsonStructSerializerBackend.h
+++ b/Engine/Plugins/VirtualProduction/RemoteControl/Source/WebRemoteControl/Private/Serialization/RCJsonStructSerializerBackend.h
@@ -13,14 +13,16 @@ class FRCJsonStructSerializerBackend
 	: public FJsonStructSerializerBackend
 {
 public:
+	/** Default flags for serializing structures in Web RC. */
+	static const EStructSerializerBackendFlags DefaultSerializerFlags = EStructSerializerBackendFlags::WriteByteArrayAsByteStream;
 
 	/**
 	 * Creates and initializes a new instance with the given flags.
 	 *
 	 * @param InArchive The archive to serialize into.
-	 * @param InFlags The flags that control the serialization behavior (typically EStructSerializerBackendFlags::Default).
+	 * @param InFlags The flags that control the serialization behavior (typically FRCJsonStructSerializerBackend::DefaultSerializerFlags).
 	 */
-	FRCJsonStructSerializerBackend(FArchive& InArchive, const EStructSerializerBackendFlags InFlags)
+	FRCJsonStructSerializerBackend(FArchive& InArchive, const EStructSerializerBackendFlags InFlags = DefaultSerializerFlags)
 		: FJsonStructSerializerBackend(InArchive, InFlags)
 	{ }
 
diff --git a/Engine/Plugins/VirtualProduction/RemoteControl/Source/WebRemoteControl/Private/WebRemoteControl.cpp b/Engine/Plugins/VirtualProduction/RemoteControl/Source/WebRemoteControl/Private/WebRemoteControl.cpp
--- a/Engine/Plugins/VirtualProduction/RemoteControl/Source/WebRemoteControl/Private/WebRemoteControl.cpp
+++ b/Engine/Plugins/VirtualProduction/RemoteControl/Source/WebRemoteControl/Private/WebRemoteControl.cpp
@@ -627,7 +627,7 @@ bool FWebRemoteControlModule::HandleObjectPropertyRoute(const FHttpServerRequest
 	{
 		TArray<uint8> WorkingBuffer;
 		FMemoryWriter Writer(WorkingBuffer);
-		FRCJsonStructSerializerBackend SerializerBackend(Writer, EStructSerializerBackendFlags::Default);
+		FRCJsonStructSerializerBackend SerializerBackend(Writer);
 		if (IRemoteControlModule::Get().GetObjectProperties(ObjectRef, SerializerBackend))
 		{
 			Response->Code = EHttpServerResponseCodes::Ok;
@@ -711,7 +711,7 @@ bool FWebRemoteControlModule::HandlePresetCallFunctionRoute(const FHttpServerReq
 	TArray<uint8> OutputBuffer;
 	FMemoryWriter Writer{ OutputBuffer };
 	TSharedPtr<TJsonWriter<UCS2CHAR>> JsonWriter = TJsonWriter<UCS2CHAR>::Create(&Writer);
-	FRCJsonStructSerializerBackend WriterBackend{ Writer, EStructSerializerBackendFlags::Default };
+	FRCJsonStructSerializerBackend WriterBackend{ Writer };
 
 	JsonWriter->WriteObjectStart();
 	JsonWriter->WriteIdentifierPrefix("ReturnedValues");
@@ -917,7 +917,7 @@ bool FWebRemoteControlModule::HandlePresetGetPropertyRoute(const FHttpServerRequ
 	{
 		{
 			JsonWriter->WriteIdentifierPrefix(TEXT("ExposedPropertyDescription"));
-			FRCJsonStructSerializerBackend SerializeBackend{Writer, EStructSerializerBackendFlags::Default};
+			FRCJsonStructSerializerBackend SerializeBackend{Writer};
 			FStructSerializer::Serialize(FRCExposedPropertyDescription{*RemoteControlProperty}, SerializeBackend, FStructSerializerPolicies());
 		}
 		
@@ -997,7 +997,7 @@ bool FWebRemoteControlModule::HandlePresetGetExposedActorPropertyRoute(const FHt
 
 			bSuccess &= IRemoteControlModule::Get().ResolveObjectProperty(ERCAccess::READ_ACCESS, Actor, FieldPath, ObjectRef);
 
-			FRCJsonStructSerializerBackend SerializerBackend(Writer, EStructSerializerBackendFlags::Default);
+			FRCJsonStructSerializerBackend SerializerBackend(Writer);
 
 			JsonWriter->WriteObjectStart();
 			JsonWriter->WriteIdentifierPrefix(TEXT("PropertyValue"));
@@ -1051,7 +1051,7 @@ bool FWebRemoteControlModule::HandlePresetGetExposedActorPropertiesRoute(const F
 		{
 			TArray<uint8> WorkingBuffer;
 			FMemoryWriter Writer(WorkingBuffer);
-			FRCJsonStructSerializerBackend Backend{Writer, EStructSerializerBackendFlags::Default};
+			FRCJsonStructSerializerBackend Backend{Writer};
 
 			FRCObjectReference Ref{ ERCAccess::READ_ACCESS, Actor};
 			if (IRemoteControlModule::Get().GetObjectProperties(Ref, Backend))
diff --git a/Engine/Plugins/VirtualProduction/RemoteControl/Source/WebRemoteControl/Private/WebRemoteControlUtils.cpp b/Engine/Plugins/VirtualProduction/RemoteControl/Source/WebRemoteControl/Private/WebRemoteControlUtils.cpp
--- a/Engine/Plugins/VirtualProduction/RemoteControl/Source/WebRemoteControl/Private/WebRemoteControlUtils.cpp
+++ b/Engine/Plugins/VirtualProduction/RemoteControl/Source/WebRemoteControl/Private/WebRemoteControlUtils.cpp
@@ -91,7 +91,7 @@ namespace RemotePayloadSerializer
 
 	void SerializeWrappedCallResponse(int32 RequestId, TUniquePtr<FHttpServerResponse> Response, FMemoryWriter& Writer)
 	{
-		FRCJsonStructSerializerBackend Backend(Writer, EStructSerializerBackendFlags::Default);
+		FRCJsonStructSerializerBackend Backend(Writer, FRCJsonStructSerializerBackend::DefaultSerializerFlags);
 		TSharedPtr<TJsonWriter<ANSICHAR>> JsonWriter = TJsonWriter<ANSICHAR>::Create(&Writer);
 		TArray<FString>* ContentTypeHeaders = Response->Headers.Find(TEXT("Content-Type"));
 		const bool bIsBinaryData = ContentTypeHeaders && ContentTypeHeaders->Contains(TEXT("image/png"));
@@ -127,7 +127,7 @@ namespace RemotePayloadSerializer
 	{
 		TArray<uint8> WorkingBuffer;
 		FMemoryWriter TemporaryBufferWriter(WorkingBuffer);
-		FRCJsonStructSerializerBackend TemporaryBackend(TemporaryBufferWriter, EStructSerializerBackendFlags::Default);
+		FRCJsonStructSerializerBackend TemporaryBackend(TemporaryBufferWriter, FRCJsonStructSerializerBackend::DefaultSerializerFlags);
 
 		int32 ColonPosition = -1;
 		int32 LastBracketPosition = -1;
@@ -221,7 +221,7 @@ namespace RemotePayloadSerializer
 		}
 
 		// write the param struct
-		FRCJsonStructSerializerBackend Backend(Writer, EStructSerializerBackendFlags::Default);
+		FRCJsonStructSerializerBackend Backend(Writer, FRCJsonStructSerializerBackend::DefaultSerializerFlags);
 		FStructSerializerPolicies Policies;
 
 		if (bOnlyReturn)
diff --git a/Engine/Plugins/VirtualProduction/RemoteControl/Source/WebRemoteControl/Private/WebRemoteControlUtils.h b/Engine/Plugins/VirtualProduction/RemoteControl/Source/WebRemoteControl/Private/WebRemoteControlUtils.h
--- a/Engine/Plugins/VirtualProduction/RemoteControl/Source/WebRemoteControl/Private/WebRemoteControlUtils.h
+++ b/Engine/Plugins/VirtualProduction/RemoteControl/Source/WebRemoteControl/Private/WebRemoteControlUtils.h
@@ -3,6 +3,7 @@
 #pragma once
 
 #include "CoreMinimal.h"
+#include "IStructSerializerBackend.h"
 #include "Serialization/RCJsonStructSerializerBackend.h"
 #include "Serialization/RCJsonStructDeserializerBackend.h"
 #include "HttpServerResponse.h"
@@ -282,7 +283,7 @@ namespace WebRemoteControlUtils
 	{
 		TArray<uint8> WorkingBuffer;
 		FMemoryWriter Writer(WorkingBuffer);
-		FRCJsonStructSerializerBackend SerializeBackend{Writer, EStructSerializerBackendFlags::Default};
+		FRCJsonStructSerializerBackend SerializeBackend{Writer, FRCJsonStructSerializerBackend::DefaultSerializerFlags};
 		FStructSerializer::Serialize(InResponseObject, SerializeBackend, FStructSerializerPolicies());
 		ConvertToUTF8(WorkingBuffer, OutResponsePayload);
 	}
diff --git a/Engine/Plugins/VirtualProduction/RemoteControl/Source/WebRemoteControl/Private/WebSocketMessageHandler.cpp b/Engine/Plugins/VirtualProduction/RemoteControl/Source/WebRemoteControl/Private/WebSocketMessageHandler.cpp
--- a/Engine/Plugins/VirtualProduction/RemoteControl/Source/WebRemoteControl/Private/WebSocketMessageHandler.cpp
+++ b/Engine/Plugins/VirtualProduction/RemoteControl/Source/WebRemoteControl/Private/WebSocketMessageHandler.cpp
@@ -88,7 +88,6 @@ void FWebSocketMessageHandler::HandleWebSocketPresetRegister(const FRemoteContro
 		Preset->OnFieldRenamed().AddRaw(this, &FWebSocketMessageHandler::OnFieldRenamed);
 		Preset->OnMetadataModified().AddRaw(this, &FWebSocketMessageHandler::OnMetadataModified);
 		Preset->OnActorPropertyModified().AddRaw(this, &FWebSocketMessageHandler::OnActorPropertyChanged);
-
 	}
 
 	ClientIds->AddUnique(WebSocketMessage.ClientId);
@@ -230,8 +229,13 @@ void FWebSocketMessageHandler::OnPropertyUnexposed(URemoteControlPreset* Owner,
 		return;
 	}
 
+	TSharedPtr<FRemoteControlEntity> Entity = Owner->GetExposedEntity(EntityId).Pin();
+	check(Entity);
+
 	//Cache the property field that was removed for end of frame notification
-	PerFrameRemovedProperties.FindOrAdd(Owner->GetFName()).AddUnique(EntityId);
+	TTuple<TArray<FGuid>, TArray<FName>>& Entries = PerFrameRemovedProperties.FindOrAdd(Owner->GetFName());
+	Entries.Key.AddUnique(EntityId);
+	Entries.Value.AddUnique(Entity->GetLabel());
 }
 
 void FWebSocketMessageHandler::OnFieldRenamed(URemoteControlPreset* Owner, FName OldFieldLabel, FName NewFieldLabel)
@@ -394,30 +398,17 @@ void FWebSocketMessageHandler::ProcessAddedProperties()
 
 void FWebSocketMessageHandler::ProcessRemovedProperties()
 {
-	for (const TPair<FName, TArray<FGuid>>& Entry : PerFrameRemovedProperties)
+	for (const TPair<FName, TTuple<TArray<FGuid>, TArray<FName>>>& Entry : PerFrameRemovedProperties)
 	{
-		if (Entry.Value.Num() <= 0 || !ShouldProcessEventForPreset(Entry.Key))
+		if (Entry.Value.Key.Num() <= 0 || !ShouldProcessEventForPreset(Entry.Key))
 		{
 			continue;
 		}
 
-		TArray<FName> Labels;
-		Labels.Reserve(Entry.Value.Num());
-		if (URemoteControlPreset* Preset = IRemoteControlModule::Get().ResolvePreset(Entry.Key))
-		{
-			for (const FGuid& Id : Entry.Value)
-			{
-				if (TSharedPtr<FRemoteControlEntity> Entity = Preset->GetExposedEntity(Id).Pin())
-				{
-					Labels.Add(Entity->GetLabel());
-				}
-			}
-		}
-
-		ensure(Labels.Num() == Entry.Value.Num());
+		ensure(Entry.Value.Key.Num() == Entry.Value.Value.Num());
 		
 		TArray<uint8> Payload;
-		WebRemoteControlUtils::SerializeResponse(FRCPresetFieldsRemovedEvent{ Entry.Key, Labels, Entry.Value}, Payload);
+		WebRemoteControlUtils::SerializeResponse(FRCPresetFieldsRemovedEvent{ Entry.Key, Entry.Value.Value, Entry.Value.Key }, Payload);
 		BroadcastToListeners(Entry.Key, Payload);
 	}
 	
diff --git a/Engine/Plugins/VirtualProduction/RemoteControl/Source/WebRemoteControl/Private/WebSocketMessageHandler.h b/Engine/Plugins/VirtualProduction/RemoteControl/Source/WebRemoteControl/Private/WebSocketMessageHandler.h
--- a/Engine/Plugins/VirtualProduction/RemoteControl/Source/WebRemoteControl/Private/WebSocketMessageHandler.h
+++ b/Engine/Plugins/VirtualProduction/RemoteControl/Source/WebRemoteControl/Private/WebSocketMessageHandler.h
@@ -122,7 +122,7 @@ private:
 	TMap<FName, TArray<FGuid>> PerFrameAddedProperties;
 
 	/** Properties that were unexposed for a frame, per preset */
-	TMap<FName, TArray<FGuid>> PerFrameRemovedProperties;
+	TMap<FName, TTuple<TArray<FGuid>, TArray<FName>>> PerFrameRemovedProperties;
 
 	/** Fields that were renamed for a frame, per preset */
 	TMap<FName, TArray<TTuple<FName, FName>>> PerFrameRenamedFields;
