RC. Fix GEditor nullptr with -game

diff --git a/Engine/Plugins/VirtualProduction/RemoteControl/Source/RemoteControl/Private/RemoteControlModule.cpp b/Engine/Plugins/VirtualProduction/RemoteControl/Source/RemoteControl/Private/RemoteControlModule.cpp
--- a/Engine/Plugins/VirtualProduction/RemoteControl/Source/RemoteControl/Private/RemoteControlModule.cpp
+++ b/Engine/Plugins/VirtualProduction/RemoteControl/Source/RemoteControl/Private/RemoteControlModule.cpp
@@ -467,7 +467,7 @@ public:
 
 #if WITH_EDITOR
 			const bool bGenerateTransaction = ObjectAccess.Access == ERCAccess::WRITE_TRANSACTION_ACCESS;
-			if (bGenerateTransaction)
+			if (GEditor && bGenerateTransaction)
 			{
 				GEditor->BeginTransaction(LOCTEXT("RemoteSetPropertyTransaction", "Remote Set Object Property"));
 			}
@@ -510,7 +510,7 @@ public:
 
 			// Generate post edit property event independently from a transaction
 #if WITH_EDITOR
-			if (bGenerateTransaction)
+			if (GEditor && bGenerateTransaction)
 			{
 				GEditor->EndTransaction();
 			}
