Fix SetLabel not updating the panel
Add missing include to RemoteControlModule
Fix world not being resolved when in -game
Fix selected entity being unexposed not updating the details view

diff --git a/Engine/Plugins/VirtualProduction/RemoteControl/Source/RemoteControl/Private/RemoteControlBinding.cpp b/Engine/Plugins/VirtualProduction/RemoteControl/Source/RemoteControl/Private/RemoteControlBinding.cpp
--- a/Engine/Plugins/VirtualProduction/RemoteControl/Source/RemoteControl/Private/RemoteControlBinding.cpp
+++ b/Engine/Plugins/VirtualProduction/RemoteControl/Source/RemoteControl/Private/RemoteControlBinding.cpp
@@ -199,11 +199,12 @@ UWorld* URemoteControlLevelDependantBinding::GetCurrentWorld() const
 	{
 		if (WorldContext.WorldType == EWorldType::Game)
 		{
-			return World;
+			World = WorldContext.World();
+			break;
 		}
 	}
 		
-	return nullptr;
+	return World;
 }
 
 void URemoteControlLevelDependantBinding::SetBoundObject(const TSoftObjectPtr<ULevel>& Level, const TSoftObjectPtr<UObject>& BoundObject)
diff --git a/Engine/Plugins/VirtualProduction/RemoteControl/Source/RemoteControl/Private/RemoteControlModule.cpp b/Engine/Plugins/VirtualProduction/RemoteControl/Source/RemoteControl/Private/RemoteControlModule.cpp
--- a/Engine/Plugins/VirtualProduction/RemoteControl/Source/RemoteControl/Private/RemoteControlModule.cpp
+++ b/Engine/Plugins/VirtualProduction/RemoteControl/Source/RemoteControl/Private/RemoteControlModule.cpp
@@ -19,8 +19,10 @@
 #include "UObject/FieldPath.h"
 
 
+
 #if WITH_EDITOR
-	#include "ScopedTransaction.h"
+#include "Editor.h"
+#include "ScopedTransaction.h"
 #endif
 
 DEFINE_LOG_CATEGORY(LogRemoteControl);
diff --git a/Engine/Plugins/VirtualProduction/RemoteControl/Source/RemoteControlUI/Private/UI/SRCPanelExposedEntitiesList.cpp b/Engine/Plugins/VirtualProduction/RemoteControl/Source/RemoteControlUI/Private/UI/SRCPanelExposedEntitiesList.cpp
--- a/Engine/Plugins/VirtualProduction/RemoteControl/Source/RemoteControlUI/Private/UI/SRCPanelExposedEntitiesList.cpp
+++ b/Engine/Plugins/VirtualProduction/RemoteControl/Source/RemoteControlUI/Private/UI/SRCPanelExposedEntitiesList.cpp
@@ -483,13 +483,21 @@ void SRCPanelExposedEntitiesList::OnEntityRemoved(const FGuid& InGroupId, const
 {
 	if (TSharedPtr<SRCPanelGroup> PanelGroup = FindGroupById(InGroupId))
 	{
-		int32 EntityIndex = PanelGroup->GetNodes().IndexOfByPredicate([InEntityId](const TSharedPtr<SRCPanelTreeNode>& Node) {return Node->GetId() == InEntityId; });
+		int32 EntityIndex = PanelGroup->GetNodes().IndexOfByPredicate([InEntityId](const TSharedPtr<SRCPanelTreeNode>& Node) { return Node->GetId() == InEntityId; });
 		if (EntityIndex != INDEX_NONE)
 		{
 			PanelGroup->GetNodes().RemoveAt(EntityIndex);
 		}
 	}
 
+	if (TSharedPtr<SRCPanelTreeNode> Node = GetSelection())
+	{
+		if (Node->GetId() == InEntityId)
+		{
+			OnSelectionChangeDelegate.Broadcast(nullptr);
+		}
+	}
+
 	FieldWidgetMap.Remove(InEntityId);
 	TreeView->RequestListRefresh();
 }
@@ -523,6 +531,15 @@ void SRCPanelExposedEntitiesList::OnGroupAdded(const FRemoteControlPresetGroup&
 void SRCPanelExposedEntitiesList::OnGroupDeleted(FRemoteControlPresetGroup DeletedGroup)
 {
 	int32 Index = FieldGroups.IndexOfByPredicate([&DeletedGroup](const TSharedPtr<SRCPanelGroup>& Group) { return Group->GetId() == DeletedGroup.Id; });
+
+	if (TSharedPtr<SRCPanelTreeNode> Node = GetSelection())
+	{
+		if (DeletedGroup.GetFields().Contains(Node->GetId()))
+		{
+			OnSelectionChangeDelegate.Broadcast(nullptr);
+		}
+	}
+	
 	if (Index != INDEX_NONE)
 	{
 		FieldGroups.RemoveAt(Index);
diff --git a/Engine/Plugins/VirtualProduction/RemoteControl/Source/RemoteControlUI/Private/UI/SRCPanelExposedField.cpp b/Engine/Plugins/VirtualProduction/RemoteControl/Source/RemoteControlUI/Private/UI/SRCPanelExposedField.cpp
--- a/Engine/Plugins/VirtualProduction/RemoteControl/Source/RemoteControlUI/Private/UI/SRCPanelExposedField.cpp
+++ b/Engine/Plugins/VirtualProduction/RemoteControl/Source/RemoteControlUI/Private/UI/SRCPanelExposedField.cpp
@@ -147,10 +147,13 @@ void SRCPanelExposedField::Refresh()
 	{
 		CachedLabel = Field->GetLabel();
 
-		// Don't update the function row needlessly.
 		if (Field->FieldType == EExposedFieldType::Property)
 		{
-			ChildSlot.AttachWidget(ConstructWidget());
+			ConstructPropertyWidget();
+		}
+		else if (Field->FieldType == EExposedFieldType::Function)
+		{
+			ConstructFunctionWidget();
 		}
 	}
 }
diff --git a/Engine/Plugins/VirtualProduction/RemoteControl/Source/RemoteControlUI/Private/UI/SRemoteControlPanel.cpp b/Engine/Plugins/VirtualProduction/RemoteControl/Source/RemoteControlUI/Private/UI/SRemoteControlPanel.cpp
--- a/Engine/Plugins/VirtualProduction/RemoteControl/Source/RemoteControlUI/Private/UI/SRemoteControlPanel.cpp
+++ b/Engine/Plugins/VirtualProduction/RemoteControl/Source/RemoteControlUI/Private/UI/SRemoteControlPanel.cpp
@@ -100,7 +100,11 @@ void SRemoteControlPanel::Construct(const FArguments& InArgs, URemoteControlPres
 	
 	EntityList = SNew(SRCPanelExposedEntitiesList, Preset.Get())
 		.DisplayValues(true)
-		.OnEntityListUpdated_Lambda([this] () { UpdateRebindButtonVisibility(); })
+		.OnEntityListUpdated_Lambda([this] ()
+		{
+			UpdateEntityDetailsView(EntityList->GetSelection());
+			UpdateRebindButtonVisibility();
+		})
 		.EditMode_Lambda([this](){ return bIsInEditMode; });
 	
 	EntityList->OnSelectionChange().AddSP(this, &SRemoteControlPanel::UpdateEntityDetailsView);
