Add owner objecst to the entity events
Add notification when an entity is modified

diff --git a/Engine/Plugins/VirtualProduction/RemoteControl/Source/WebRemoteControl/Private/RemoteControlModels.h b/Engine/Plugins/VirtualProduction/RemoteControl/Source/WebRemoteControl/Private/RemoteControlModels.h
--- a/Engine/Plugins/VirtualProduction/RemoteControl/Source/WebRemoteControl/Private/RemoteControlModels.h
+++ b/Engine/Plugins/VirtualProduction/RemoteControl/Source/WebRemoteControl/Private/RemoteControlModels.h
@@ -60,6 +60,58 @@ namespace RemoteControlModels
 	}
 }
 
+/**
+ * Short description of an unreal object.
+ */
+USTRUCT()
+struct FRCObjectDescription
+{
+	GENERATED_BODY()
+
+	FRCObjectDescription() = default;
+
+	FRCObjectDescription(UObject* InObject)
+	{
+		checkSlow(InObject);
+		
+#if WITH_EDITOR
+		if (AActor* Actor = Cast<AActor>(InObject))
+		{
+			Name = Actor->GetActorLabel();
+		}
+		else if(UActorComponent* Component = Cast<UActorComponent>(InObject))
+		{
+			Name = Component->GetOwner()->GetActorLabel();
+		}
+		else
+#endif
+		{
+			// Get the class name when dealing with BP libraries and subsystems. 
+			Name = InObject->HasAnyFlags(RF_ClassDefaultObject | RF_ArchetypeObject) ? InObject->GetClass()->GetName() : InObject->GetName();
+		}
+		
+#if WITH_EDITOR
+		Class = InObject->GetClass()->GetDisplayNameText().ToString();
+#else
+		Class = InObject->GetClass()->GetName();
+#endif
+
+		Path = InObject->GetPathName();
+	}
+
+	/** Name of the object. */
+	UPROPERTY()
+	FString Name;
+
+	/** Class of the object. */
+	UPROPERTY()
+	FString Class;
+
+	/** Path of the object. */
+	UPROPERTY()
+	FString Path;
+};
+
 USTRUCT()
 struct FRCPropertyDescription
 {
@@ -216,6 +268,7 @@ struct FRCExposedPropertyDescription
 		UnderlyingProperty = RCProperty.GetProperty();
 		Metadata = RCProperty.GetMetadata();
 		Id = RCProperty.GetId().ToString();
+		OwnerObjects.Append(RCProperty.GetBoundObjects());
 	}
 
 	/** The label displayed in the remote control panel for this exposed property. */
@@ -233,8 +286,11 @@ struct FRCExposedPropertyDescription
 	/** Metadata specific to this exposed property. */
 	UPROPERTY()
 	TMap<FName, FString> Metadata;
-};
 
+	/** The objects that own the underlying property. */
+	UPROPERTY()
+	TArray<FRCObjectDescription> OwnerObjects; 
+};
 
 USTRUCT()
 struct FRCExposedFunctionDescription
@@ -248,6 +304,7 @@ struct FRCExposedFunctionDescription
 		, Id(Function.GetId().ToString())
 		, UnderlyingFunction(Function.GetFunction())
 	{
+		OwnerObjects.Append(Function.GetBoundObjects());
 	}
 
 	/** The label displayed in the remote control panel for this exposed property. */
@@ -261,38 +318,10 @@ struct FRCExposedFunctionDescription
 	/** The underlying exposed function. */
 	UPROPERTY()
 	FRCFunctionDescription UnderlyingFunction;
-};
-
-USTRUCT()
-struct FRCActorDescription
-{
-	GENERATED_BODY()
-
-	FRCActorDescription() = default;
-
-	FRCActorDescription(const AActor* InActor)
-	{
-		checkSlow(InActor);
-#if WITH_EDITOR
-		Name = InActor->GetActorLabel();
-#else
-		Name = InActor->GetName();
-#endif
-		Path = InActor->GetPathName();
-		Class = InActor->GetClass()->GetName();
-	}
 
-	/** The label of the actor. */
+	/** The objects that own the underlying function. */
 	UPROPERTY()
-	FString Name;
-
-	/** The object path of the actor. */
-	UPROPERTY()
-	FString Path;
-
-	/** The class of the actor. */
-	UPROPERTY()
-	FString Class;
+	TArray<FRCObjectDescription> OwnerObjects;
 };
 
 USTRUCT()
@@ -307,7 +336,7 @@ struct FRCExposedActorDescription
 	{
 		if (AActor* Actor = Cast<AActor>(InExposedActor.Path.ResolveObject()))
 		{
-			UnderlyingActor = FRCActorDescription{ Actor };
+			UnderlyingActor = FRCObjectDescription{ Actor };
 		}
 
 		Id = InExposedActor.GetId().ToString();
@@ -318,12 +347,12 @@ struct FRCExposedActorDescription
 	FName DisplayName;
 	
 	/** Unique identifier for the exposed actor. */
-    UPROPERTY()
-    FString Id;
+	UPROPERTY()
+	FString Id;
 
 	/** The underlying exposed actor. */
 	UPROPERTY()
-	FRCActorDescription UnderlyingActor;
+	FRCObjectDescription UnderlyingActor;
 };
 
 USTRUCT()
@@ -394,6 +423,68 @@ private:
 	}
 };
 
+/**
+ * Holds lists of modified RC entities.
+ * @Note that this does not mean the underlying property/function/actor was modified,
+ * but rather that the entity structure itself was modified in some way.
+ */
+USTRUCT()
+struct FRCPresetModifiedEntitiesDescription
+{
+	GENERATED_BODY()
+
+	FRCPresetModifiedEntitiesDescription() = default;
+	
+	FRCPresetModifiedEntitiesDescription(const URemoteControlPreset* Preset, const TArray<FGuid>& EntityIds)
+	{
+		checkSlow(Preset);
+		for (const FGuid& Id : EntityIds)
+		{
+			AddModifiedEntity(Preset, Id);
+		}
+	}
+
+public:
+	/** The list of modified RC properties. */
+	UPROPERTY()
+	TArray<FRCExposedPropertyDescription> ModifiedRCProperties;
+
+	/** The list of modified RC functions. */
+	UPROPERTY()
+	TArray<FRCExposedFunctionDescription> ModifiedRCFunctions;
+
+	/** The list of modified RC actors. */
+	UPROPERTY()
+	TArray<FRCExposedActorDescription> ModifiedRCActors;
+
+private:
+	/** Add an exposed field to this group description. */
+	void AddModifiedEntity(const URemoteControlPreset* Preset, const FGuid& EntityId)
+	{
+		if (TSharedPtr<const FRemoteControlProperty> RCProperty = Preset->GetExposedEntity<FRemoteControlProperty>(EntityId).Pin())
+		{
+			if (FProperty* Property = RCProperty->GetProperty())
+			{
+				ModifiedRCProperties.Emplace(*RCProperty);	
+			}
+		}
+		else if (TSharedPtr<const FRemoteControlFunction> RCFunction = Preset->GetExposedEntity<FRemoteControlFunction>(EntityId).Pin())
+		{
+			if (RCFunction->GetFunction())
+			{
+				ModifiedRCFunctions.Emplace(*RCFunction);
+			}
+		}
+		else if (TSharedPtr<const FRemoteControlActor> RCActor = Preset->GetExposedEntity<FRemoteControlActor>(EntityId).Pin())
+		{
+			if (RCActor->GetActor())
+			{
+				ModifiedRCActors.Emplace(*RCActor);
+			}
+		}
+	}
+};
+
 USTRUCT()
 struct FRCPresetDescription
 {
@@ -497,27 +588,6 @@ struct FRCAssetDescription
 	TMap<FName, FString> Metadata;
 };
 
-USTRUCT()
-struct FRCObjectDescription
-{
-	GENERATED_BODY()
-
-	FRCObjectDescription() = default;
-
-	FRCObjectDescription(const UObject* InObject)
-	{
-		checkSlow(InObject);
-		Name = InObject->GetName();
-		Path = InObject->GetPathName();
-	}
-
-	UPROPERTY()
-	FString Name;
-
-	UPROPERTY()
-	FString Path;
-};
-
 USTRUCT()
 struct FRCPresetFieldRenamed
 {
diff --git a/Engine/Plugins/VirtualProduction/RemoteControl/Source/WebRemoteControl/Private/RemoteControlResponse.h b/Engine/Plugins/VirtualProduction/RemoteControl/Source/WebRemoteControl/Private/RemoteControlResponse.h
--- a/Engine/Plugins/VirtualProduction/RemoteControl/Source/WebRemoteControl/Private/RemoteControlResponse.h
+++ b/Engine/Plugins/VirtualProduction/RemoteControl/Source/WebRemoteControl/Private/RemoteControlResponse.h
@@ -143,7 +143,7 @@ struct FSearchActorResponse
 	}
 
 	UPROPERTY()
-	TArray<FRCActorDescription> Actors;
+	TArray<FRCObjectDescription> Actors;
 };
 
 USTRUCT()
@@ -297,3 +297,47 @@ struct FRCPresetFieldsAddedEvent
 	UPROPERTY()
 	FRCPresetDescription Description;
 };
+
+/**
+ * Event triggered when an exposed entity struct is modified.
+ */
+USTRUCT()
+struct FRCPresetEntitiesModifiedEvent
+{
+	GENERATED_BODY()
+
+	FRCPresetEntitiesModifiedEvent() = default;
+
+	FRCPresetEntitiesModifiedEvent(URemoteControlPreset* InPreset, const TArray<FGuid>& InModifiedEntities)
+		: Type(TEXT("PresetEntitiesModified"))
+	{
+		checkSlow(InPreset);
+		PresetName = InPreset->GetFName();
+		PresetId = InPreset->GetPresetId().ToString();
+		ModifiedEntities = FRCPresetModifiedEntitiesDescription{InPreset, InModifiedEntities};
+	}
+
+	/**
+	 * Type of the event.
+	 */
+	UPROPERTY()
+	FString Type;
+
+	/**
+	 * Name of the preset which contains the modified entities.
+	 */
+	UPROPERTY()
+	FName PresetName;
+	
+	/**
+	 * ID of the preset that contains the modified entities.
+	 */
+	UPROPERTY()
+	FString PresetId;
+
+	/**
+	 * The entities that were modified in the last frame.
+	 */
+	UPROPERTY()
+	FRCPresetModifiedEntitiesDescription ModifiedEntities;
+};
\ No newline at end of file
diff --git a/Engine/Plugins/VirtualProduction/RemoteControl/Source/WebRemoteControl/Private/WebSocketMessageHandler.cpp b/Engine/Plugins/VirtualProduction/RemoteControl/Source/WebRemoteControl/Private/WebSocketMessageHandler.cpp
--- a/Engine/Plugins/VirtualProduction/RemoteControl/Source/WebRemoteControl/Private/WebSocketMessageHandler.cpp
+++ b/Engine/Plugins/VirtualProduction/RemoteControl/Source/WebRemoteControl/Private/WebSocketMessageHandler.cpp
@@ -88,6 +88,7 @@ void FWebSocketMessageHandler::HandleWebSocketPresetRegister(const FRemoteContro
 		Preset->OnFieldRenamed().AddRaw(this, &FWebSocketMessageHandler::OnFieldRenamed);
 		Preset->OnMetadataModified().AddRaw(this, &FWebSocketMessageHandler::OnMetadataModified);
 		Preset->OnActorPropertyModified().AddRaw(this, &FWebSocketMessageHandler::OnActorPropertyChanged);
+		Preset->OnEntitiesUpdated().AddRaw(this, &FWebSocketMessageHandler::OnEntitiesModified);
 	}
 
 	ClientIds->AddUnique(WebSocketMessage.ClientId);
@@ -325,6 +326,19 @@ void FWebSocketMessageHandler::OnActorPropertyChanged(URemoteControlPreset* Owne
 	}
 }
 
+void FWebSocketMessageHandler::OnEntitiesModified(URemoteControlPreset* Owner, const TSet<FGuid>& ModifiedEntities)
+{
+	// We do not need to store these event for the current frame since this was already handled by the preset in this case.
+	if (!Owner || ModifiedEntities.Num() == 0)
+	{
+		return;
+	}
+	
+	TArray<uint8> Payload;
+	WebRemoteControlUtils::SerializeResponse(FRCPresetEntitiesModifiedEvent{Owner, ModifiedEntities.Array()}, Payload);
+	BroadcastToListeners(Owner->GetFName(), Payload);
+}
+
 void FWebSocketMessageHandler::OnConnectionClosedCallback(FGuid ClientId)
 {
 	//Cleanup client that were waiting for callbacks
diff --git a/Engine/Plugins/VirtualProduction/RemoteControl/Source/WebRemoteControl/Private/WebSocketMessageHandler.h b/Engine/Plugins/VirtualProduction/RemoteControl/Source/WebRemoteControl/Private/WebSocketMessageHandler.h
--- a/Engine/Plugins/VirtualProduction/RemoteControl/Source/WebRemoteControl/Private/WebSocketMessageHandler.h
+++ b/Engine/Plugins/VirtualProduction/RemoteControl/Source/WebRemoteControl/Private/WebSocketMessageHandler.h
@@ -46,6 +46,7 @@ private:
 	void OnFieldRenamed(URemoteControlPreset* Owner, FName OldFieldLabel, FName NewFieldLabel);
 	void OnMetadataModified(URemoteControlPreset* Owner);
 	void OnActorPropertyChanged(URemoteControlPreset* Owner, FRemoteControlActor& Actor, UObject* ModifiedObject, FProperty* ModifiedProperty);
+	void OnEntitiesModified(URemoteControlPreset* Owner, const TSet<FGuid>& ModifiedEntities);
 
 	/** Callback when a websocket connection was closed. Let us clean out registrations */
 	void OnConnectionClosedCallback(FGuid ClientId);
