RC Function replication

diff --git a/Engine/Plugins/Runtime/nDisplay/Source/DisplayClusterRemoteControlInterceptor/Private/DisplayClusterRemoteControlInterceptor.cpp b/Engine/Plugins/Runtime/nDisplay/Source/DisplayClusterRemoteControlInterceptor/Private/DisplayClusterRemoteControlInterceptor.cpp
--- a/Engine/Plugins/Runtime/nDisplay/Source/DisplayClusterRemoteControlInterceptor/Private/DisplayClusterRemoteControlInterceptor.cpp
+++ b/Engine/Plugins/Runtime/nDisplay/Source/DisplayClusterRemoteControlInterceptor/Private/DisplayClusterRemoteControlInterceptor.cpp
@@ -14,6 +14,7 @@
 
 #include "Features/IModularFeatures.h"
 #include "HAL/IConsoleManager.h"
+#include "Misc/CommandLine.h"
 
 
 // Data interception source
@@ -29,13 +30,17 @@ static TAutoConsoleVariable<int32> CVarInterceptOnMasterOnly(
 
 // Magic numbers for now. Unfortunately there is no any ID management for binary events yet.
 // This will be refactored once we have some global events registry to prevent any ID conflicts.
-const int32 EventId_SetObjectProperties   = 0xaabb0701;
-const int32 EventId_ResetObjectProperties = 0xaabb0702;
+const int32 EventId_SetObjectProperties		= 0xaabb0701;
+const int32 EventId_ResetObjectProperties	= 0xaabb0702;
+const int32 EventId_InvokeCall				= 0xaabb0703;
 
 
 FDisplayClusterRemoteControlInterceptor::FDisplayClusterRemoteControlInterceptor()
 	: bInterceptOnMasterOnly(CVarInterceptOnMasterOnly.GetValueOnGameThread() == 1)
+	, bForceApply(false)
 {
+	bForceApply = FParse::Param(FCommandLine::Get(), TEXT("ClusterForceApplyResponse"));
+	
 	// Set up cluster events handler
 	EventsListener.BindRaw(this, &FDisplayClusterRemoteControlInterceptor::OnClusterEventBinaryHandler);
 	// Subscribe for cluster events
@@ -65,7 +70,11 @@ ERCIResponse FDisplayClusterRemoteControlInterceptor::SetObjectProperties(FRCIPr
 	static const FString EventName = FString(TEXT("SetObjectProperties"));
 	EmitReplicationEvent(EventId_SetObjectProperties, Buffer, EventName);
 
-	// The caller should not process the RC event
+	if (bForceApply)
+	{
+		return ERCIResponse::Apply;
+	}
+
 	return ERCIResponse::Intercept;
 }
 
@@ -80,7 +89,30 @@ ERCIResponse FDisplayClusterRemoteControlInterceptor::ResetObjectProperties(FRCI
 	static const FString EventName = FString(TEXT("ResetObjectProperties"));
 	EmitReplicationEvent(EventId_ResetObjectProperties, Buffer, EventName);
 
-	// The caller should not process the RC event
+	if (bForceApply)
+	{
+		return ERCIResponse::Apply;
+	}
+
+	return ERCIResponse::Intercept;
+}
+
+ERCIResponse FDisplayClusterRemoteControlInterceptor::InvokeCall(FRCIFunctionMetadata& InFunction)
+{
+	// Serialize command data to binary
+	TArray<uint8> Buffer;
+	FMemoryWriter MemoryWiter(Buffer);
+	MemoryWiter << InFunction;
+
+	// Replicate data
+	static const FString EventName = FString(TEXT("InvokeCall"));
+	EmitReplicationEvent(EventId_InvokeCall, Buffer, EventName);
+
+	if (bForceApply)
+	{
+		return ERCIResponse::Apply;
+	}
+	
 	return ERCIResponse::Intercept;
 }
 
@@ -116,6 +148,10 @@ void FDisplayClusterRemoteControlInterceptor::OnClusterEventBinaryHandler(const
 			OnReplication_ResetObjectProperties(Event.EventData);
 			break;
 
+		case EventId_InvokeCall:
+			OnReplication_InvokeCall(Event.EventData);
+			break;
+
 		default:
 			// Unsupported event, skipping it
 			break;
@@ -172,3 +208,28 @@ void FDisplayClusterRemoteControlInterceptor::OnReplication_ResetObjectPropertie
 		}
 	}
 }
+
+void FDisplayClusterRemoteControlInterceptor::OnReplication_InvokeCall(const TArray<uint8>& Buffer)
+{
+	UE_LOG(LogDisplayClusterRemoteControlInterceptor, VeryVerbose, TEXT("Processing replication event InvokeCall (0x%d): %d bytes"), EventId_InvokeCall, Buffer.Num());
+
+	// Deserialize command data
+	FMemoryReader MemoryReader(Buffer);
+	FRCIFunctionMetadata FunctionMetadata;
+	MemoryReader << FunctionMetadata;
+
+	// Initialization
+	IModularFeatures& ModularFeatures = IModularFeatures::Get();
+	const FName ProcessorFeatureName = IRemoteControlInterceptionFeatureProcessor::GetName();
+	const int32 ProcessorsAmount = ModularFeatures.GetModularFeatureImplementationCount(IRemoteControlInterceptionFeatureProcessor::GetName());
+
+	// Send the command to the processor(s)
+	for (int32 ProcessorIdx = 0; ProcessorIdx < ProcessorsAmount; ++ProcessorIdx)
+	{
+		IRemoteControlInterceptionFeatureProcessor* const Processor = static_cast<IRemoteControlInterceptionFeatureProcessor*>(ModularFeatures.GetModularFeatureImplementation(ProcessorFeatureName, ProcessorIdx));
+		if (Processor)
+		{
+			Processor->InvokeCall(FunctionMetadata);
+		}
+	}
+}
diff --git a/Engine/Plugins/Runtime/nDisplay/Source/DisplayClusterRemoteControlInterceptor/Private/DisplayClusterRemoteControlInterceptor.h b/Engine/Plugins/Runtime/nDisplay/Source/DisplayClusterRemoteControlInterceptor/Private/DisplayClusterRemoteControlInterceptor.h
--- a/Engine/Plugins/Runtime/nDisplay/Source/DisplayClusterRemoteControlInterceptor/Private/DisplayClusterRemoteControlInterceptor.h
+++ b/Engine/Plugins/Runtime/nDisplay/Source/DisplayClusterRemoteControlInterceptor/Private/DisplayClusterRemoteControlInterceptor.h
@@ -22,6 +22,7 @@ public:
 	// IRemoteControlInterceptionCommands interface
 	virtual ERCIResponse SetObjectProperties(FRCIPropertiesMetadata& InProperties) override;
 	virtual ERCIResponse ResetObjectProperties(FRCIObjectMetadata& InObject) override;
+	virtual ERCIResponse InvokeCall(FRCIFunctionMetadata& InFunction) override;
 	// ~IRemoteControlInterceptionCommands interface
 
 private:
@@ -35,10 +36,14 @@ private:
 	void OnReplication_SetObjectProperties   (const TArray<uint8>& Buffer);
 	// Process ResetObjectProperties command replication data
 	void OnReplication_ResetObjectProperties (const TArray<uint8>& Buffer);
+	// Process InvokeCall command replication data
+	void OnReplication_InvokeCall (const TArray<uint8>& Buffer);
 
 private:
 	// CVar value MasterOnly
 	const bool bInterceptOnMasterOnly;
 	// Cluster events listener
 	FOnClusterEventBinaryListener EventsListener;
+	// Force Apply the ERCIResponse
+	bool bForceApply;
 };
diff --git a/Engine/Plugins/VirtualProduction/RemoteControl/Source/RemoteControl/Private/RemoteControlInterceptionProcessor.cpp b/Engine/Plugins/VirtualProduction/RemoteControl/Source/RemoteControl/Private/RemoteControlInterceptionProcessor.cpp
--- a/Engine/Plugins/VirtualProduction/RemoteControl/Source/RemoteControl/Private/RemoteControlInterceptionProcessor.cpp
+++ b/Engine/Plugins/VirtualProduction/RemoteControl/Source/RemoteControl/Private/RemoteControlInterceptionProcessor.cpp
@@ -7,8 +7,8 @@
 #include "RemoteControlInterceptionHelpers.h"
 #include "Backends/CborStructDeserializerBackend.h"
 #include "Backends/JsonStructDeserializerBackend.h"
-#include "Serialization/MemoryWriter.h"
 #include "Serialization/MemoryReader.h"
+#include "StructDeserializer.h"
 #include "UObject/FieldPath.h"
 
 
@@ -65,3 +65,48 @@ void FRemoteControlInterceptionProcessor::ResetObjectProperties(FRCIObjectMetada
 		IRemoteControlModule::Get().ResetObjectProperties(ObjectRef, false);
 	}
 }
+
+void FRemoteControlInterceptionProcessor::InvokeCall(FRCIFunctionMetadata& InFunction)
+{
+	// This is invoke on a final render node
+	UObject* LoadedObject = StaticLoadObject(UObject::StaticClass(), nullptr, *InFunction.ObjectPath, nullptr, LOAD_None, nullptr, true);
+	UFunction* LoadedFunction = Cast<UFunction>(StaticLoadObject(UObject::StaticClass(), nullptr, *InFunction.FunctionPath, nullptr, LOAD_None, nullptr, true));
+	
+	if (!ensure(LoadedObject) || !ensure(LoadedFunction))
+	{
+		return;
+	}
+
+	FStructOnScope FunctionArgs{ LoadedFunction };
+	
+	// Create the backend based on SerializationMethod
+	FMemoryReader BackendMemoryReader(InFunction.Payload);
+	TSharedPtr<IStructDeserializerBackend> StructDeserializerBackend = nullptr;
+	if (InFunction.PayloadType == ERCIPayloadType::Cbor)
+	{
+		StructDeserializerBackend = MakeShared<FCborStructDeserializerBackend>(BackendMemoryReader);
+	}
+	else if (InFunction.PayloadType == ERCIPayloadType::Json)
+	{
+		StructDeserializerBackend = MakeShared<FJsonStructDeserializerBackend>(BackendMemoryReader);
+	}
+	else
+	{
+		// It should never be the case
+		checkNoEntry();
+	}
+
+	if (FStructDeserializer::Deserialize(FunctionArgs.GetStructMemory(), *const_cast<UStruct*>(FunctionArgs.GetStruct()), *StructDeserializerBackend, FStructDeserializerPolicies()))
+	{
+		FRCCallReference CallRef;
+		CallRef.Object = LoadedObject;
+		CallRef.Function = LoadedFunction;
+		
+		FRCCall Call;
+		Call.CallRef = MoveTemp(CallRef);
+		Call.bGenerateTransaction = InFunction.bGenerateTransaction;
+		Call.ParamStruct = FStructOnScope(FunctionArgs.GetStruct(), FunctionArgs.GetStructMemory());
+
+		IRemoteControlModule::Get().InvokeCall(Call);
+	}
+}
diff --git a/Engine/Plugins/VirtualProduction/RemoteControl/Source/RemoteControl/Private/RemoteControlInterceptionProcessor.h b/Engine/Plugins/VirtualProduction/RemoteControl/Source/RemoteControl/Private/RemoteControlInterceptionProcessor.h
--- a/Engine/Plugins/VirtualProduction/RemoteControl/Source/RemoteControl/Private/RemoteControlInterceptionProcessor.h
+++ b/Engine/Plugins/VirtualProduction/RemoteControl/Source/RemoteControl/Private/RemoteControlInterceptionProcessor.h
@@ -19,5 +19,6 @@ public:
 	// IRemoteControlInterceptionCommands interface
 	virtual void SetObjectProperties(FRCIPropertiesMetadata& InProperties) override;
 	virtual void ResetObjectProperties(FRCIObjectMetadata& InObject) override;
+	virtual void InvokeCall(FRCIFunctionMetadata& InFunction) override;
 	// ~IRemoteControlInterceptionCommands interface
 };
diff --git a/Engine/Plugins/VirtualProduction/RemoteControl/Source/RemoteControl/Private/RemoteControlModule.cpp b/Engine/Plugins/VirtualProduction/RemoteControl/Source/RemoteControl/Private/RemoteControlModule.cpp
--- a/Engine/Plugins/VirtualProduction/RemoteControl/Source/RemoteControl/Private/RemoteControlModule.cpp
+++ b/Engine/Plugins/VirtualProduction/RemoteControl/Source/RemoteControl/Private/RemoteControlModule.cpp
@@ -18,8 +18,6 @@
 #include "UObject/Class.h"
 #include "UObject/FieldPath.h"
 
-
-
 #if WITH_EDITOR
 #include "Editor.h"
 #include "ScopedTransaction.h"
@@ -125,11 +123,27 @@ namespace RemoteControlUtil
 		return GetFirstPreset(Filter);
 	}
 
+	URemoteControlPreset* FindPresetByName(FName PresetName)
+	{
+		IAssetRegistry& AssetRegistry = FModuleManager::Get().LoadModuleChecked<FAssetRegistryModule>(TEXT("AssetRegistry")).Get();
+
+		TArray<FAssetData> Assets;
+		AssetRegistry.GetAssetsByClass(URemoteControlPreset::StaticClass()->GetFName(), Assets);
+
+		FAssetData* FoundAsset = Assets.FindByPredicate([&PresetName](const FAssetData& InAsset)
+		{
+			return InAsset.AssetName == PresetName;
+		});
+
+		return FoundAsset ? Cast<URemoteControlPreset>(FoundAsset->GetAsset()) : nullptr;
+	}
+
 	URemoteControlPreset* GetPresetByName(FName PresetName)
 	{
 		FARFilter Filter = GetBasePresetFilter();
 		Filter.PackageNames = { PresetName };
-		return GetFirstPreset(Filter);
+		URemoteControlPreset* FoundPreset = GetFirstPreset(Filter);
+		return FoundPreset ? FoundPreset : FindPresetByName(PresetName);
 	}
 
 	FGuid GetPresetId(const FAssetData& PresetAsset)
@@ -264,10 +278,39 @@ public:
 		return bSuccess;
 	}
 
-	virtual bool InvokeCall(FRCCall& InCall) override
+	virtual bool InvokeCall(FRCCall& InCall, ERCPayloadType InPayloadType, const TArray<uint8>& InInterceptPayload) override
 	{
 		if (InCall.IsValid())
 		{
+			// Check the replication path before apply property values
+			if (InInterceptPayload.Num() != 0)
+			{
+				FRCIFunctionMetadata FunctionMetadata(InCall.CallRef.Object->GetPathName(), InCall.CallRef.Function->GetPathName(), InCall.bGenerateTransaction, ToExternal(InPayloadType), InInterceptPayload);
+
+				// Initialization
+				IModularFeatures& ModularFeatures = IModularFeatures::Get();
+				const FName InterceptorFeatureName = IRemoteControlInterceptionFeatureInterceptor::GetName();
+				const int32 InterceptorsAmount = ModularFeatures.GetModularFeatureImplementationCount(IRemoteControlInterceptionFeatureInterceptor::GetName());
+
+				// Pass interception command data to all available interceptors
+				bool bShouldIntercept = false;
+				for (int32 InterceptorIdx = 0; InterceptorIdx < InterceptorsAmount; ++InterceptorIdx)
+				{
+					IRemoteControlInterceptionFeatureInterceptor* const Interceptor = static_cast<IRemoteControlInterceptionFeatureInterceptor*>(ModularFeatures.GetModularFeatureImplementation(InterceptorFeatureName, InterceptorIdx));
+					if (Interceptor)
+					{
+						// Update response flag
+						bShouldIntercept |= (Interceptor->InvokeCall(FunctionMetadata) == ERCIResponse::Intercept);
+					}
+				}
+
+				// Don't process the RC message if any of interceptors returned ERCIResponse::Intercept
+				if (bShouldIntercept)
+				{
+					return true;
+				}
+			}
+			
 #if WITH_EDITOR
 			FScopedTransaction Transaction(LOCTEXT("RemoteCallTransaction", "Remote Call Transaction Wrap"), InCall.bGenerateTransaction);
 #endif
@@ -440,19 +483,19 @@ public:
 			const int32 InterceptorsAmount = ModularFeatures.GetModularFeatureImplementationCount(IRemoteControlInterceptionFeatureInterceptor::GetName());
 
 			// Pass interception command data to all available interceptors
-			bool bShouldApply = false;
+			bool bShouldIntercept = false;
 			for (int32 InterceptorIdx = 0; InterceptorIdx < InterceptorsAmount; ++InterceptorIdx)
 			{
 				IRemoteControlInterceptionFeatureInterceptor* const Interceptor = static_cast<IRemoteControlInterceptionFeatureInterceptor*>(ModularFeatures.GetModularFeatureImplementation(InterceptorFeatureName, InterceptorIdx));
 				if (Interceptor)
 				{
 					// Update response flag
-					bShouldApply |= (Interceptor->SetObjectProperties(PropsMetadata) == ERCIResponse::Apply);
+					bShouldIntercept |= (Interceptor->SetObjectProperties(PropsMetadata) == ERCIResponse::Intercept);
 				}
 			}
 
 			// Don't process the RC message if any of interceptors returned ERCIResponse::Intercept
-			if (bShouldApply)
+			if (bShouldIntercept)
 			{
 				return true;
 			}
@@ -540,19 +583,19 @@ public:
 			const int32 InterceptorsAmount = ModularFeatures.GetModularFeatureImplementationCount(IRemoteControlInterceptionFeatureInterceptor::GetName());
 
 			// Pass interception command data to all available interceptors
-			bool bShouldApply = false;
+			bool bShouldIntercept = false;
 			for (int32 InterceptorIdx = 0; InterceptorIdx < InterceptorsAmount; ++InterceptorIdx)
 			{
 				IRemoteControlInterceptionFeatureInterceptor* const Interceptor = static_cast<IRemoteControlInterceptionFeatureInterceptor*>(ModularFeatures.GetModularFeatureImplementation(InterceptorFeatureName, InterceptorIdx));
 				if (Interceptor)
 				{
 					// Update response flag
-					bShouldApply |= (Interceptor->ResetObjectProperties(ObjectMetadata) == ERCIResponse::Apply);
+					bShouldIntercept |= (Interceptor->ResetObjectProperties(ObjectMetadata) == ERCIResponse::Intercept);
 				}
 			}
 
 			// Don't process the RC message if any of interceptors returned ERCIResponse::Intercept
-			if (bShouldApply)
+			if (bShouldIntercept)
 			{
 				return true;
 			}
diff --git a/Engine/Plugins/VirtualProduction/RemoteControl/Source/RemoteControl/Private/Tests/RemoteControlInterceptionTest.cpp b/Engine/Plugins/VirtualProduction/RemoteControl/Source/RemoteControl/Private/Tests/RemoteControlInterceptionTest.cpp
--- a/Engine/Plugins/VirtualProduction/RemoteControl/Source/RemoteControl/Private/Tests/RemoteControlInterceptionTest.cpp
+++ b/Engine/Plugins/VirtualProduction/RemoteControl/Source/RemoteControl/Private/Tests/RemoteControlInterceptionTest.cpp
@@ -12,7 +12,6 @@
 #include "RemoteControlInterceptionProcessor.h"
 #include "RemoteControlInterceptionTestData.h"
 #include "RemoteControlPreset.h"
-#include "StructDeserializer.h"
 #include "StructSerializer.h"
 
 // Interception related headers
@@ -22,7 +21,6 @@
 #include "IRemoteControlInterceptionFeature.h"
 #include "Serialization/MemoryWriter.h"
 #include "Serialization/MemoryReader.h"
-#include "UObject/FieldPath.h"
 
 
 /**
@@ -65,6 +63,21 @@ protected:
 
 		return ERCIResponse::Intercept;
 	}
+
+	virtual ERCIResponse InvokeCall(FRCIFunctionMetadata& InFunction) override
+	{
+		// Get processor feature
+		IRemoteControlInterceptionFeatureProcessor* const Processor = static_cast<IRemoteControlInterceptionFeatureProcessor*>(
+			IModularFeatures::Get().GetModularFeatureImplementation(IRemoteControlInterceptionFeatureProcessor::GetName(), 0));
+
+		// In case the processor feature has been registered, forward data directly to the processor
+		if (Processor)
+		{
+			Processor->InvokeCall(InFunction);
+		}
+		
+		return ERCIResponse::Intercept;
+	}
 	// ~IRemoteControlInterceptionCommands interface
 };
 
@@ -85,6 +98,11 @@ struct FRemoteControlInterceptionTest
 		FString Int32ValuePropertyFullPath = FString::Printf(TEXT("%s.%s"), *CustomStructProperty->GetName(), *Int32ValueProperty->GetName());
 		Int32ValuePropertyRCProp = Preset->ExposeProperty(RemoteControlInterceptionTestObject.Get(), FRCFieldPathInfo{ Int32ValuePropertyFullPath }).Pin();
 
+		// Expose function
+		UClass* TestObjectClass = RemoteControlInterceptionTestObject->StaticClass();
+		UFunction* TestFunctionClass = TestObjectClass->FindFunctionByName("TestFunction");
+		TestFunctionRC = Preset->ExposeFunction(RemoteControlInterceptionTestObject.Get(), TestFunctionClass).Pin();
+		
 		// Reset tested value
 		RemoteControlInterceptionTestObject->CustomStruct.Int32Value = FRemoteControlInterceptionTestStruct::Int32ValueDefault;
 
@@ -134,7 +152,7 @@ struct FRemoteControlInterceptionTest
 		for (UObject* Object : Int32ValuePropertyRCProp->GetBoundObjects())
 		{
 			IRemoteControlModule::Get().ResolveObjectProperty(ObjectRef.Access, Object, ObjectRef.PropertyPathInfo, ObjectRef);
-			// Arhive should be intercepted inside SetObjectProperties
+			// Archive should be intercepted inside SetObjectProperties
 			IRemoteControlModule::Get().SetObjectProperties(ObjectRef, CborStructDeserializerBackend, ERCPayloadType::Cbor, ApiBuffer);
 		}
 	}
@@ -163,7 +181,7 @@ struct FRemoteControlInterceptionTest
 		for (UObject* Object : Int32ValuePropertyRCProp->GetBoundObjects())
 		{
 			IRemoteControlModule::Get().ResolveObjectProperty(ObjectRef.Access, Object, ObjectRef.PropertyPathInfo, ObjectRef);
-			// Arhive should be intercepted inside SetObjectProperties
+			// Archive should be intercepted inside SetObjectProperties
 			IRemoteControlModule::Get().SetObjectProperties(ObjectRef, DeserializerBackend, ERCPayloadType::Cbor, CborBuffer);
 		}
 	}
@@ -200,10 +218,10 @@ struct FRemoteControlInterceptionTest
 
 	void ResetObjectProperties()
 	{
-		// Set reste value before reset object property
+		// Set reset value before reset object property
 		RemoteControlInterceptionTestObject->CustomStruct.Int32Value = Int32ValueTest;
 
-		// Set ovject reference
+		// Set object reference
 		FRCObjectReference ObjectRef;
 		ObjectRef.Property = Int32ValuePropertyRCProp->GetProperty();
 		ObjectRef.Access = ERCAccess::WRITE_TRANSACTION_ACCESS;
@@ -219,11 +237,133 @@ struct FRemoteControlInterceptionTest
 		Test->TestNotEqual(TEXT("The property value should be applied only after interception"), Int32ValueTest, RemoteControlInterceptionTestObject->CustomStruct.Int32Value);
 	}
 
+	void InvokeCall()
+	{
+		if (!TestFunctionRC.IsValid())
+		{
+			Test->AddError(TEXT("TestFunctionRC not valid"));
+			return;
+		}
+
+		constexpr auto InStructName = TEXT("InStruct");
+		constexpr auto InTestFactorName = TEXT("InTestFactor");
+		const FName Int32ValueStructName = GET_MEMBER_NAME_CHECKED(FRemoteControlInterceptionFunctionParamStruct, Int32Value);
+		const FName IntArrayStructName = GET_MEMBER_NAME_CHECKED(FRemoteControlInterceptionFunctionParamStruct, IntArray);
+		const FName IntStringStructName = GET_MEMBER_NAME_CHECKED(FRemoteControlInterceptionFunctionParamStruct, IntString);
+
+		// Define test struct
+		FRemoteControlInterceptionFunctionParamStruct FunctionParamStruct;
+		FunctionParamStruct.Int32Value = 14;
+		FunctionParamStruct.IntArray.Add(1);
+		FunctionParamStruct.IntArray.Add(2);
+		FunctionParamStruct.IntArray.Add(3);
+		FunctionParamStruct.IntString = FString::FromInt(FunctionParamStruct.Int32Value);
+		constexpr int32 TestFactorInt = 3;
+
+		// Copy test struct into the object struct
+		RemoteControlInterceptionTestObject->FunctionParamStruct = FunctionParamStruct;
+		
+		// Set a input function argument values
+		FStructOnScope FunctionArgs{ TestFunctionRC->GetFunction() };
+		for (TFieldIterator<FProperty> It(TestFunctionRC->GetFunction()); It; ++It)
+		{
+			if (It->HasAnyPropertyFlags(CPF_Parm) && !It->HasAnyPropertyFlags(CPF_ReturnParm))
+			{
+				const FProperty* Property = *It;
+				void* const ValueAddress = Property->ContainerPtrToValuePtr<void>(FunctionArgs.GetStructMemory());
+				const FName PropertyName = Property->NamePrivate;
+
+				// Set struct input values
+				if (PropertyName == InStructName)
+				{
+					if (const FStructProperty* StructProperty =  CastField<FStructProperty>(Property))
+					{
+						StructProperty->InitializeValue_InContainer(FunctionArgs.GetStructMemory());
+
+						for (TFieldIterator<FProperty> StructIt(StructProperty->Struct); StructIt; ++StructIt)
+						{
+							const FProperty* InnerStructProperty = *StructIt;
+							const FName StructPropertyName = InnerStructProperty->NamePrivate;
+							void* InnerStructValueAddress = InnerStructProperty->ContainerPtrToValuePtr<void>(ValueAddress, 0);
+
+							// Set test integer value
+							if (StructPropertyName == Int32ValueStructName)
+							{
+								if (const FNumericProperty* StructNumericProperty = CastField<FNumericProperty>(InnerStructProperty))
+								{
+									StructNumericProperty->CopyCompleteValue(InnerStructValueAddress, &FunctionParamStruct.Int32Value);
+								}
+							}
+
+							// Set test array values
+							else if (StructPropertyName == IntArrayStructName)
+							{
+								if (const FArrayProperty* StructArrayProperty = CastField<FArrayProperty>(InnerStructProperty))
+								{
+									if (const FNumericProperty* ArrayNumericProperty = CastField<FNumericProperty>(StructArrayProperty->Inner))
+									{
+										FScriptArrayHelper ScriptArrayHelper(StructArrayProperty, StructArrayProperty->ContainerPtrToValuePtr<void>(ValueAddress));
+
+										for (const int32 IntValue : FunctionParamStruct.IntArray)
+										{
+											const int32 Index = ScriptArrayHelper.AddValue();
+											void* ArrayElementValuePtr = ScriptArrayHelper.GetRawPtr(Index);
+											ArrayNumericProperty->CopyCompleteValue(ArrayElementValuePtr, &IntValue);
+										}	
+									}
+								}
+							}
+							else if (StructPropertyName == IntStringStructName)
+							{
+								if (const FStrProperty* StructStringProperty = CastField<FStrProperty>(InnerStructProperty))
+								{
+									StructStringProperty->SetPropertyValue(InnerStructValueAddress, FunctionParamStruct.IntString);
+								}
+							}
+						}
+					}
+				}
+				// Set integer input value
+				else if (PropertyName == InTestFactorName)
+				{
+					if (const FNumericProperty* NumericProperty = CastField<FNumericProperty>(Property))
+					{
+						NumericProperty->CopyCompleteValue(ValueAddress, &TestFactorInt);
+					}
+				}
+			}
+		}
+
+		// Set testing Archives
+		TArray<uint8> JsonBuffer;
+		FMemoryWriter Writer(JsonBuffer);
+
+		// Set serializers
+		FJsonStructSerializerBackend SerializerBackend(Writer, EStructSerializerBackendFlags::Default);
+		FStructSerializer::Serialize(FunctionArgs.GetStructMemory(), *const_cast<UStruct*>(FunctionArgs.GetStruct()), SerializerBackend);
+
+		FRCCallReference CallRef;
+		CallRef.Object = RemoteControlInterceptionTestObject.Get();
+		CallRef.Function = TestFunctionRC->GetFunction();
+		
+		FRCCall Call;
+		Call.CallRef = MoveTemp(CallRef);
+		Call.bGenerateTransaction = true;
+		Call.ParamStruct = FStructOnScope(FunctionArgs.GetStruct(), FunctionArgs.GetStructMemory());
+
+		const bool bSuccess = IRemoteControlModule::Get().InvokeCall(Call, ERCPayloadType::Json, JsonBuffer);
+		Test->TestTrue(TEXT("InvokeCall should return true"), bSuccess);
+
+		Test->TestEqual(TEXT("The property value should be applied after interception"), FunctionParamStruct.Int32Value * TestFactorInt, RemoteControlInterceptionTestObject->FunctionParamStruct.Int32Value);
+		Test->TestEqual(TEXT("The property value should be applied after interception"), FunctionParamStruct.IntArray[2] * TestFactorInt, RemoteControlInterceptionTestObject->FunctionParamStruct.IntArray[2]);
+	}
+
 private:
 	FAutomationTestBase* Test;
 	TStrongObjectPtr<URemoteControlPreset> Preset;
 	TStrongObjectPtr<URemoteControlInterceptionTestObject> RemoteControlInterceptionTestObject;
 	TSharedPtr<FRemoteControlProperty> Int32ValuePropertyRCProp;
+	TSharedPtr<FRemoteControlFunction> TestFunctionRC;
 	FProperty* CustomStructProperty = nullptr;
 	FProperty* Int32ValueProperty = nullptr;
 	const int32 Int32ValueTest = 34780;
@@ -241,6 +381,7 @@ bool FRemoteControlPresetInterceptionTest::RunTest(const FString& Parameters)
 	RemoteControlInterceptionTest.TestCbor();
 	RemoteControlInterceptionTest.TestJson();
 	RemoteControlInterceptionTest.ResetObjectProperties();
+	RemoteControlInterceptionTest.InvokeCall();
 	
 	return true;
 }
diff --git a/Engine/Plugins/VirtualProduction/RemoteControl/Source/RemoteControl/Private/Tests/RemoteControlInterceptionTestData.h b/Engine/Plugins/VirtualProduction/RemoteControl/Source/RemoteControl/Private/Tests/RemoteControlInterceptionTestData.h
--- a/Engine/Plugins/VirtualProduction/RemoteControl/Source/RemoteControl/Private/Tests/RemoteControlInterceptionTestData.h
+++ b/Engine/Plugins/VirtualProduction/RemoteControl/Source/RemoteControl/Private/Tests/RemoteControlInterceptionTestData.h
@@ -16,6 +16,20 @@ struct FRemoteControlInterceptionTestStruct
 	int32 Int32Value = Int32ValueDefault;
 };
 
+USTRUCT(BlueprintType)
+struct FRemoteControlInterceptionFunctionParamStruct
+{
+	GENERATED_BODY()
+
+	UPROPERTY(EditAnywhere, Category = "RC")
+	int32 Int32Value = 0;
+
+	UPROPERTY(EditAnywhere, Category = "RC")
+	TArray<int32> IntArray;
+
+	UPROPERTY(EditAnywhere, Category = "RC")
+	FString IntString;
+};
 
 UCLASS()
 class URemoteControlInterceptionTestObject : public UObject
@@ -25,4 +39,32 @@ public:
 
 	UPROPERTY(EditAnywhere, Category = "RC")
 	FRemoteControlInterceptionTestStruct CustomStruct;
+
+	UPROPERTY(EditAnywhere, Category = "RC")
+	FRemoteControlInterceptionFunctionParamStruct FunctionParamStruct;
+
+	UFUNCTION(BlueprintCallable, Category = "RC")
+	FRemoteControlInterceptionFunctionParamStruct TestFunction(const FRemoteControlInterceptionFunctionParamStruct& InStruct, int32 InTestFactor)
+	{
+		FRemoteControlInterceptionFunctionParamStruct Copy = InStruct;
+
+		Copy.Int32Value *= InTestFactor;
+
+		for (int32& IntValue : Copy.IntArray)
+		{
+			IntValue *= InTestFactor;
+		}
+
+		int32 StringInt = 0;
+		LexFromString(StringInt, *Copy.IntString);
+
+		StringInt *= InTestFactor;
+		
+		Copy.IntString = FString::FromInt(StringInt);
+
+		FunctionParamStruct = Copy;
+
+		return Copy;
+	}
 };
+
diff --git a/Engine/Plugins/VirtualProduction/RemoteControl/Source/RemoteControl/Public/IRemoteControlModule.h b/Engine/Plugins/VirtualProduction/RemoteControl/Source/RemoteControl/Public/IRemoteControlModule.h
--- a/Engine/Plugins/VirtualProduction/RemoteControl/Source/RemoteControl/Public/IRemoteControlModule.h
+++ b/Engine/Plugins/VirtualProduction/RemoteControl/Source/RemoteControl/Public/IRemoteControlModule.h
@@ -204,9 +204,11 @@ public:
 	 * This is a thin wrapper around UObject::ProcessEvent
 	 * This expects that the caller has already validated the call as it will assert otherwise.
 	 * @param InCall the remote call structure to call.
+	 * @param InPayloadType the payload type archive.
+	 * @param InInterceptPayload the payload reference archive for the interception.
 	 * @return true if the call was allowed and done.
 	 */
-	virtual bool InvokeCall(FRCCall& InCall) = 0;
+	virtual bool InvokeCall(FRCCall& InCall, ERCPayloadType InPayloadType = ERCPayloadType::Json, const TArray<uint8>& InInterceptPayload = TArray<uint8>()) = 0;
 
 	/**
 	 * Resolve a remote object reference to a property
diff --git a/Engine/Plugins/VirtualProduction/RemoteControl/Source/WebRemoteControl/Private/RemoteControlRequest.h b/Engine/Plugins/VirtualProduction/RemoteControl/Source/WebRemoteControl/Private/RemoteControlRequest.h
--- a/Engine/Plugins/VirtualProduction/RemoteControl/Source/WebRemoteControl/Private/RemoteControlRequest.h
+++ b/Engine/Plugins/VirtualProduction/RemoteControl/Source/WebRemoteControl/Private/RemoteControlRequest.h
@@ -10,6 +10,9 @@ struct FBlockDelimiters
 {
 	int64 BlockStart = -1;
 	int64 BlockEnd = -1;
+
+	/** Get the size of current block */
+	int64 GetBlockSize() const { return BlockEnd - BlockStart; }
 };
 
 /**
diff --git a/Engine/Plugins/VirtualProduction/RemoteControl/Source/WebRemoteControl/Private/WebRemoteControl.cpp b/Engine/Plugins/VirtualProduction/RemoteControl/Source/WebRemoteControl/Private/WebRemoteControl.cpp
--- a/Engine/Plugins/VirtualProduction/RemoteControl/Source/WebRemoteControl/Private/WebRemoteControl.cpp
+++ b/Engine/Plugins/VirtualProduction/RemoteControl/Source/WebRemoteControl/Private/WebRemoteControl.cpp
@@ -57,8 +57,6 @@
 #include "Templates/UnrealTemplate.h"
 
 #define LOCTEXT_NAMESPACE "WebRemoteControl"
-
-
 // Boot the server on startup flag
 static TAutoConsoleVariable<int32> CVarWebControlStartOnBoot(TEXT("WebControl.EnableServerOnStartup"), 0, TEXT("Enable the Web Control servers (web and websocket) on startup."));
 
@@ -704,9 +702,7 @@ bool FWebRemoteControlModule::HandlePresetCallFunctionRoute(const FHttpServerReq
 	}
 
 	FBlockDelimiters Delimiters = CallRequest.GetParameterDelimiters(FRCPresetCallRequest::ParametersLabel());
-
-	FMemoryReader Reader{ CallRequest.TCHARBody };
-	FRCJsonStructDeserializerBackend ReaderBackend{ Reader };
+	const int64 DelimitersSize = Delimiters.GetBlockSize();
 
 	TArray<uint8> OutputBuffer;
 	FMemoryWriter Writer{ OutputBuffer };
@@ -719,10 +715,24 @@ bool FWebRemoteControlModule::HandlePresetCallFunctionRoute(const FHttpServerReq
 
 	bool bSuccess = false;
 
-	if (Delimiters.BlockStart != Delimiters.BlockEnd)
-	{
-		Reader.Seek(Delimiters.BlockStart);
-		Reader.SetLimitSize(Delimiters.BlockEnd + 1);
+	if (Delimiters.BlockStart != Delimiters.BlockEnd &&
+		CallRequest.TCHARBody.IsValidIndex(Delimiters.BlockStart) &&
+		CallRequest.TCHARBody.IsValidIndex(DelimitersSize)
+		)
+	{
+		/**
+		 * In order to have a replication payload we need to copy the inner payload from TCHARBody to new buffer
+		 * Example:
+		 * Original buffer from HTTP rquest: { "Parameters": { "NewLocation": {"X": 0, "Y": 0, "Z": 400} }
+		 * New buffer: "NewLocation": {"X": 0, "Y": 0, "Z": 400}
+		 */
+		TArray<uint8> FunctionPayload;
+		FunctionPayload.SetNumUninitialized(DelimitersSize);
+		const uint8* DataStart = &CallRequest.TCHARBody[Delimiters.BlockStart];	
+		FMemory::Memcpy(FunctionPayload.GetData(), DataStart, DelimitersSize);
+
+		FMemoryReader Reader{ FunctionPayload };
+		FRCJsonStructDeserializerBackend ReaderBackend{ Reader };
 
 		// Copy the default arguments.
 		FStructOnScope FunctionArgs{ RCFunction->GetFunction() };
@@ -749,7 +759,9 @@ bool FWebRemoteControlModule::HandlePresetCallFunctionRoute(const FHttpServerReq
 				Call.bGenerateTransaction = CallRequest.GenerateTransaction;
 				Call.ParamStruct = FStructOnScope(FunctionArgs.GetStruct(), FunctionArgs.GetStructMemory());
 
-				bSuccess &= IRemoteControlModule::Get().InvokeCall(Call);
+				// Invoke call with replication payload
+				bSuccess &= IRemoteControlModule::Get().InvokeCall(Call, ERCPayloadType::Json, FunctionPayload);
+
 				if (bSuccess)
 				{
 					FStructOnScope ReturnedStruct{ FunctionArgs.GetStruct() };
diff --git a/Engine/Plugins/VirtualProduction/RemoteControlInterception/Source/RemoteControlInterception/Public/IRemoteControlInterceptionCommands.h b/Engine/Plugins/VirtualProduction/RemoteControlInterception/Source/RemoteControlInterception/Public/IRemoteControlInterceptionCommands.h
--- a/Engine/Plugins/VirtualProduction/RemoteControlInterception/Source/RemoteControlInterception/Public/IRemoteControlInterceptionCommands.h
+++ b/Engine/Plugins/VirtualProduction/RemoteControlInterception/Source/RemoteControlInterception/Public/IRemoteControlInterceptionCommands.h
@@ -3,9 +3,9 @@
 #pragma once
 
 #include "CoreMinimal.h"
-#include "Templates/SharedPointer.h"
 
 struct FRCIObjectMetadata;
+struct FRCIFunctionMetadata;
 struct FRCIPropertiesMetadata;
 
 /**
@@ -65,9 +65,18 @@ public:
 	 * @param InObject - Metadata of the object and its properties that are going to be modified
 	 *
 	 * @return - Return types depends on the interface implementation (template parameter)
-	 * @param OutResponse - interception response flag (not requred for the handlers)
+	 * @param OutResponse - interception response flag (not required for the handlers)
 	 */
 	virtual TResponseType ResetObjectProperties(FRCIObjectMetadata& InObject) = 0;
+	
+	/**
+	 * InvokeCall command to process
+	 *
+	 * @param InFunction - Metadata of the UFunction that are going to be modified
+	 *
+	 * @return - Return types depends on the interface implementation (template parameter)
+	 */
+	virtual TResponseType InvokeCall(FRCIFunctionMetadata& InFunction) = 0;
 };
 
 
@@ -100,7 +109,7 @@ public:
 	 * FRCObjectMetadata serialization
 	 *
 	 * @param Ar       - The archive
-	 * @param Instance - Instance to serialize/deserialize
+	 * @param Interception - Instance to serialize/deserialize
 	 *
 	 * @return FArchive instance.
 	 */
@@ -148,7 +157,7 @@ public:
 	//~ Begin FRCObjectMetadata
 	virtual bool IsValid() const override
 	{
-		return FRCIObjectMetadata::IsValid() && PayloadType != ERCIPayloadType::Json && Payload.Num() > 0;
+		return FRCIObjectMetadata::IsValid() && Payload.Num() > 0;
 	}
 	//~ End FRCObjectMetadata
 
@@ -162,12 +171,12 @@ public:
 	 */
 	friend FArchive& operator<<(FArchive& Ar, FRCIPropertiesMetadata& Instance)
 	{
-		// Call achive from the parent class
+		// Call archive from the parent class
 		FRCIObjectMetadata& Super = static_cast<FRCIObjectMetadata&>(Instance);
 		Ar << Super;
 		
 		Ar << Instance.PayloadType;
-		// const cast is needed for keeping the const& on a constractor and as a class parameter
+		// const cast is needed for keeping the const& on a constructor and as a class parameter
 		Ar << const_cast<TArray<uint8>&>(Instance.Payload);
 		return Ar;
 	}
@@ -179,3 +188,65 @@ public:
 	/** Property payload to intercept */
 	const TArray<uint8> Payload;
 };
+
+/**
+* UFunction metadata for custom interception/replication/processing purposes
+*/
+struct FRCIFunctionMetadata
+{
+public:
+	FRCIFunctionMetadata()
+		: bGenerateTransaction(false)
+		, PayloadType(ERCIPayloadType::Json)
+	{ }
+
+	FRCIFunctionMetadata(const FString& InObjectPath, const FString& InFunctionPath, const bool bInGenerateTransaction, const ERCIPayloadType InPayloadType, const TArray<uint8>& InPayload)
+		: ObjectPath(InObjectPath)
+		, FunctionPath(InFunctionPath)
+		, bGenerateTransaction(bInGenerateTransaction)
+		, PayloadType(InPayloadType)
+		, Payload(InPayload)
+	{ }
+
+public:
+	/** Returns true if reference is valid */
+	bool IsValid() const
+	{
+		return !ObjectPath.IsEmpty() && !FunctionPath.IsEmpty() && Payload.Num() > 0;
+	}
+
+	/**
+	* FRCObjectMetadata serialization
+	*
+	* @param Ar				- The archive
+	* @param Interception	- Instance to serialize/deserialize
+	*
+	* @return FArchive instance.
+	*/
+	friend FArchive& operator<<(FArchive& Ar, FRCIFunctionMetadata& Interception)
+	{
+		Ar << Interception.ObjectPath;
+		Ar << Interception.FunctionPath;
+		Ar << Interception.bGenerateTransaction;
+		Ar << Interception.PayloadType;
+		// const cast is needed for keeping the const& on a constructor and as a class parameter
+		Ar << const_cast<TArray<uint8>&>(Interception.Payload);
+		return Ar;
+	}
+
+public:
+	/** Owner object path */
+	FString ObjectPath;
+
+	/** UFunction object path */
+	FString FunctionPath;
+
+	/** Should the call generate transaction */
+	bool bGenerateTransaction;
+
+	/** Intercepted payload type */
+	ERCIPayloadType PayloadType;
+
+	/** Property payload to intercept */
+	const TArray<uint8> Payload;
+};
