Allow registering request pre-processors after the RC Module has started.
Disable exposing properties from transient packages such as the material editor
Allow disabling web remote control from the command line.
Fix actor added to level not updating the expose picker
Fix web app not being able to display exposed properties in a packaged game

diff --git a/Engine/Plugins/VirtualProduction/RemoteControl/Source/RemoteControlUI/Private/RemoteControlUIModule.cpp b/Engine/Plugins/VirtualProduction/RemoteControl/Source/RemoteControlUI/Private/RemoteControlUIModule.cpp
--- a/Engine/Plugins/VirtualProduction/RemoteControl/Source/RemoteControlUI/Private/RemoteControlUIModule.cpp
+++ b/Engine/Plugins/VirtualProduction/RemoteControl/Source/RemoteControlUI/Private/RemoteControlUIModule.cpp
@@ -302,15 +302,24 @@ bool FRemoteControlUIModule::ShouldDisplayExposeIcon(const TSharedRef<IPropertyH
 			return false;
 		}
 
-		// Don't display an expose icon for settings
 		if (PropertyHandle->GetNumOuterObjects() == 1)
 		{
 			TArray<UObject*> OuterObjects;
 			PropertyHandle->GetOuterObjects(OuterObjects);
-		
-			if (OuterObjects[0] && OuterObjects[0]->HasAnyFlags(RF_ClassDefaultObject | RF_ArchetypeObject))
+
+			if (OuterObjects[0])
 			{
-				return false;
+				// Don't display an expose icon for default objects.
+				if (OuterObjects[0]->HasAnyFlags(RF_ClassDefaultObject | RF_ArchetypeObject))
+				{
+					return false;
+				}
+
+				// Don't display an expose icon for transient objects such as material editor parameters.
+				if (OuterObjects[0]->GetOutermost()->HasAnyFlags(RF_Transient))
+				{
+					return false;
+				}
 			}
 		}
 	}
diff --git a/Engine/Plugins/VirtualProduction/RemoteControl/Source/RemoteControlUI/Private/UI/SRemoteControlPanel.cpp b/Engine/Plugins/VirtualProduction/RemoteControl/Source/RemoteControlUI/Private/UI/SRemoteControlPanel.cpp
--- a/Engine/Plugins/VirtualProduction/RemoteControl/Source/RemoteControlUI/Private/UI/SRemoteControlPanel.cpp
+++ b/Engine/Plugins/VirtualProduction/RemoteControl/Source/RemoteControlUI/Private/UI/SRemoteControlPanel.cpp
@@ -438,8 +438,6 @@ TSharedRef<SWidget> SRemoteControlPanel::CreateCPUThrottleButton() const
 TSharedRef<SWidget> SRemoteControlPanel::CreateExposeButton()
 {	
 	FMenuBuilder MenuBuilder(true, nullptr);
-
-	
 	
 	SAssignNew(BlueprintPicker, SRCPanelFunctionPicker)
 		.AllowDefaultObjects(true)
@@ -622,9 +620,29 @@ void SRemoteControlPanel::OnActorAddedToLevel(AActor* Actor)
 		{
 			ClassPicker->Refresh();
 		}
+
+		UpdateActorFunctionPicker();
+	}
+}
+
+void SRemoteControlPanel::OnLevelActorsRemoved(AActor* Actor)
+{
+	if (Actor)
+	{
+		if (ClassPicker)
+		{
+			ClassPicker->Refresh();
+		}
+
+		UpdateActorFunctionPicker();
 	}
 }
 
+void SRemoteControlPanel::OnLevelActorListChanged()
+{
+	UpdateActorFunctionPicker();
+}
+
 void SRemoteControlPanel::CacheActorClass(AActor* Actor)
 {
 	if (RemoteControlPanelUtils::IsExposableActor(Actor))
@@ -663,6 +681,8 @@ void SRemoteControlPanel::RegisterEvents()
 	if (GEngine)
 	{
 		GEngine->OnLevelActorAdded().AddSP(this, &SRemoteControlPanel::OnActorAddedToLevel);
+		GEngine->OnLevelActorListChanged().AddSP(this, &SRemoteControlPanel::OnLevelActorListChanged);
+		GEngine->OnLevelActorDeleted().AddSP(this, &SRemoteControlPanel::OnLevelActorsRemoved);
 	}
 
 	Preset->OnEntityExposed().AddSP(this, &SRemoteControlPanel::OnEntityExposed);
@@ -676,6 +696,8 @@ void SRemoteControlPanel::UnregisterEvents()
 	
 	if (GEngine)
 	{
+		GEngine->OnLevelActorDeleted().RemoveAll(this);
+		GEngine->OnLevelActorListChanged().RemoveAll(this);
 		GEngine->OnLevelActorAdded().RemoveAll(this);
 	}
 	
@@ -911,6 +933,17 @@ FReply SRemoteControlPanel::OnClickRebindAllButton()
 	return FReply::Handled();
 }
 
+void SRemoteControlPanel::UpdateActorFunctionPicker()
+{
+	if (GEditor && ActorFunctionPicker)
+	{
+		GEditor->GetTimerManager()->SetTimerForNextTick(FTimerDelegate::CreateLambda([this]()
+		{
+			ActorFunctionPicker->Refresh();
+		}));
+	}
+}
+
 void SRemoteControlPanel::OnEntityExposed(URemoteControlPreset* InPreset, const FGuid& InEntityId)
 {
 	CachedExposedProperties.Empty();
diff --git a/Engine/Plugins/VirtualProduction/RemoteControl/Source/RemoteControlUI/Private/UI/SRemoteControlPanel.h b/Engine/Plugins/VirtualProduction/RemoteControl/Source/RemoteControlUI/Private/UI/SRemoteControlPanel.h
--- a/Engine/Plugins/VirtualProduction/RemoteControl/Source/RemoteControlUI/Private/UI/SRemoteControlPanel.h
+++ b/Engine/Plugins/VirtualProduction/RemoteControl/Source/RemoteControlUI/Private/UI/SRemoteControlPanel.h
@@ -139,8 +139,10 @@ private:
 	/** Cache the classes (and parent classes) of all actors in the level. */
 	void CacheLevelClasses();
 	
-	/** Handles updating the level class list when an actor is added to the level. */
+	//~ Handlers for various level actor events.
 	void OnActorAddedToLevel(AActor* Actor);
+	void OnLevelActorsRemoved(AActor* Actor);
+	void OnLevelActorListChanged();
 
 	/** Handles caching an actor's class and parent classes. */
 	void CacheActorClass(AActor* Actor);
@@ -163,8 +165,13 @@ private:
 	//~ Handlers called in order to clear the exposed property cache.
 	void OnEntityExposed(URemoteControlPreset* InPreset, const FGuid& InEntityId);
 	void OnEntityUnexposed(URemoteControlPreset* InPreset, const FGuid& InEntityId);
+	
 	/** Toggles the logging part of UI */
 	void OnLogCheckboxToggle(ECheckBoxState State);
+
+	/** Triggers a next frame update of the actor function picker to ensure that added actors are valid. */
+	void UpdateActorFunctionPicker();
+
 private:
 	/** Holds the preset asset. */
 	TStrongObjectPtr<URemoteControlPreset> Preset;
diff --git a/Engine/Plugins/VirtualProduction/RemoteControl/Source/WebRemoteControl/Private/RemoteControlModels.h b/Engine/Plugins/VirtualProduction/RemoteControl/Source/WebRemoteControl/Private/RemoteControlModels.h
--- a/Engine/Plugins/VirtualProduction/RemoteControl/Source/WebRemoteControl/Private/RemoteControlModels.h
+++ b/Engine/Plugins/VirtualProduction/RemoteControl/Source/WebRemoteControl/Private/RemoteControlModels.h
@@ -79,10 +79,6 @@ struct FRCObjectDescription
 		{
 			Name = Actor->GetActorLabel();
 		}
-		else if(UActorComponent* Component = Cast<UActorComponent>(InObject))
-		{
-			Name = Component->GetOwner()->GetActorLabel();
-		}
 		else
 #endif
 		{
@@ -90,12 +86,7 @@ struct FRCObjectDescription
 			Name = InObject->HasAnyFlags(RF_ClassDefaultObject | RF_ArchetypeObject) ? InObject->GetClass()->GetName() : InObject->GetName();
 		}
 		
-#if WITH_EDITOR
-		Class = InObject->GetClass()->GetDisplayNameText().ToString();
-#else
 		Class = InObject->GetClass()->GetName();
-#endif
-
 		Path = InObject->GetPathName();
 	}
 
@@ -267,7 +258,7 @@ struct FRCExposedPropertyDescription
 	{
 		UnderlyingProperty = RCProperty.GetProperty();
 		Metadata = RCProperty.GetMetadata();
-		Id = RCProperty.GetId().ToString();
+		ID = RCProperty.GetId().ToString();
 		OwnerObjects.Append(RCProperty.GetBoundObjects());
 	}
 
@@ -277,7 +268,7 @@ struct FRCExposedPropertyDescription
 
 	/** Unique identifier for the exposed property. */
 	UPROPERTY()
-	FString Id;
+	FString ID;
 
 	/** The underlying exposed property. */
 	UPROPERTY()
@@ -301,7 +292,7 @@ struct FRCExposedFunctionDescription
 
 	FRCExposedFunctionDescription(const FRemoteControlFunction& Function)
 		: DisplayName(Function.GetLabel())
-		, Id(Function.GetId().ToString())
+		, ID(Function.GetId().ToString())
 		, UnderlyingFunction(Function.GetFunction())
 	{
 		OwnerObjects.Append(Function.GetBoundObjects());
@@ -313,7 +304,7 @@ struct FRCExposedFunctionDescription
 
 	/** Unique identifier for the exposed function. */
 	UPROPERTY()
-	FString Id;
+	FString ID;
 	
 	/** The underlying exposed function. */
 	UPROPERTY()
@@ -339,7 +330,7 @@ struct FRCExposedActorDescription
 			UnderlyingActor = FRCObjectDescription{ Actor };
 		}
 
-		Id = InExposedActor.GetId().ToString();
+		ID = InExposedActor.GetId().ToString();
 	}
 
 	/** The label displayed in the remote control panel for this exposed actor. */
@@ -348,7 +339,7 @@ struct FRCExposedActorDescription
 	
 	/** Unique identifier for the exposed actor. */
 	UPROPERTY()
-	FString Id;
+	FString ID;
 
 	/** The underlying exposed actor. */
 	UPROPERTY()
@@ -498,7 +489,7 @@ struct FRCPresetDescription
 
 		Name = Preset->GetName();
 		Path = Preset->GetPathName();
-		Id = Preset->GetPresetId().ToString();
+		ID = Preset->GetPresetId().ToString();
 
 		Algo::Transform(Preset->Layout.GetGroups(), Groups, [Preset](const FRemoteControlPresetGroup& Group) { return FRCPresetLayoutGroupDescription{ Preset, Group }; });
 	}
@@ -519,7 +510,7 @@ struct FRCPresetDescription
 	 * Unique identifier for the preset, can be used to make requests to the API.
 	 */
 	UPROPERTY()
-	FString Id;
+	FString ID;
 
 	/**
 	 * The groups containing exposed entities.
@@ -538,7 +529,7 @@ struct FRCShortPresetDescription
 	FRCShortPresetDescription(const FAssetData& PresetAsset)
 	{
 		Name = PresetAsset.AssetName;
-		Id = PresetAsset.GetTagValueRef<FGuid>(FName("PresetId")).ToString();
+		ID = PresetAsset.GetTagValueRef<FGuid>(FName("PresetId")).ToString();
 		Path = PresetAsset.ObjectPath;
 	}
 
@@ -552,7 +543,7 @@ struct FRCShortPresetDescription
 	 * Unique identifier for the preset, can be used to make requests to the API.
 	 */
 	UPROPERTY()
-	FString Id;
+	FString ID;
 
 	/**
 	 * Object path of the preset.
diff --git a/Engine/Plugins/VirtualProduction/RemoteControl/Source/WebRemoteControl/Private/WebRemoteControl.cpp b/Engine/Plugins/VirtualProduction/RemoteControl/Source/WebRemoteControl/Private/WebRemoteControl.cpp
--- a/Engine/Plugins/VirtualProduction/RemoteControl/Source/WebRemoteControl/Private/WebRemoteControl.cpp
+++ b/Engine/Plugins/VirtualProduction/RemoteControl/Source/WebRemoteControl/Private/WebRemoteControl.cpp
@@ -96,6 +96,11 @@ namespace WebRemoteControl
 
 void FWebRemoteControlModule::StartupModule()
 {
+	if (FParse::Param(FCommandLine::Get(), TEXT("RCWebControlDisable")))
+	{
+		return;
+	}
+	
 #if WITH_EDITOR
 	RegisterSettings();
 #endif
@@ -125,6 +130,11 @@ void FWebRemoteControlModule::StartupModule()
 
 void FWebRemoteControlModule::ShutdownModule()
 {
+	if (FParse::Param(FCommandLine::Get(), TEXT("RCWebControlDisable")))
+	{
+		return;
+	}
+	
 	EditorRoutes.UnregisterRoutes(this);
 	WebSocketHandler->UnregisterRoutes(this);
 	StopHttpServer();
@@ -137,17 +147,33 @@ void FWebRemoteControlModule::ShutdownModule()
 
 FDelegateHandle FWebRemoteControlModule::RegisterRequestPreprocessor(FHttpRequestHandler RequestPreprocessor)
 {
-	FDelegateHandle Handle;
+	FDelegateHandle WebRCHandle{FDelegateHandle::GenerateNewHandle};
+	
+	PreprocessorsToRegister.Add(WebRCHandle, RequestPreprocessor);
+	
+	FDelegateHandle HttpRouterHandle;
 	if (HttpRouter)
 	{
-		Handle = HttpRouter->RegisterRequestPreprocessor(MoveTemp(RequestPreprocessor));
+		HttpRouterHandle = HttpRouter->RegisterRequestPreprocessor(MoveTemp(RequestPreprocessor));
 	}
-	return Handle;
+
+	PreprocessorsHandleMappings.Add(WebRCHandle, HttpRouterHandle);
+	
+	return WebRCHandle;
 }
 
 void FWebRemoteControlModule::UnregisterRequestPreprocessor(const FDelegateHandle& RequestPreprocessorHandle)
 {
-	HttpRouter->UnregisterRequestPreprocessor(RequestPreprocessorHandle);
+	PreprocessorsToRegister.Remove(RequestPreprocessorHandle);
+	if (FDelegateHandle* HttpRouterHandle = PreprocessorsHandleMappings.Find(RequestPreprocessorHandle))
+	{
+		if (HttpRouterHandle->IsValid())
+		{
+			HttpRouter->UnregisterRequestPreprocessor(*HttpRouterHandle);
+		}
+		
+		PreprocessorsHandleMappings.Remove(RequestPreprocessorHandle);
+	}
 }
 
 void FWebRemoteControlModule::RegisterRoute(const FRemoteControlRoute& Route)
@@ -201,6 +227,20 @@ void FWebRemoteControlModule::StartHttpServer()
 			StartRoute(Route);
 		}
 
+		// Go through externally registered request pre-processors and register them with the http router.
+		for (const TPair<FDelegateHandle, FHttpRequestHandler>& Handler : PreprocessorsToRegister)
+		{
+			// Find the pre-processors HTTP-handle from the one we generated.
+			FDelegateHandle& Handle = PreprocessorsHandleMappings.FindChecked(Handler.Key);
+			if (Handle.IsValid())
+			{
+				HttpRouter->UnregisterRequestPreprocessor(Handle);
+			}
+
+			// Update the preprocessor handle mapping.
+			Handle = HttpRouter->RegisterRequestPreprocessor(Handler.Value);
+		}
+
 		FHttpServerModule::Get().StartAllListeners();
 
 		OnHttpServerStartedDelegate.Broadcast(HttpServerPort);
diff --git a/Engine/Plugins/VirtualProduction/RemoteControl/Source/WebRemoteControl/Private/WebRemoteControl.h b/Engine/Plugins/VirtualProduction/RemoteControl/Source/WebRemoteControl/Private/WebRemoteControl.h
--- a/Engine/Plugins/VirtualProduction/RemoteControl/Source/WebRemoteControl/Private/WebRemoteControl.h
+++ b/Engine/Plugins/VirtualProduction/RemoteControl/Source/WebRemoteControl/Private/WebRemoteControl.h
@@ -176,6 +176,14 @@ private:
 	/** Holds the client currently making a request. */
 	FGuid ActingClientId;
 
+	/** List of preprocessor delegates that need to be registered when the server is started. */
+	TMap<FDelegateHandle, FHttpRequestHandler> PreprocessorsToRegister;
+
+	/**
+	 * Mappings of preprocessors delegate handles generated from the WebRC module to the ones generated from the Http Module.
+	 */
+	TMap<FDelegateHandle, FDelegateHandle> PreprocessorsHandleMappings;
+
 	//~ Server started stopped delegates.
 	FOnWebServerStarted OnHttpServerStartedDelegate;
 	FSimpleMulticastDelegate OnHttpServerStoppedDelegate;
diff --git a/Engine/Plugins/VirtualProduction/RemoteControl/Source/WebRemoteControl/Private/WebSocketMessageHandler.cpp b/Engine/Plugins/VirtualProduction/RemoteControl/Source/WebRemoteControl/Private/WebSocketMessageHandler.cpp
--- a/Engine/Plugins/VirtualProduction/RemoteControl/Source/WebRemoteControl/Private/WebSocketMessageHandler.cpp
+++ b/Engine/Plugins/VirtualProduction/RemoteControl/Source/WebRemoteControl/Private/WebSocketMessageHandler.cpp
@@ -385,7 +385,7 @@ void FWebSocketMessageHandler::ProcessAddedProperties()
 		FRCPresetDescription AddedPropertiesDescription;
 		AddedPropertiesDescription.Name = Preset->GetName();
 		AddedPropertiesDescription.Path = Preset->GetPathName();
-		AddedPropertiesDescription.Id = Preset->GetPresetId().ToString();
+		AddedPropertiesDescription.ID = Preset->GetPresetId().ToString();
 
 		TMap<FRemoteControlPresetGroup*, TArray<FGuid>> GroupedNewFields;
 
diff --git a/Engine/Plugins/VirtualProduction/RemoteControlWebInterface/WebApp/Client/src/shared/index.ts b/Engine/Plugins/VirtualProduction/RemoteControlWebInterface/WebApp/Client/src/shared/index.ts
--- a/Engine/Plugins/VirtualProduction/RemoteControlWebInterface/WebApp/Client/src/shared/index.ts
+++ b/Engine/Plugins/VirtualProduction/RemoteControlWebInterface/WebApp/Client/src/shared/index.ts
@@ -68,7 +68,7 @@ export interface IFunction {
 }
 
 export interface IExposedFunction {
-  Id: string;
+  ID: string;
   DisplayName: string;
   UnderlyingFunction: IFunction;
   Metadata: { [key: string]: string };
@@ -82,7 +82,7 @@ export interface IProperty {
 }
 
 export interface IExposedProperty {
-  Id: string;
+  ID: string;
   DisplayName: string;
   Metadata: Record<string, string>;
   Widget: WidgetType;
diff --git a/Engine/Plugins/VirtualProduction/RemoteControlWebInterface/WebApp/Server/src/UnrealEngine.ts b/Engine/Plugins/VirtualProduction/RemoteControlWebInterface/WebApp/Server/src/UnrealEngine.ts
--- a/Engine/Plugins/VirtualProduction/RemoteControlWebInterface/WebApp/Server/src/UnrealEngine.ts
+++ b/Engine/Plugins/VirtualProduction/RemoteControlWebInterface/WebApp/Server/src/UnrealEngine.ts
@@ -150,8 +150,8 @@ export namespace UnrealEngine {
             if (!property)
               continue;
 
-            setPayloadValueInternal(payloads, [message.PresetName, property.Id], field.PropertyValue);
-            Notify.emitValueChange(message.PresetName, property.Id, field.PropertyValue);
+            setPayloadValueInternal(payloads, [message.PresetName, property.ID], field.PropertyValue);
+            Notify.emitValueChange(message.PresetName, property.ID, field.PropertyValue);
           }
           break;
         }
@@ -281,9 +281,9 @@ export namespace UnrealEngine {
 
       for (const Group of Preset.Groups) {
         for (const Function of Group.ExposedFunctions)
-          Preset.Exposed[Function.Id] = Function;
+          Preset.Exposed[Function.ID] = Function;
         for (const Property of Group.ExposedProperties)
-          Preset.Exposed[Property.Id] = Property;
+          Preset.Exposed[Property.ID] = Property;
 
         Preset.ExposedProperties.push(...Group.ExposedProperties);
         Preset.ExposedFunctions.push(...Group.ExposedFunctions);
@@ -395,8 +395,9 @@ export namespace UnrealEngine {
     for (const Preset of presets) {
       for (const group of Preset.Groups) {
         for (const property of group.ExposedProperties) {
-           const value = await get<UnrealApi.PropertyValues>(`/remote/preset/${Preset.Name}/property/${property.Id}`);
-           setPayloadValueInternal(updatedPayloads, [Preset.Name, property.Id], value?.PropertyValues?.[0]?.PropertyValue);
+           const value = await get<UnrealApi.PropertyValues>(`/remote/preset/${Preset.Name}/property/${property.ID}`);
+           console.log(value);
+           setPayloadValueInternal(updatedPayloads, [Preset.Name, property.ID], value?.PropertyValues?.[0]?.PropertyValue);
         }
       }
     }
diff --git a/Engine/Plugins/VirtualProduction/RemoteControlWebInterface/WebApp/package.json b/Engine/Plugins/VirtualProduction/RemoteControlWebInterface/WebApp/package.json
--- a/Engine/Plugins/VirtualProduction/RemoteControlWebInterface/WebApp/package.json
+++ b/Engine/Plugins/VirtualProduction/RemoteControlWebInterface/WebApp/package.json
@@ -1,6 +1,6 @@
 {
   "name": "conductor",
-  "version": "0.3.8",
+  "version": "0.3.9",
   "scripts": {
     "postinstall": "concurrently \"cd ./Client && npm install && cd ..\" \"cd ./Server && npm install && cd ..\"",
     "build": "node ./Scripts/build.js",
