Fix BP recompile not updating the preset

diff --git a/Engine/Plugins/VirtualProduction/RemoteControl/Source/RemoteControl/Private/RemoteControlEntity.cpp b/Engine/Plugins/VirtualProduction/RemoteControl/Source/RemoteControl/Private/RemoteControlEntity.cpp
--- a/Engine/Plugins/VirtualProduction/RemoteControl/Source/RemoteControl/Private/RemoteControlEntity.cpp
+++ b/Engine/Plugins/VirtualProduction/RemoteControl/Source/RemoteControl/Private/RemoteControlEntity.cpp
@@ -87,6 +87,15 @@ FRemoteControlEntity::FRemoteControlEntity(URemoteControlPreset* InPreset, FName
 	Bindings.Append(InBindings);
 }
 
+const UScriptStruct* FRemoteControlEntity::GetStruct() const
+{
+	if (URemoteControlPreset* Preset = Owner.Get())
+	{
+		return Preset->GetExposedEntityType(Id);
+	}
+	return nullptr;
+}
+
 FName FRemoteControlEntity::Rename(FName NewLabel)
 {
 	if (URemoteControlPreset* Preset = Owner.Get())
diff --git a/Engine/Plugins/VirtualProduction/RemoteControl/Source/RemoteControl/Private/RemoteControlField.cpp b/Engine/Plugins/VirtualProduction/RemoteControl/Source/RemoteControl/Private/RemoteControlField.cpp
--- a/Engine/Plugins/VirtualProduction/RemoteControl/Source/RemoteControl/Private/RemoteControlField.cpp
+++ b/Engine/Plugins/VirtualProduction/RemoteControl/Source/RemoteControl/Private/RemoteControlField.cpp
@@ -305,6 +305,7 @@ FRemoteControlFunction::FRemoteControlFunction(URemoteControlPreset* InPreset, F
 	check(InFunction);
 	FunctionArguments = MakeShared<FStructOnScope>(InFunction);
 	InFunction->InitializeStruct(FunctionArguments->GetStructMemory());
+	CachedFunction = InFunction;
 	AssignDefaultFunctionArguments();
 	OwnerClass = GetSupportedBindingClass();
 
@@ -449,4 +450,4 @@ FArchive& operator<<(FArchive& Ar, FRemoteControlFunction& RCFunction)
 	}
 		
 	return Ar;
-}
\ No newline at end of file
+}
diff --git a/Engine/Plugins/VirtualProduction/RemoteControl/Source/RemoteControl/Private/RemoteControlPreset.cpp b/Engine/Plugins/VirtualProduction/RemoteControl/Source/RemoteControl/Private/RemoteControlPreset.cpp
--- a/Engine/Plugins/VirtualProduction/RemoteControl/Source/RemoteControl/Private/RemoteControlPreset.cpp
+++ b/Engine/Plugins/VirtualProduction/RemoteControl/Source/RemoteControl/Private/RemoteControlPreset.cpp
@@ -22,6 +22,7 @@
 
 #if WITH_EDITOR
 #include "Editor.h"
+#include "Engine/Blueprint.h"
 #include "TimerManager.h"
 #endif
 
@@ -777,7 +778,11 @@ TWeakPtr<FRemoteControlFunction> URemoteControlPreset::ExposeFunction(UObject* O
 	}
 
 	FRemoteControlFunction RCFunction{ this, Registry->GenerateUniqueLabel(DesiredName), Function->GetName(), Function, { FindOrAddBinding(Object) } };
-	return StaticCastSharedPtr<FRemoteControlFunction>(Expose(MoveTemp(RCFunction), FRemoteControlFunction::StaticStruct(), Args.GroupId));
+	TSharedPtr<FRemoteControlFunction> RCFunctionPtr = StaticCastSharedPtr<FRemoteControlFunction>(Expose(MoveTemp(RCFunction), FRemoteControlFunction::StaticStruct(), Args.GroupId));
+
+	RegisterOnCompileEvent(RCFunctionPtr);
+	
+	return RCFunctionPtr;
 }
 
 TSharedPtr<FRemoteControlEntity> URemoteControlPreset::Expose(FRemoteControlEntity&& Entity, UScriptStruct* EntityType, const FGuid& GroupId)
@@ -902,9 +907,34 @@ void URemoteControlPreset::RegisterEntityDelegates()
 	for (const TSharedPtr<FRemoteControlEntity>& Entity : Registry->GetExposedEntities())
 	{
 		Entity->OnEntityModifiedDelegate.BindUObject(this, &URemoteControlPreset::OnEntityModified);
+
+		if (Entity->GetStruct() == FRemoteControlFunction::StaticStruct())
+		{
+			RegisterOnCompileEvent(StaticCastSharedPtr<FRemoteControlFunction>(Entity));
+		}
 	}
 }
 
+void URemoteControlPreset::RegisterOnCompileEvent(const TSharedPtr<FRemoteControlFunction>& RCFunction)
+{
+#if WITH_EDITOR
+	if (UFunction* UnderlyingFunction = RCFunction->GetFunction())
+	{
+		if (UBlueprintGeneratedClass* BPClass = Cast<UBlueprintGeneratedClass>(UnderlyingFunction->GetOwnerClass()))
+		{
+			if (UBlueprint* Blueprint = Cast<UBlueprint>(BPClass->ClassGeneratedBy))
+			{
+				if (!BlueprintsWithRegisteredDelegates.Contains(Blueprint))
+				{
+					BlueprintsWithRegisteredDelegates.Emplace(Blueprint);
+					Blueprint->OnCompiled().AddUObject(this, &URemoteControlPreset::OnBlueprintRecompiled);
+				}
+			}
+		}
+	}
+#endif
+}
+
 TOptional<FRemoteControlFunction> URemoteControlPreset::GetFunction(FName FunctionLabel) const
 {
 PRAGMA_DISABLE_DEPRECATION_WARNINGS
@@ -1573,7 +1603,10 @@ void URemoteControlPreset::RegisterDelegates()
 	FCoreUObjectDelegates::OnObjectPropertyChanged.AddUObject(this, &URemoteControlPreset::OnObjectPropertyChanged);
 	FCoreUObjectDelegates::OnPreObjectPropertyChanged.AddUObject(this, &URemoteControlPreset::OnPreObjectPropertyChanged);
 
-	GEngine->OnLevelActorDeleted().AddUObject(this, &URemoteControlPreset::OnActorDeleted);
+	if (GEngine)
+	{
+		GEngine->OnLevelActorDeleted().AddUObject(this, &URemoteControlPreset::OnActorDeleted);
+	}
 	FEditorDelegates::PostPIEStarted.AddUObject(this, &URemoteControlPreset::OnPieEvent);
 	FEditorDelegates::EndPIE.AddUObject(this, &URemoteControlPreset::OnPieEvent);
 	
@@ -1588,6 +1621,14 @@ void URemoteControlPreset::RegisterDelegates()
 void URemoteControlPreset::UnregisterDelegates()
 {
 #if WITH_EDITOR
+	for (TWeakObjectPtr<UBlueprint> Blueprint : BlueprintsWithRegisteredDelegates)
+	{
+		if (Blueprint.IsValid())
+		{
+			Blueprint->OnCompiled().RemoveAll(this);
+		}
+	}
+	
 	FEditorDelegates::MapChange.RemoveAll(this);
 
 	FCoreDelegates::OnEndFrame.RemoveAll(this);
@@ -1703,11 +1744,45 @@ void URemoteControlPreset::OnEndFrame()
 
 void URemoteControlPreset::OnMapChange(uint32)
 {
-	// Delay the refresh in order for the old actor points to be invalid.
+	// Delay the refresh in order for the old actors to be invalid.
 	GEditor->GetTimerManager()->SetTimerForNextTick(FTimerDelegate::CreateLambda([this]()
+	{
+		Algo::Transform(Registry->GetExposedEntities(), PerFrameUpdatedEntities, [](const TSharedPtr<FRemoteControlEntity>& Entity) { return Entity->GetId(); });
+	}));
+}
+
+void URemoteControlPreset::OnBlueprintRecompiled(UBlueprint* Blueprint)
+{
+	if (!Blueprint)
+	{
+		return;
+	}
+
+	for (TSharedPtr<FRemoteControlFunction> RCFunction : Registry->GetExposedEntities<FRemoteControlFunction>())
+	{
+		if (UClass* Class = RCFunction->GetSupportedBindingClass())
 		{
-			Algo::Transform(Registry->GetExposedEntities(), PerFrameUpdatedEntities, [](const TSharedPtr<FRemoteControlEntity>& Entity) { return Entity->GetId(); });
-		}));
+			if (Class->ClassGeneratedBy == Blueprint)
+			{
+				// Recreate the function arguments with the function from the new BP and copy the old ones on top of it.
+				if (UFunction* OldFunction = RCFunction->GetFunction())
+				{
+					UClass* NewClass = Blueprint->GeneratedClass;
+					if (UFunction* NewFunction = NewClass->FindFunctionByName(OldFunction->GetFName()))
+					{
+						FStructOnScope NewFunctionOnScope{ NewFunction };
+						for (TFieldIterator<FProperty> It(OldFunction); It; ++It)
+						{
+							It->CopyCompleteValue_InContainer(NewFunctionOnScope.GetStructMemory(), RCFunction->FunctionArguments->GetStructMemory());
+						}
+						*RCFunction->FunctionArguments = MoveTemp(NewFunctionOnScope);
+
+						PerFrameUpdatedEntities.Add(RCFunction->GetId());
+					}
+				}
+			}
+		}
+	}
 }
 
 #endif
diff --git a/Engine/Plugins/VirtualProduction/RemoteControl/Source/RemoteControl/Public/RemoteControlEntity.h b/Engine/Plugins/VirtualProduction/RemoteControl/Source/RemoteControl/Public/RemoteControlEntity.h
--- a/Engine/Plugins/VirtualProduction/RemoteControl/Source/RemoteControl/Public/RemoteControlEntity.h
+++ b/Engine/Plugins/VirtualProduction/RemoteControl/Source/RemoteControl/Public/RemoteControlEntity.h
@@ -21,6 +21,9 @@ struct REMOTECONTROL_API FRemoteControlEntity
 
 	virtual ~FRemoteControlEntity(){}
 
+	/** Get the type of this remote control entity. */
+	const UScriptStruct* GetStruct() const;
+
 	/**
 	 * Change this entity's label.
 	 * @param NewLabel the desired label.
diff --git a/Engine/Plugins/VirtualProduction/RemoteControl/Source/RemoteControl/Public/RemoteControlField.h b/Engine/Plugins/VirtualProduction/RemoteControl/Source/RemoteControl/Public/RemoteControlField.h
--- a/Engine/Plugins/VirtualProduction/RemoteControl/Source/RemoteControl/Public/RemoteControlField.h
+++ b/Engine/Plugins/VirtualProduction/RemoteControl/Source/RemoteControl/Public/RemoteControlField.h
@@ -207,6 +207,10 @@ public:
 	 */
 	TSharedPtr<class FStructOnScope> FunctionArguments;
 
+private:
+	/** Parse function metadata to get the function's default parameters */
+	void AssignDefaultFunctionArguments();
+
 private:
 	/** Whether the function is callable in a packaged build. */
 	UPROPERTY()
@@ -219,9 +223,6 @@ private:
 	/** Cached resolved underlying function used to avoid doing a findobject while serializing. */
 	mutable TWeakObjectPtr<UFunction> CachedFunction;
 
-private:
-	/** Parse function metadata to get the function`s default parameters */
-	void AssignDefaultFunctionArguments();
 };
 
 template<> struct TStructOpsTypeTraits<FRemoteControlFunction> : public TStructOpsTypeTraitsBase2<FRemoteControlFunction>
diff --git a/Engine/Plugins/VirtualProduction/RemoteControl/Source/RemoteControl/Public/RemoteControlPreset.h b/Engine/Plugins/VirtualProduction/RemoteControl/Source/RemoteControl/Public/RemoteControlPreset.h
--- a/Engine/Plugins/VirtualProduction/RemoteControl/Source/RemoteControl/Public/RemoteControlPreset.h
+++ b/Engine/Plugins/VirtualProduction/RemoteControl/Source/RemoteControl/Public/RemoteControlPreset.h
@@ -9,6 +9,7 @@
 #include "UObject/SoftObjectPtr.h"
 #include "Templates/PimplPtr.h"
 #include "Templates/UnrealTypeTraits.h"
+
 #include "RemoteControlPreset.generated.h"
 
 class AActor;
@@ -18,6 +19,7 @@ struct FRCFieldPathInfo;
 struct FRemoteControlActor;
 struct FRemoteControlPresetLayout;
 class FRemoteControlPresetRebindingManager;
+class UBlueprint;
 class URemoteControlExposeRegistry;
 class URemoteControlBinding;
 class URemoteControlPreset;
@@ -780,6 +782,7 @@ private:
 	void OnReplaceObjects(const TMap<UObject*, UObject*>& ReplacementObjectMap);
 	void OnEndFrame();
 	void OnMapChange(uint32 /*MapChangeFlags*/);
+	void OnBlueprintRecompiled(UBlueprint* Blueprint);
 #endif
 	
 	/** Get a field ptr using it's id. */
@@ -813,6 +816,9 @@ private:
 
 	/** Register delegates for all exposed entities. */
 	void RegisterEntityDelegates();
+
+	/** Register an event triggered when the exposed function's owner blueprint is compiled. */
+	void RegisterOnCompileEvent(const TSharedPtr<FRemoteControlFunction>& RCFunction);
 	
 private:
 	/** Preset unique ID */
@@ -873,6 +879,11 @@ private:
 	/** Holds manager that handles rebinding unbound entities upon load or map change. */
 	TPimplPtr<FRemoteControlPresetRebindingManager> RebindingManager;
 
+#if WITH_EDITOR
+	/** List of blueprints for which we have registered events. */
+	TSet<TWeakObjectPtr<UBlueprint>> BlueprintsWithRegisteredDelegates;
+#endif
+
 	friend FRemoteControlTarget;
 	friend FRemoteControlPresetLayout;
 	friend FRemoteControlEntity;
