RemoteControlUI layout persistence

diff --git a/Engine/Plugins/VirtualProduction/RemoteControl/Source/RemoteControlUI/Private/RemoteControlUISettings.h b/Engine/Plugins/VirtualProduction/RemoteControl/Source/RemoteControlUI/Private/RemoteControlUISettings.h
--- a/Engine/Plugins/VirtualProduction/RemoteControl/Source/RemoteControlUI/Private/RemoteControlUISettings.h
+++ b/Engine/Plugins/VirtualProduction/RemoteControl/Source/RemoteControlUI/Private/RemoteControlUISettings.h
@@ -22,4 +22,8 @@ public:
 	/** Show a warning icon for exposed editor-only fields. */
 	UPROPERTY(config, EditAnywhere, Category = "Remote Control Preset", DisplayName = "Show a warning when exposing editor-only entities.")
 	bool bDisplayInEditorOnlyWarnings = false;
+
+	/** The split widget control ratio between entity tree and details/protocol binding list. */
+	UPROPERTY(config)
+	float TreeBindingSplitRatio = 0.7;
 };
\ No newline at end of file
diff --git a/Engine/Plugins/VirtualProduction/RemoteControl/Source/RemoteControlUI/Private/UI/SRemoteControlPanel.cpp b/Engine/Plugins/VirtualProduction/RemoteControl/Source/RemoteControlUI/Private/UI/SRemoteControlPanel.cpp
--- a/Engine/Plugins/VirtualProduction/RemoteControl/Source/RemoteControlUI/Private/UI/SRemoteControlPanel.cpp
+++ b/Engine/Plugins/VirtualProduction/RemoteControl/Source/RemoteControlUI/Private/UI/SRemoteControlPanel.cpp
@@ -29,6 +29,7 @@
 #include "RemoteControlPanelStyle.h"
 #include "RemoteControlPreset.h"
 #include "RemoteControlUIModule.h"
+#include "RemoteControlUISettings.h"
 #include "ScopedTransaction.h"
 #include "SClassViewer.h"
 #include "SRCPanelExposedEntitiesList.h"
@@ -108,7 +109,21 @@ void SRemoteControlPanel::Construct(const FArguments& InArgs, URemoteControlPres
 		.EditMode_Lambda([this](){ return bIsInEditMode; });
 	
 	EntityList->OnSelectionChange().AddSP(this, &SRemoteControlPanel::UpdateEntityDetailsView);
-	
+
+	const TAttribute<float> TreeBindingSplitRatioTop = TAttribute<float>::Create(
+		TAttribute<float>::FGetter::CreateLambda([]()
+		{
+			URemoteControlUISettings* Settings = GetMutableDefault<URemoteControlUISettings>();
+			return Settings->TreeBindingSplitRatio;
+		}));
+
+	const TAttribute<float> TreeBindingSplitRatioBottom = TAttribute<float>::Create(
+		TAttribute<float>::FGetter::CreateLambda([]()
+		{
+			URemoteControlUISettings* Settings = GetMutableDefault<URemoteControlUISettings>();
+			return 1.0f - Settings->TreeBindingSplitRatio;
+		}));
+
 	ChildSlot
 	[
 		SNew(SVerticalBox)
@@ -227,9 +242,14 @@ void SRemoteControlPanel::Construct(const FArguments& InArgs, URemoteControlPres
 					+ SWidgetSwitcher::Slot()
 					[
 						SNew(SSplitter)
-						.Orientation(EOrientation::Orient_Vertical)
+						.Orientation(EOrientation::Orient_Vertical)						
 						+ SSplitter::Slot()
-						.Value(0.7f)
+						.Value(TreeBindingSplitRatioTop)
+						.OnSlotResized(SSplitter::FOnSlotResized::CreateLambda([](float InNewSize)
+						{
+							URemoteControlUISettings* Settings = GetMutableDefault<URemoteControlUISettings>();
+							Settings->TreeBindingSplitRatio = InNewSize;
+						}))
 						[
 							// Exposed entities List
 							SNew(SBorder)
@@ -239,21 +259,25 @@ void SRemoteControlPanel::Construct(const FArguments& InArgs, URemoteControlPres
 							]
 						]
 						+ SSplitter::Slot()
-						.Value(0.3f)
+						.Value(TreeBindingSplitRatioBottom)
 						[
 							SNew(SSplitter)
 							.Orientation(Orient_Vertical)
+							.HitDetectionSplitterHandleSize(0.0f) // Effectively disables user manipulation (but keeps visuals)
 							+ SSplitter::Slot()
-							.Value(0.4f)
+							.SizeRule(SSplitter::SizeToContent)
 							[
 								SNew(SBorder)
 								.BorderImage(FEditorStyle::GetBrush("ToolPanel.GroupBorder"))
 								[
-									CreateEntityDetailsView()
+									SNew(SBox)
+									.MinDesiredHeight(40.0f)
+									[
+										CreateEntityDetailsView()
+									]
 								]
 							]
 							+ SSplitter::Slot()
-							.Value(0.6f)
 							[
 								SNew(SBorder)
 								.BorderImage(FEditorStyle::GetBrush("ToolPanel.GroupBorder"))
@@ -847,6 +871,7 @@ TSharedRef<SWidget> SRemoteControlPanel::CreateEntityDetailsView()
 	{
 		return EntityDetailsView->GetWidget().ToSharedRef();
 	}
+
 	return SNullWidget::NullWidget;
 }
 
