RC fix protocol replication in -Game

diff --git a/Engine/Plugins/VirtualProduction/RemoteControl/Source/RemoteControl/Private/RemoteControlPreset.cpp b/Engine/Plugins/VirtualProduction/RemoteControl/Source/RemoteControl/Private/RemoteControlPreset.cpp
--- a/Engine/Plugins/VirtualProduction/RemoteControl/Source/RemoteControl/Private/RemoteControlPreset.cpp
+++ b/Engine/Plugins/VirtualProduction/RemoteControl/Source/RemoteControl/Private/RemoteControlPreset.cpp
@@ -26,6 +26,7 @@
 #include "TimerManager.h"
 #endif
 
+URemoteControlPreset::FOnPostLoadRemoteControlPreset URemoteControlPreset::OnPostLoadRemoteControlPreset;
 
 #define LOCTEXT_NAMESPACE "RemoteControlPreset" 
 
@@ -655,6 +656,8 @@ void URemoteControlPreset::PostLoad()
 {
 	Super::PostLoad();
 
+	OnPostLoadRemoteControlPreset.Broadcast(this);
+
 	RegisterDelegates();
 
 	CacheFieldsData();
diff --git a/Engine/Plugins/VirtualProduction/RemoteControl/Source/RemoteControl/Public/RemoteControlField.h b/Engine/Plugins/VirtualProduction/RemoteControl/Source/RemoteControl/Public/RemoteControlField.h
--- a/Engine/Plugins/VirtualProduction/RemoteControl/Source/RemoteControl/Public/RemoteControlField.h
+++ b/Engine/Plugins/VirtualProduction/RemoteControl/Source/RemoteControl/Public/RemoteControlField.h
@@ -77,7 +77,7 @@ public:
 	 * The map holds protocol bindings stores the protocol mapping and protocol-specific mapping
 	 */
 	UPROPERTY()
-	TSet<FRemoteControlProtocolBinding> ProtocolBinding;
+	TSet<FRemoteControlProtocolBinding> ProtocolBindings;
 	
 protected:
 	/**
diff --git a/Engine/Plugins/VirtualProduction/RemoteControl/Source/RemoteControl/Public/RemoteControlPreset.h b/Engine/Plugins/VirtualProduction/RemoteControl/Source/RemoteControl/Public/RemoteControlPreset.h
--- a/Engine/Plugins/VirtualProduction/RemoteControl/Source/RemoteControl/Public/RemoteControlPreset.h
+++ b/Engine/Plugins/VirtualProduction/RemoteControl/Source/RemoteControl/Public/RemoteControlPreset.h
@@ -439,7 +439,11 @@ class REMOTECONTROL_API URemoteControlPreset : public UObject
 {
 public:
 	GENERATED_BODY()
-
+	
+	/** Callback for post remote control preset load, called by URemoteControlPreset::PostLoad function */
+	DECLARE_MULTICAST_DELEGATE_OneParam(FOnPostLoadRemoteControlPreset, URemoteControlPreset* /* InPreset */);
+	static FOnPostLoadRemoteControlPreset OnPostLoadRemoteControlPreset;
+	
 	URemoteControlPreset();
 
 	//~ Begin UObject interface
diff --git a/Engine/Plugins/VirtualProduction/RemoteControl/Source/RemoteControlProtocol/Private/RemoteControlProtocolModule.cpp b/Engine/Plugins/VirtualProduction/RemoteControl/Source/RemoteControlProtocol/Private/RemoteControlProtocolModule.cpp
--- a/Engine/Plugins/VirtualProduction/RemoteControl/Source/RemoteControlProtocol/Private/RemoteControlProtocolModule.cpp
+++ b/Engine/Plugins/VirtualProduction/RemoteControl/Source/RemoteControlProtocol/Private/RemoteControlProtocolModule.cpp
@@ -1,15 +1,34 @@
 // Copyright Epic Games, Inc. All Rights Reserved.
 
 #include "RemoteControlProtocolModule.h"
+
+#include "RemoteControlPreset.h"
 #include "RemoteControlProtocol.h"
 
+#include "Misc/CommandLine.h"
+
 DEFINE_LOG_CATEGORY(LogRemoteControlProtocol);
 
 #define LOCTEXT_NAMESPACE "RemoteControlProtocol"
 
+void FRemoteControlProtocolModule::StartupModule()
+{
+	bRCProtocolsDisable = FParse::Param(FCommandLine::Get(), TEXT("RCProtocolsDisable"));
+
+	if (!bRCProtocolsDisable)
+	{
+		// Called when remote control preset loaded
+		URemoteControlPreset::OnPostLoadRemoteControlPreset.AddRaw(this, &FRemoteControlProtocolModule::OnPostLoadRemoteControlPreset);
+	}
+}
+
 void FRemoteControlProtocolModule::ShutdownModule()
 {
-	EmptyProtocols();
+	if (!bRCProtocolsDisable)
+	{
+		EmptyProtocols();
+		URemoteControlPreset::OnPostLoadRemoteControlPreset.RemoveAll(this);
+	}
 }
 
 TArray<FName> FRemoteControlProtocolModule::GetProtocolNames() const
@@ -29,14 +48,21 @@ TSharedPtr<IRemoteControlProtocol> FRemoteControlProtocolModule::GetProtocolByNa
 	return nullptr;
 }
 
-void FRemoteControlProtocolModule::AddProtocol(FName InProtocolName, TSharedRef<IRemoteControlProtocol> InProtocol)
+bool FRemoteControlProtocolModule::AddProtocol(FName InProtocolName, TSharedRef<IRemoteControlProtocol> InProtocol)
 {
+	if (bRCProtocolsDisable)
+	{
+		return false;
+	}
+
 	InProtocol->Init();
 
 	// Check for presence of required RangeInputTemplate property
 	checkf(InProtocol->GetRangeInputTemplateProperty(), TEXT("ScriptStruct for Protocol %s did not have the required RangeInputTemplate Property."), *InProtocolName.ToString());
 	
 	Protocols.Add(InProtocolName, MoveTemp(InProtocol));
+
+	return true;
 }
 
 void FRemoteControlProtocolModule::RemoveProtocol(FName InProtocolName, TSharedRef<IRemoteControlProtocol> InProtocol)
@@ -59,6 +85,35 @@ void FRemoteControlProtocolModule::EmptyProtocols()
 	Protocols.Empty();
 }
 
+void FRemoteControlProtocolModule::OnPostLoadRemoteControlPreset(URemoteControlPreset* InPreset) const
+{
+	if (!ensure(InPreset))
+	{
+		return;
+	}
+	
+	
+	for(TWeakPtr<FRemoteControlProperty> ExposedPropertyWeakPtr : InPreset->GetExposedEntities<FRemoteControlProperty>())
+	{
+		if (TSharedPtr<FRemoteControlProperty> ExposedPropertyPtr = ExposedPropertyWeakPtr.Pin())
+		{
+			for(FRemoteControlProtocolBinding& Binding : ExposedPropertyPtr->ProtocolBindings)
+			{
+				const TSharedPtr<IRemoteControlProtocol> Protocol = GetProtocolByName(Binding.GetProtocolName());
+				// Supporting plugin needs to be loaded/protocol available.
+				if(Protocol.IsValid())
+				{
+					Protocol->Bind(Binding.GetRemoteControlProtocolEntityPtr());
+				}
+			}	
+		}
+		else
+		{
+			ensure(false);
+		}
+	}
+}
+
 
 IMPLEMENT_MODULE(FRemoteControlProtocolModule, RemoteControlProtocol);
 
diff --git a/Engine/Plugins/VirtualProduction/RemoteControl/Source/RemoteControlProtocol/Private/RemoteControlProtocolModule.h b/Engine/Plugins/VirtualProduction/RemoteControl/Source/RemoteControlProtocol/Private/RemoteControlProtocolModule.h
--- a/Engine/Plugins/VirtualProduction/RemoteControl/Source/RemoteControlProtocol/Private/RemoteControlProtocolModule.h
+++ b/Engine/Plugins/VirtualProduction/RemoteControl/Source/RemoteControlProtocol/Private/RemoteControlProtocolModule.h
@@ -4,12 +4,15 @@
 
 #include "IRemoteControlProtocolModule.h"
 
+class URemoteControlPreset;
+
 REMOTECONTROLPROTOCOL_API DECLARE_LOG_CATEGORY_EXTERN(LogRemoteControlProtocol, Log, All);
 
 class FRemoteControlProtocolModule : public IRemoteControlProtocolModule
 {
 public:
 	//~ Begin IModuleInterface interface
+	virtual void StartupModule() override;
 	virtual void ShutdownModule() override;
 	//~ End IModuleInterface interface
 	
@@ -17,12 +20,20 @@ public:
 	virtual int32 GetProtocolNum() const override { return Protocols.Num(); }
 	virtual TArray<FName> GetProtocolNames() const override;
 	virtual TSharedPtr<IRemoteControlProtocol> GetProtocolByName(FName InProtocolName) const override;
-	virtual void AddProtocol(FName InProtocolName, TSharedRef<IRemoteControlProtocol> InProtocol) override;
+	virtual bool AddProtocol(FName InProtocolName, TSharedRef<IRemoteControlProtocol> InProtocol) override;
 	virtual void RemoveProtocol(FName InProtocolName, TSharedRef<IRemoteControlProtocol> InProtocol) override;
 	virtual void EmptyProtocols() override;
+	virtual bool IsRCProtocolsDisable() const override { return bRCProtocolsDisable; };
 	//~ End IRemoteControlProtocolModule interface
 
+private:
+	/** Called when remote control preset loaded */
+	void OnPostLoadRemoteControlPreset(URemoteControlPreset* InPreset) const;
+
 private:
 	/** Map with all protocols */
 	TMap<FName, TSharedRef<IRemoteControlProtocol>> Protocols;
+
+	/** Prevents the Remote Control protocols from starting. Can be set from command line with -RCProtocolsDisable */
+	bool bRCProtocolsDisable = false;
 };
diff --git a/Engine/Plugins/VirtualProduction/RemoteControl/Source/RemoteControlProtocol/Public/IRemoteControlProtocolModule.h b/Engine/Plugins/VirtualProduction/RemoteControl/Source/RemoteControlProtocol/Public/IRemoteControlProtocolModule.h
--- a/Engine/Plugins/VirtualProduction/RemoteControl/Source/RemoteControlProtocol/Public/IRemoteControlProtocolModule.h
+++ b/Engine/Plugins/VirtualProduction/RemoteControl/Source/RemoteControlProtocol/Public/IRemoteControlProtocolModule.h
@@ -38,8 +38,9 @@ public:
 	 * Register the protocol
 	 * @param InProtocolName protocol name
 	 * @param InProtocol protocol shared reference
+	 * @return true if protocol added
 	 */
-	virtual void AddProtocol(FName InProtocolName, TSharedRef<IRemoteControlProtocol> InProtocol) = 0;
+	virtual bool AddProtocol(FName InProtocolName, TSharedRef<IRemoteControlProtocol> InProtocol) = 0;
 
 	/**
 	 * Unregister the protocol
@@ -50,4 +51,7 @@ public:
 
 	/** Unregister all protocols */
 	virtual void EmptyProtocols() = 0;
+
+	/** Whether protocol disabled */
+	virtual bool IsRCProtocolsDisable() const = 0;
 };
diff --git a/Engine/Plugins/VirtualProduction/RemoteControl/Source/RemoteControlProtocolWidgets/Private/RemoteControlProtocolWidgetsModule.cpp b/Engine/Plugins/VirtualProduction/RemoteControl/Source/RemoteControlProtocolWidgets/Private/RemoteControlProtocolWidgetsModule.cpp
--- a/Engine/Plugins/VirtualProduction/RemoteControl/Source/RemoteControlProtocolWidgets/Private/RemoteControlProtocolWidgetsModule.cpp
+++ b/Engine/Plugins/VirtualProduction/RemoteControl/Source/RemoteControlProtocolWidgets/Private/RemoteControlProtocolWidgetsModule.cpp
@@ -2,59 +2,12 @@
 
 #include "RemoteControlProtocolWidgetsModule.h"
 
-#include "IRemoteControlModule.h"
-#include "IRemoteControlProtocol.h"
-#include "IRemoteControlProtocolModule.h"
 #include "RemoteControlPreset.h"
-#include "AssetRegistry/IAssetRegistry.h"
-#include "Modules/ModuleManager.h"
 #include "ViewModels/ProtocolEntityViewModel.h"
 #include "Widgets/SRCProtocolBindingList.h"
 
 DEFINE_LOG_CATEGORY(LogRemoteControlProtocolWidgets);
 
-IMPLEMENT_MODULE(FRemoteControlProtocolWidgetsModule, RemoteControlProtocolWidgets);
-
-void FRemoteControlProtocolWidgetsModule::StartupModule()
-{
-	// Register Asset Load Event
-	FCoreUObjectDelegates::OnAssetLoaded.AddRaw(this, &FRemoteControlProtocolWidgetsModule::OnAssetLoaded);
-}
-
-void FRemoteControlProtocolWidgetsModule::ShutdownModule()
-{
-	// Unregister Asset Load Event
-	FCoreUObjectDelegates::OnAssetLoaded.RemoveAll(this);
-}
-
-void FRemoteControlProtocolWidgetsModule::OnAssetLoaded(UObject* InAsset)
-{
-	// @note: workaround for a current limitation that requires rebinding protocols on Preset load
-	// Check if asset is a RemoteControlPreset
-	if(URemoteControlPreset* Preset = Cast<URemoteControlPreset>(InAsset))
-	{
-		if(!ensure(IsValid(Preset)))
-		{
-			UE_LOG(LogRemoteControlProtocolWidgets, Error, TEXT("Attempted to load an invalid Preset: %s"), *InAsset->GetFName().ToString());
-			return;
-		}
-		
-		// Iterate over each contained property
-		for(TWeakPtr<FRemoteControlProperty> ExposedProperty : Preset->GetExposedEntities<FRemoteControlProperty>())
-		{
-			for(FRemoteControlProtocolBinding& Binding : ExposedProperty.Pin()->ProtocolBinding)
-			{
-				const TSharedPtr<IRemoteControlProtocol> Protocol = IRemoteControlProtocolModule::Get().GetProtocolByName(Binding.GetProtocolName());
-				// Supporting plugin needs to be loaded/protocol available.
-				if(Protocol.IsValid())
-				{
-					Protocol->Bind(Binding.GetRemoteControlProtocolEntityPtr());
-				}
-			}
-		}
-	}
-}
-
 TSharedRef<SWidget> FRemoteControlProtocolWidgetsModule::GenerateDetailsForEntity(URemoteControlPreset* InPreset, const FGuid& InFieldId, const EExposedFieldType& InFieldType)
 {
 	check(InPreset);
@@ -67,3 +20,5 @@ TSharedRef<SWidget> FRemoteControlProtocolWidgetsModule::GenerateDetailsForEntit
 
 	return SNullWidget::NullWidget;
 }
+
+IMPLEMENT_MODULE(FRemoteControlProtocolWidgetsModule, RemoteControlProtocolWidgets);
diff --git a/Engine/Plugins/VirtualProduction/RemoteControl/Source/RemoteControlProtocolWidgets/Private/RemoteControlProtocolWidgetsModule.h b/Engine/Plugins/VirtualProduction/RemoteControl/Source/RemoteControlProtocolWidgets/Private/RemoteControlProtocolWidgetsModule.h
--- a/Engine/Plugins/VirtualProduction/RemoteControl/Source/RemoteControlProtocolWidgets/Private/RemoteControlProtocolWidgetsModule.h
+++ b/Engine/Plugins/VirtualProduction/RemoteControl/Source/RemoteControlProtocolWidgets/Private/RemoteControlProtocolWidgetsModule.h
@@ -3,7 +3,6 @@
 #pragma once
 
 #include "IRemoteControlProtocolWidgetsModule.h"
-#include "Templates/SharedPointer.h"
 
 REMOTECONTROLPROTOCOLWIDGETS_API DECLARE_LOG_CATEGORY_EXTERN(LogRemoteControlProtocolWidgets, Log, All);
 
@@ -12,14 +11,6 @@ class URemoteControlPreset;
 class FRemoteControlProtocolWidgetsModule : public IRemoteControlProtocolWidgetsModule
 {
 public:
-	// IModuleInterface
-	virtual void StartupModule() override;
-	virtual void ShutdownModule() override;
-
 	/** Creates a widget for the given Preset Field */
 	virtual TSharedRef<SWidget> GenerateDetailsForEntity(URemoteControlPreset* InPreset, const FGuid& InFieldId, const EExposedFieldType& InFieldType) override;
-
-protected:
-	/** Called when any asset is loaded */
-	void OnAssetLoaded(UObject* InAsset);
 };
diff --git a/Engine/Plugins/VirtualProduction/RemoteControl/Source/RemoteControlProtocolWidgets/Private/ViewModels/ProtocolBindingViewModel.cpp b/Engine/Plugins/VirtualProduction/RemoteControl/Source/RemoteControlProtocolWidgets/Private/ViewModels/ProtocolBindingViewModel.cpp
--- a/Engine/Plugins/VirtualProduction/RemoteControl/Source/RemoteControlProtocolWidgets/Private/ViewModels/ProtocolBindingViewModel.cpp
+++ b/Engine/Plugins/VirtualProduction/RemoteControl/Source/RemoteControlProtocolWidgets/Private/ViewModels/ProtocolBindingViewModel.cpp
@@ -26,12 +26,12 @@ namespace Commands
 	{
 		if(TSharedPtr<FRemoteControlProperty> RCProperty = InPreset->GetExposedEntity<FRemoteControlProperty>(InArgs.EntityId).Pin())
 		{
-			FRemoteControlProtocolBinding* ProtocolBinding = RCProperty->ProtocolBinding.FindByHash(GetTypeHash(InArgs.BindingId), InArgs.BindingId);
-			check(ProtocolBinding);
+			FRemoteControlProtocolBinding* ProtocolBindings = RCProperty->ProtocolBindings.FindByHash(GetTypeHash(InArgs.BindingId), InArgs.BindingId);
+			check(ProtocolBindings);
 
-			const TSharedPtr<TStructOnScope<FRemoteControlProtocolEntity>> EntityPtr = ProtocolBinding->GetRemoteControlProtocolEntityPtr();
+			const TSharedPtr<TStructOnScope<FRemoteControlProtocolEntity>> EntityPtr = ProtocolBindings->GetRemoteControlProtocolEntityPtr();
 			const FRemoteControlProtocolMapping RangesData(RCProperty->GetProperty(), (*EntityPtr)->GetRangePropertySize());
-			ProtocolBinding->AddMapping(RangesData);
+			ProtocolBindings->AddMapping(RangesData);
 
 			return RangesData.GetId();
 		}
@@ -43,10 +43,10 @@ namespace Commands
 	{
 		if(TSharedPtr<FRemoteControlProperty> RCProperty = InPreset->GetExposedEntity<FRemoteControlProperty>(InArgs.EntityId).Pin())
 		{
-			FRemoteControlProtocolBinding* ProtocolBinding = RCProperty->ProtocolBinding.FindByHash(GetTypeHash(InArgs.BindingId), InArgs.BindingId);
-			check(ProtocolBinding);
+			FRemoteControlProtocolBinding* ProtocolBindings = RCProperty->ProtocolBindings.FindByHash(GetTypeHash(InArgs.BindingId), InArgs.BindingId);
+			check(ProtocolBindings);
 
-			const int32 NumRemoved = ProtocolBinding->RemoveMapping(InArgs.RangeId);
+			const int32 NumRemoved = ProtocolBindings->RemoveMapping(InArgs.RangeId);
 			return NumRemoved > 0;
 		}
 		
@@ -385,7 +385,7 @@ FRemoteControlProtocolBinding* FProtocolBindingViewModel::GetBinding() const
 
 	if (TSharedPtr<FRemoteControlProperty> RCProperty = Preset->GetExposedEntity<FRemoteControlProperty>(PropertyId).Pin())
 	{
-		return RCProperty->ProtocolBinding.FindByHash(GetTypeHash(BindingId), BindingId);
+		return RCProperty->ProtocolBindings.FindByHash(GetTypeHash(BindingId), BindingId);
 	}
 
 	checkNoEntry();
diff --git a/Engine/Plugins/VirtualProduction/RemoteControl/Source/RemoteControlProtocolWidgets/Private/ViewModels/ProtocolEntityViewModel.cpp b/Engine/Plugins/VirtualProduction/RemoteControl/Source/RemoteControlProtocolWidgets/Private/ViewModels/ProtocolEntityViewModel.cpp
--- a/Engine/Plugins/VirtualProduction/RemoteControl/Source/RemoteControlProtocolWidgets/Private/ViewModels/ProtocolEntityViewModel.cpp
+++ b/Engine/Plugins/VirtualProduction/RemoteControl/Source/RemoteControlProtocolWidgets/Private/ViewModels/ProtocolEntityViewModel.cpp
@@ -37,7 +37,7 @@ namespace Commands
 
 				FRemoteControlProtocolBinding ProtocolBinding(InArgs.ProtocolName, InArgs.EntityId, RemoteControlProtocolEntityPtr);
 				Protocol->Bind(ProtocolBinding.GetRemoteControlProtocolEntityPtr());
-				RCProperty->ProtocolBinding.Emplace(MoveTemp(ProtocolBinding));
+				RCProperty->ProtocolBindings.Emplace(MoveTemp(ProtocolBinding));
 
 				return MakeShared<FRemoteControlProtocolBinding>(ProtocolBinding);
 			}
@@ -50,7 +50,7 @@ namespace Commands
 	{
 		if (TSharedPtr<FRemoteControlProperty> RCProperty = InPreset->GetExposedEntity<FRemoteControlProperty>(InArgs.EntityId).Pin())
 		{
-			TSet<FRemoteControlProtocolBinding>& Test = RCProperty->ProtocolBinding;
+			TSet<FRemoteControlProtocolBinding>& Test = RCProperty->ProtocolBindings;
 			FRemoteControlProtocolBinding* FoundElement = Test.FindByHash(GetTypeHash(InArgs.BindingId), InArgs.BindingId);
 			if (FoundElement)
 			{
@@ -92,7 +92,7 @@ void FProtocolEntityViewModel::Initialize()
 	if (TSharedPtr<FRemoteControlProperty> RCProperty = Preset->GetExposedEntity<FRemoteControlProperty>(PropertyId).Pin())
 	{
 		Property = RCProperty->GetProperty();
-		for(FRemoteControlProtocolBinding& Binding : RCProperty->ProtocolBinding)
+		for(FRemoteControlProtocolBinding& Binding : RCProperty->ProtocolBindings)
 		{
 			const TSharedPtr<IRemoteControlProtocol> Protocol = IRemoteControlProtocolModule::Get().GetProtocolByName(Binding.GetProtocolName());
 			// Supporting plugin needs to be loaded/protocol available
diff --git a/Engine/Plugins/VirtualProduction/RemoteControl/Source/RemoteControlProtocolWidgets/Public/IRemoteControlProtocolWidgetsModule.h b/Engine/Plugins/VirtualProduction/RemoteControl/Source/RemoteControlProtocolWidgets/Public/IRemoteControlProtocolWidgetsModule.h
--- a/Engine/Plugins/VirtualProduction/RemoteControl/Source/RemoteControlProtocolWidgets/Public/IRemoteControlProtocolWidgetsModule.h
+++ b/Engine/Plugins/VirtualProduction/RemoteControl/Source/RemoteControlProtocolWidgets/Public/IRemoteControlProtocolWidgetsModule.h
@@ -1,9 +1,9 @@
-﻿// Copyright Epic Games, Inc. All Rights Reserved.
+// Copyright Epic Games, Inc. All Rights Reserved.
 
 #pragma once
 
 #include "Modules/ModuleInterface.h"
-#include "Templates/SharedPointer.h"
+#include "Modules/ModuleManager.h"
 #include "RemoteControlField.h"
 
 class SWidget;
diff --git a/Engine/Plugins/VirtualProduction/RemoteControlProtocolDMX/Source/RemoteControlProtocolDMX/Private/RemoteControlProtocolDMXModule.cpp b/Engine/Plugins/VirtualProduction/RemoteControlProtocolDMX/Source/RemoteControlProtocolDMX/Private/RemoteControlProtocolDMXModule.cpp
--- a/Engine/Plugins/VirtualProduction/RemoteControlProtocolDMX/Source/RemoteControlProtocolDMX/Private/RemoteControlProtocolDMXModule.cpp
+++ b/Engine/Plugins/VirtualProduction/RemoteControlProtocolDMX/Source/RemoteControlProtocolDMX/Private/RemoteControlProtocolDMXModule.cpp
@@ -14,7 +14,11 @@ public:
 	//~ Begin IModuleInterface
 	virtual void StartupModule() override
 	{
-		IRemoteControlProtocolModule::Get().AddProtocol(FRemoteControlProtocolDMX::ProtocolName, MakeShared<FRemoteControlProtocolDMX>());
+		const IRemoteControlProtocolModule& RemoteControlProtocolModule = IRemoteControlProtocolModule::Get();
+		if (!RemoteControlProtocolModule.IsRCProtocolsDisable())
+		{
+			IRemoteControlProtocolModule::Get().AddProtocol(FRemoteControlProtocolDMX::ProtocolName, MakeShared<FRemoteControlProtocolDMX>());
+		}
 	}
 	//~ End IModuleInterface
 };
diff --git a/Engine/Plugins/VirtualProduction/RemoteControlProtocolMIDI/Source/RemoteControlProtocolMIDI/Private/RemoteControlProtocolMIDIModule.cpp b/Engine/Plugins/VirtualProduction/RemoteControlProtocolMIDI/Source/RemoteControlProtocolMIDI/Private/RemoteControlProtocolMIDIModule.cpp
--- a/Engine/Plugins/VirtualProduction/RemoteControlProtocolMIDI/Source/RemoteControlProtocolMIDI/Private/RemoteControlProtocolMIDIModule.cpp
+++ b/Engine/Plugins/VirtualProduction/RemoteControlProtocolMIDI/Source/RemoteControlProtocolMIDI/Private/RemoteControlProtocolMIDIModule.cpp
@@ -37,8 +37,13 @@ void FRemoteControlProtocolMIDIModule::StartupModule()
 	}
 #endif // WITH_EDITOR
 
-	GetOnMIDIDevicesUpdated().AddRaw(this, &FRemoteControlProtocolMIDIModule::HandleMIDIDevicesUpdated);
-	IRemoteControlProtocolModule::Get().AddProtocol(FRemoteControlProtocolMIDI::ProtocolName, MakeShared<FRemoteControlProtocolMIDI>());
+
+	const IRemoteControlProtocolModule& RemoteControlProtocolModule = IRemoteControlProtocolModule::Get();
+	if (!RemoteControlProtocolModule.IsRCProtocolsDisable())
+	{
+		IRemoteControlProtocolModule::Get().AddProtocol(FRemoteControlProtocolMIDI::ProtocolName, MakeShared<FRemoteControlProtocolMIDI>());
+		GetOnMIDIDevicesUpdated().AddRaw(this, &FRemoteControlProtocolMIDIModule::HandleMIDIDevicesUpdated);
+	}
 }
 
 void FRemoteControlProtocolMIDIModule::ShutdownModule()
@@ -76,7 +81,7 @@ void FRemoteControlProtocolMIDIModule::HandleMIDIDevicesUpdated(FMIDIDeviceColle
 			// Iterate over each contained property
 			for(TWeakPtr<FRemoteControlProperty> ExposedProperty : Preset->GetExposedEntities<FRemoteControlProperty>())
 			{
-				for(FRemoteControlProtocolBinding& Binding : ExposedProperty.Pin()->ProtocolBinding)
+				for(FRemoteControlProtocolBinding& Binding : ExposedProperty.Pin()->ProtocolBindings)
 				{
 					if(Binding.GetProtocolName() == ProtocolName)
 					{
diff --git a/Engine/Plugins/VirtualProduction/RemoteControlProtocolOSC/Source/RemoteControlProtocolOSC/Private/RemoteControlProtocolOSCModule.cpp b/Engine/Plugins/VirtualProduction/RemoteControlProtocolOSC/Source/RemoteControlProtocolOSC/Private/RemoteControlProtocolOSCModule.cpp
--- a/Engine/Plugins/VirtualProduction/RemoteControlProtocolOSC/Source/RemoteControlProtocolOSC/Private/RemoteControlProtocolOSCModule.cpp
+++ b/Engine/Plugins/VirtualProduction/RemoteControlProtocolOSC/Source/RemoteControlProtocolOSC/Private/RemoteControlProtocolOSCModule.cpp
@@ -42,12 +42,15 @@ void FRemoteControlProtocolOSCModule::StartupModule()
 	}
 #endif // WITH_EDITOR
 
-	URemoteControlProtocolOSCSettings* OSCSettings = GetMutableDefault<URemoteControlProtocolOSCSettings>();
-	
-	IRemoteControlProtocolModule::Get().AddProtocol(FRemoteControlProtocolOSC::ProtocolName, MakeShared<FRemoteControlProtocolOSC>());
+	const IRemoteControlProtocolModule& RemoteControlProtocolModule = IRemoteControlProtocolModule::Get();
+	if (!RemoteControlProtocolModule.IsRCProtocolsDisable())
+	{
+		IRemoteControlProtocolModule::Get().AddProtocol(FRemoteControlProtocolOSC::ProtocolName, MakeShared<FRemoteControlProtocolOSC>());
 
-	// Init OSC servers after AddProtocol 
-	OSCSettings->InitOSCServers();
+		// Init OSC servers after AddProtocol
+		URemoteControlProtocolOSCSettings* OSCSettings = GetMutableDefault<URemoteControlProtocolOSCSettings>();
+		OSCSettings->InitOSCServers();
+	}
 }
 
 void FRemoteControlProtocolOSCModule::ShutdownModule()
