Fix CIS not finding SNullWidget in RemoteControlLogger.

diff --git a/Engine/Plugins/VirtualProduction/RemoteControl/Source/RemoteControl/Private/RemoteControlLogger.cpp b/Engine/Plugins/VirtualProduction/RemoteControl/Source/RemoteControl/Private/RemoteControlLogger.cpp
--- a/Engine/Plugins/VirtualProduction/RemoteControl/Source/RemoteControl/Private/RemoteControlLogger.cpp
+++ b/Engine/Plugins/VirtualProduction/RemoteControl/Source/RemoteControl/Private/RemoteControlLogger.cpp
@@ -3,6 +3,7 @@
 #include "RemoteControlLogger.h"
 
 #include "Modules/ModuleManager.h"
+#include "Widgets/SNullWidget.h"
 
 #if WITH_EDITOR
 #include "IMessageLogListing.h"
