RC DMX and RC Logger non-unity fixes

diff --git a/Engine/Plugins/VirtualProduction/RemoteControl/RemoteControl.uplugin b/Engine/Plugins/VirtualProduction/RemoteControl/RemoteControl.uplugin
--- a/Engine/Plugins/VirtualProduction/RemoteControl/RemoteControl.uplugin
+++ b/Engine/Plugins/VirtualProduction/RemoteControl/RemoteControl.uplugin
@@ -41,7 +41,7 @@
 		},
 		{
 			"Name": "RemoteControlCommon",
-			"Type": "Editor",
+			"Type": "Runtime",
 			"LoadingPhase": "Default"
 		},
 		{
diff --git a/Engine/Plugins/VirtualProduction/RemoteControl/Source/RemoteControl/Private/RemoteControlLogger.cpp b/Engine/Plugins/VirtualProduction/RemoteControl/Source/RemoteControl/Private/RemoteControlLogger.cpp
--- a/Engine/Plugins/VirtualProduction/RemoteControl/Source/RemoteControl/Private/RemoteControlLogger.cpp
+++ b/Engine/Plugins/VirtualProduction/RemoteControl/Source/RemoteControl/Private/RemoteControlLogger.cpp
@@ -8,6 +8,8 @@
 #if WITH_EDITOR
 #include "IMessageLogListing.h"
 #include "MessageLogModule.h"
+#else
+#include "Widgets/SNullWidget.h"
 #endif
 
 #define LOCTEXT_NAMESPACE "RemoteControlLogger"
diff --git a/Engine/Plugins/VirtualProduction/RemoteControl/Source/RemoteControl/RemoteControl.Build.cs b/Engine/Plugins/VirtualProduction/RemoteControl/Source/RemoteControl/RemoteControl.Build.cs
--- a/Engine/Plugins/VirtualProduction/RemoteControl/Source/RemoteControl/RemoteControl.Build.cs
+++ b/Engine/Plugins/VirtualProduction/RemoteControl/Source/RemoteControl/RemoteControl.Build.cs
@@ -10,6 +10,7 @@ public class RemoteControl : ModuleRules
 			new string[] {
 				"Core",
 				"CoreUObject",
+				"RemoteControlCommon",
 			}
 		);
 
@@ -19,6 +20,7 @@ public class RemoteControl : ModuleRules
 				"Engine",
 				"RemoteControlInterception",
 				"Serialization",
+				"SlateCore"
 			}
 		);
 
diff --git a/Engine/Plugins/VirtualProduction/RemoteControl/Source/RemoteControlCommon/RemoteControlCommon.Build.cs b/Engine/Plugins/VirtualProduction/RemoteControl/Source/RemoteControlCommon/RemoteControlCommon.Build.cs
--- a/Engine/Plugins/VirtualProduction/RemoteControl/Source/RemoteControlCommon/RemoteControlCommon.Build.cs
+++ b/Engine/Plugins/VirtualProduction/RemoteControl/Source/RemoteControlCommon/RemoteControlCommon.Build.cs
@@ -8,29 +8,43 @@ public class RemoteControlCommon : ModuleRules
 	public RemoteControlCommon(ReadOnlyTargetRules Target) : base(Target)
 	{
 		PublicDependencyModuleNames.AddRange(
-			new string[] {}
+			new string[] { }
 		);
 
-        PrivateDependencyModuleNames.AddRange(
+		PrivateDependencyModuleNames.AddRange(
 			new string[] {
 				"ApplicationCore",
-				"AssetRegistry",
-				"AssetTools",
-				"BlueprintGraph",
 				"Core",
 				"CoreUObject",
-				"DeveloperSettings",
-				"EditorWidgets",
-				"EditorStyle",
-				"EditorSubsystem",
 				"Engine",
 				"InputCore",
-				"Projects",
-				"PropertyEditor",
-				"Slate",
-				"SlateCore",
-				"UnrealEd",
 			}
 		);
+		
+		if (Target.bBuildEditor)
+		{
+			PublicDependencyModuleNames.AddRange(
+				new string[] {
+					"PropertyEditor",
+				}
+			);
+
+			PrivateDependencyModuleNames.AddRange(
+				new string[] {
+					"AssetRegistry",
+					"AssetTools",
+					"BlueprintGraph",
+					"DeveloperSettings",
+					"EditorWidgets",
+					"EditorStyle",
+					"EditorSubsystem",
+					"Projects",
+					"PropertyEditor",
+					"Slate",
+					"SlateCore",
+					"UnrealEd",
+				}
+			);
+		}
     }
 }
diff --git a/Engine/Plugins/VirtualProduction/RemoteControlProtocolDMX/Source/RemoteControlProtocolDMXEditor/Private/RemoteControlProtocolDMXEditorModule.cpp b/Engine/Plugins/VirtualProduction/RemoteControlProtocolDMX/Source/RemoteControlProtocolDMXEditor/Private/RemoteControlProtocolDMXEditorModule.cpp
--- a/Engine/Plugins/VirtualProduction/RemoteControlProtocolDMX/Source/RemoteControlProtocolDMXEditor/Private/RemoteControlProtocolDMXEditorModule.cpp
+++ b/Engine/Plugins/VirtualProduction/RemoteControlProtocolDMX/Source/RemoteControlProtocolDMXEditor/Private/RemoteControlProtocolDMXEditorModule.cpp
@@ -4,6 +4,7 @@
 
 #include "PropertyEditorModule.h"
 #include "Modules/ModuleInterface.h"
+#include "Modules/ModuleManager.h"
 
 #include "RemoteControlProtocolDMX.h"
 #include "RemoteControlProtocolDMXEditorCustomization.h"
