Forward raw transform property changes to their respective setter
Fix modifying transform not being usable in -game

diff --git a/Engine/Plugins/VirtualProduction/RemoteControl/Source/RemoteControl/Private/RemoteControlModule.cpp b/Engine/Plugins/VirtualProduction/RemoteControl/Source/RemoteControl/Private/RemoteControlModule.cpp
--- a/Engine/Plugins/VirtualProduction/RemoteControl/Source/RemoteControl/Private/RemoteControlModule.cpp
+++ b/Engine/Plugins/VirtualProduction/RemoteControl/Source/RemoteControl/Private/RemoteControlModule.cpp
@@ -2,9 +2,12 @@
 
 #include "IRemoteControlModule.h"
 #include "IRemoteControlInterceptionFeature.h"
+#include "IStructDeserializerBackend.h"
 
 #include "AssetRegistry/AssetRegistryModule.h"
 #include "AssetRegistry/IAssetRegistry.h"
+#include "Backends/CborStructSerializerBackend.h"
+#include "Components/StaticMeshComponent.h"
 #include "Features/IModularFeatures.h"
 #include "Misc/CoreMisc.h"
 #include "IStructSerializerBackend.h"
@@ -14,9 +17,10 @@
 #include "RemoteControlPreset.h"
 #include "StructSerializer.h"
 #include "StructDeserializer.h"
-#include "UObject/UnrealType.h"
+#include "Misc/ScopeExit.h"
 #include "UObject/Class.h"
 #include "UObject/FieldPath.h"
+#include "UObject/UnrealType.h"
 
 #if WITH_EDITOR
 #include "Editor.h"
@@ -26,6 +30,12 @@
 DEFINE_LOG_CATEGORY(LogRemoteControl);
 #define LOCTEXT_NAMESPACE "RemoteControl"
 
+struct FRCInterceptionPayload
+{
+	TArray<uint8> Payload;
+	ERCPayloadType Type;
+};
+
 namespace RemoteControlUtil
 {
 	const FName NAME_DisplayName(TEXT("DisplayName"));
@@ -36,6 +46,8 @@ namespace RemoteControlUtil
 	const FName NAME_BlueprintSetter(TEXT("BlueprintSetter"));
 	const FName NAME_AllowPrivateAccess(TEXT("AllowPrivateAccess"));
 
+	TMap<TWeakFieldPtr<FProperty>, TWeakObjectPtr<UFunction>> CachedSetterFunctions;
+
 	bool CompareFunctionName(const FString& FunctionName, const FString& ScriptName)
 	{
 		int32 SemiColonIndex = INDEX_NONE;
@@ -157,6 +169,228 @@ namespace RemoteControlUtil
 
 		return Id;
 	}
+
+	/** Returns whether the access is a write access regardless of if it generates a transaction. */
+	bool IsWriteAccess(ERCAccess Access)
+	{
+		return Access == ERCAccess::WRITE_ACCESS || Access == ERCAccess::WRITE_TRANSACTION_ACCESS;
+	}
+}
+
+namespace RemoteControlSetterUtils
+{
+	static TMap<TWeakFieldPtr<FProperty>, TWeakObjectPtr<UFunction>> CachedSetterFunctions;
+
+	struct FConvertToFunctionCallArgs
+	{
+		FConvertToFunctionCallArgs(const FRCObjectReference& InObjectReference, IStructDeserializerBackend& InReaderBackend, FRCCall& OutCall)
+			: ObjectReference(InObjectReference)
+			, ReaderBackend(InReaderBackend)
+			, Call(OutCall) 
+		{
+		}
+		
+		const FRCObjectReference& ObjectReference;
+		IStructDeserializerBackend& ReaderBackend;
+		FRCCall& Call;
+	};
+	
+	UFunction* FindSetterFunction(FProperty* Property)
+	{
+		// UStruct properties cannot have setters.
+		if (!ensure(Property) || !Property->GetOwnerClass())
+		{
+			return nullptr;
+		}
+
+		// Check if the property setter is already cached.
+		TWeakObjectPtr<UFunction> SetterPtr = CachedSetterFunctions.FindRef(Property);
+		if (SetterPtr.IsValid())
+		{
+			return SetterPtr.Get();
+		}
+		
+		UFunction* SetterFunction = nullptr;
+#if WITH_EDITOR
+		const FString& SetterName =  Property->GetMetaData(*RemoteControlUtil::NAME_BlueprintSetter.ToString());
+		if (!SetterName.IsEmpty())
+		{
+			SetterFunction = Property->GetOwnerClass()->FindFunctionByName(*SetterName);
+		}
+#endif
+
+		FString PropertyName = Property->GetName();
+		if (Property->IsA<FBoolProperty>()) 
+		{
+			PropertyName.RemoveFromStart("b", ESearchCase::CaseSensitive);
+		}
+
+		static const TArray<FString> SetterPrefixes = {
+			FString("Set"),
+			FString("K2_Set")
+		};
+
+		for (const FString& Prefix : SetterPrefixes)
+		{
+			FName SetterFunctionName = FName(Prefix + PropertyName);
+			SetterFunction = Property->GetOwnerClass()->FindFunctionByName(SetterFunctionName);
+			if (SetterFunction)
+			{
+				break;
+			}
+		}
+
+		if (SetterFunction)
+		{
+			CachedSetterFunctions.Add(Property,  SetterFunction);
+		}
+
+		return SetterFunction;
+	}
+	
+	FProperty* FindSetterArgument(UFunction* SetterFunction, FProperty* PropertyToModify)
+	{
+		FProperty* SetterArgument = nullptr;
+
+		if (!ensure(SetterFunction))
+		{
+			return nullptr;
+		}
+
+		// Check if the first parameter for the setter function matches the parameter value.
+		for (TFieldIterator<FProperty> PropertyIt(SetterFunction); PropertyIt; ++PropertyIt)
+		{
+			if (PropertyIt->HasAnyPropertyFlags(CPF_Parm) && !PropertyIt->HasAnyPropertyFlags(CPF_ReturnParm))
+			{
+				if (PropertyIt->SameType(PropertyToModify))
+				{
+					SetterArgument = *PropertyIt;
+				}
+				
+				break;
+			}
+		}
+
+		return SetterArgument;
+	}
+
+	void CreateRCCall(FConvertToFunctionCallArgs& InOutArgs, UFunction* InFunction, FStructOnScope&& InFunctionArguments, FRCInterceptionPayload& OutPayload)
+	{
+		ensure(InFunctionArguments.GetStruct() && InFunctionArguments.GetStruct()->IsA<UFunction>());
+		
+		// Create the output payload for interception purposes.
+		FMemoryWriter Writer{OutPayload.Payload};
+		FCborStructSerializerBackend WriterBackend{Writer, EStructSerializerBackendFlags::Default};
+		FStructSerializer::Serialize(InFunctionArguments.GetStructMemory(), *const_cast<UStruct*>(InFunctionArguments.GetStruct()), WriterBackend, FStructSerializerPolicies());
+		OutPayload.Type = ERCPayloadType::Cbor;
+
+		InOutArgs.Call.bGenerateTransaction = InOutArgs.ObjectReference.Access == ERCAccess::WRITE_TRANSACTION_ACCESS ? true : false;
+		InOutArgs.Call.CallRef.Function = InFunction;
+		InOutArgs.Call.CallRef.Object = InOutArgs.ObjectReference.Object;
+		InOutArgs.Call.ParamStruct = MoveTemp(InFunctionArguments);
+	}
+
+	void CreateRCCall(FConvertToFunctionCallArgs& InOutArgs, UFunction* InFunction, FStructOnScope&& InFunctionArguments)
+	{
+		ensure(InFunctionArguments.GetStruct() && InFunctionArguments.GetStruct()->IsA<UFunction>());
+		
+		InOutArgs.Call.bGenerateTransaction = InOutArgs.ObjectReference.Access == ERCAccess::WRITE_TRANSACTION_ACCESS ? true : false;
+		InOutArgs.Call.CallRef.Function = InFunction;
+		InOutArgs.Call.CallRef.Object = InOutArgs.ObjectReference.Object;
+		InOutArgs.Call.ParamStruct = MoveTemp(InFunctionArguments);
+	}
+
+	/** Create the payload to pass to be passed to a property's setter function. */
+	TOptional<FStructOnScope> CreateSetterFunctionPayload(UFunction* InSetterFunction, FConvertToFunctionCallArgs& InOutArgs)
+	{
+		TOptional<FStructOnScope> OptionalArgsOnScope;
+
+		bool bSuccess = false;
+		if (FProperty* SetterArgument = FindSetterArgument(InSetterFunction, InOutArgs.ObjectReference.Property.Get()))
+		{
+			FStructOnScope ArgsOnScope{InSetterFunction};
+			
+			// First put the complete property value from the object in the struct on scope
+			// in case the user only a part of the incoming structure (ie. Providing only { "x": 2 } in the case of a vector.
+			const uint8* ContainerAddress = InOutArgs.ObjectReference.Property->ContainerPtrToValuePtr<uint8>(InOutArgs.ObjectReference.ContainerAdress);
+			InOutArgs.ObjectReference.Property->CopyCompleteValue(ArgsOnScope.GetStructMemory(), ContainerAddress);
+
+			// Temporarily rename the setter argument in order to deserialize the incoming property modification on top of it
+			// regardless of the argument name.
+			FName OldSetterArgumentName = SetterArgument->GetFName();
+			SetterArgument->Rename(InOutArgs.ObjectReference.Property->GetFName());
+			{
+				ON_SCOPE_EXIT
+				{
+					SetterArgument->Rename(OldSetterArgumentName);
+				};
+
+				// Then deserialize the input value on top of it and reset the setter property name.
+				bSuccess = FStructDeserializer::Deserialize((void*)ArgsOnScope.GetStructMemory(), *const_cast<UStruct*>(ArgsOnScope.GetStruct()), InOutArgs.ReaderBackend, FStructDeserializerPolicies());
+			}
+			
+			if (bSuccess)
+			{
+				OptionalArgsOnScope = MoveTemp(ArgsOnScope);
+			}
+		}
+		
+		return OptionalArgsOnScope;
+	}
+	
+	bool ConvertModificationToFunctionCall(FConvertToFunctionCallArgs& InOutArgs, UFunction* InSetterFunction, FRCInterceptionPayload& OutPayload)
+	{
+		if (TOptional<FStructOnScope> ArgsOnScope = CreateSetterFunctionPayload(InSetterFunction, InOutArgs))
+		{
+			CreateRCCall(InOutArgs, InSetterFunction, MoveTemp(*ArgsOnScope), OutPayload);
+			return true;
+		}
+		
+		return false;
+	}
+
+	bool ConvertModificationToFunctionCall(FConvertToFunctionCallArgs& InOutArgs, UFunction* InSetterFunction)
+	{
+		if (TOptional<FStructOnScope> ArgsOnScope = CreateSetterFunctionPayload(InSetterFunction, InOutArgs))
+		{
+			CreateRCCall(InOutArgs, InSetterFunction, MoveTemp(*ArgsOnScope));
+			return true;
+		}
+		
+		return false;
+	}
+
+	bool ConvertModificationToFunctionCall(FConvertToFunctionCallArgs& InOutArgs)
+	{
+		TRACE_CPUPROFILER_EVENT_SCOPE(RemoteControlSetterUtils::ConvertModificationToFunctionCall);
+		if (!InOutArgs.ObjectReference.Object.IsValid() || !InOutArgs.ObjectReference.Property.IsValid())
+		{
+			return false;
+		}
+		
+		if (UFunction* SetterFunction = FindSetterFunction(InOutArgs.ObjectReference.Property.Get()))
+		{
+			return ConvertModificationToFunctionCall(InOutArgs, SetterFunction);
+		}
+
+		return false;
+	}
+
+	bool ConvertModificationToFunctionCall(FConvertToFunctionCallArgs& InArgs, FRCInterceptionPayload& OutPayload)
+	{
+		TRACE_CPUPROFILER_EVENT_SCOPE(RemoteControlSetterUtils::ConvertModificationToFunctionCall);
+		if (!InArgs.ObjectReference.Object.IsValid() || !InArgs.ObjectReference.Property.IsValid())
+		{
+			return false;
+		}
+		
+		if (UFunction* SetterFunction = FindSetterFunction(InArgs.ObjectReference.Property.Get()))
+		{
+			return ConvertModificationToFunctionCall(InArgs, SetterFunction, OutPayload);
+		}
+
+		return false;
+	}
 }
 
 /**
@@ -278,7 +512,7 @@ public:
 		return bSuccess;
 	}
 
-	virtual bool InvokeCall(FRCCall& InCall, ERCPayloadType InPayloadType, const TArray<uint8>& InInterceptPayload) override
+	virtual bool InvokeCall(FRCCall& InCall, ERCPayloadType InPayloadType = ERCPayloadType::Json, const TArray<uint8>& InInterceptPayload = TArray<uint8>()) override
 	{
 		if (InCall.IsValid())
 		{
@@ -374,7 +608,10 @@ public:
 					if (PropertyPath.Resolve(Object))
 					{
 						FProperty* ResolvedProperty = PropertyPath.GetResolvedData().Field;
-						if (RemoteControlUtil::IsPropertyAllowed(ResolvedProperty, AccessType, bObjectInGame))
+
+						// When resolving a property for writing, resolve successfully if it should use a setter since it will end up using it. 
+						if ((RemoteControlUtil::IsWriteAccess(AccessType) && PropertyModificationShouldUseSetter(Object, ResolvedProperty))
+							|| RemoteControlUtil::IsPropertyAllowed(ResolvedProperty, AccessType, bObjectInGame))
 						{
 							OutObjectRef = FRCObjectReference{ AccessType , Object, MoveTemp(PropertyPath) };
 						}
@@ -470,9 +707,29 @@ public:
 
 	virtual bool SetObjectProperties(const FRCObjectReference& ObjectAccess, IStructDeserializerBackend& Backend, ERCPayloadType InPayloadType, const TArray<uint8>& InPayload) override
 	{
-		// Check the replication path before apply property values
+		// Check the replication path before applying property values
 		if (InPayload.Num() != 0 && ObjectAccess.Object.IsValid())
 		{
+			// Convert raw property modifications to setter function calls if necessary.
+			if (PropertyModificationShouldUseSetter(ObjectAccess.Object.Get(), ObjectAccess.Property.Get()))
+			{
+				FRCCall	Call;
+				FRCInterceptionPayload InterceptionPayload;
+				constexpr bool bCreateInterceptionPayload = true;
+				RemoteControlSetterUtils::FConvertToFunctionCallArgs Args(ObjectAccess, Backend, Call);
+				if (RemoteControlSetterUtils::ConvertModificationToFunctionCall(Args, InterceptionPayload))
+				{
+					return InvokeCall(Call, InterceptionPayload.Type, InterceptionPayload.Payload);
+				}
+			}
+			
+			// If a setter wasn't used, verify if the property should be allowed.
+			bool bObjectInGame = !GIsEditor || (ObjectAccess.Object.IsValid() && ObjectAccess.Object->GetOutermost()->HasAnyPackageFlags(PKG_PlayInEditor));
+			if (!RemoteControlUtil::IsPropertyAllowed(ObjectAccess.Property.Get(), ObjectAccess.Access, bObjectInGame))
+			{
+				return false;
+			}
+
 			// Build interception command
 			FString PropertyPathString = TFieldPath<FProperty>(ObjectAccess.Property.Get()).ToString();
 			FRCIPropertiesMetadata PropsMetadata(ObjectAccess.Object->GetPathName(), PropertyPathString, ObjectAccess.PropertyPathInfo.ToString(), ToExternal(ObjectAccess.Access), ToExternal(InPayloadType), InPayload);
@@ -501,9 +758,20 @@ public:
 			}
 		}
 
+		// Convert raw property modifications to setter function calls if necessary.
+		if (PropertyModificationShouldUseSetter(ObjectAccess.Object.Get(), ObjectAccess.Property.Get()))
+		{
+			FRCCall	Call;
+			RemoteControlSetterUtils::FConvertToFunctionCallArgs Args(ObjectAccess, Backend, Call);
+			if (RemoteControlSetterUtils::ConvertModificationToFunctionCall(Args))
+			{
+				return InvokeCall(Call);
+			}
+		}
+
 		//Setting object properties require a property and can't be done at the object level. Property must be valid to move forward
 		if (ObjectAccess.IsValid()
-			&& (ObjectAccess.Access == ERCAccess::WRITE_ACCESS || ObjectAccess.Access == ERCAccess::WRITE_TRANSACTION_ACCESS)
+			&& (RemoteControlUtil::IsWriteAccess(ObjectAccess.Access))
 			&& ObjectAccess.Property.IsValid()
 			&& ObjectAccess.PropertyPathInfo.IsResolved())
 		{
@@ -601,9 +869,7 @@ public:
 			}
 		}
 
-
-		if (ObjectAccess.IsValid()
-			&& (ObjectAccess.Access == ERCAccess::WRITE_ACCESS || ObjectAccess.Access == ERCAccess::WRITE_TRANSACTION_ACCESS))
+		if (ObjectAccess.IsValid() && RemoteControlUtil::IsWriteAccess(ObjectAccess.Access))
 		{
 			UObject* Object = ObjectAccess.Object.Get();
 			UStruct* ContainerType = ObjectAccess.ContainerType.Get();
@@ -842,9 +1108,24 @@ private:
 		
 		CachedPresetsByName.FindOrAdd(AssetData.AssetName).AddUnique(AssetData);
 	}
+	
+	bool PropertyModificationShouldUseSetter(UObject* Object, FProperty* Property)
+	{
+		if (!Property || !Object)
+		{
+			return false;
+		}
+		
+		const bool bObjectInGamePackage = !GIsEditor || Object->GetOutermost()->HasAnyPackageFlags(PKG_PlayInEditor);
+		if (!RemoteControlUtil::IsPropertyAllowed(Property, ERCAccess::WRITE_ACCESS, bObjectInGamePackage))
+		{
+			return !!RemoteControlSetterUtils::FindSetterFunction(Property);
+		}
+
+		return false;
+	}
 
 private:
-	
 	/** Cache of preset names to preset assets */
 	mutable TMap<FName, TArray<FAssetData>> CachedPresetsByName;
 
diff --git a/Engine/Plugins/VirtualProduction/RemoteControl/Source/RemoteControl/Private/RemoteControlPreset.cpp b/Engine/Plugins/VirtualProduction/RemoteControl/Source/RemoteControl/Private/RemoteControlPreset.cpp
--- a/Engine/Plugins/VirtualProduction/RemoteControl/Source/RemoteControl/Private/RemoteControlPreset.cpp
+++ b/Engine/Plugins/VirtualProduction/RemoteControl/Source/RemoteControl/Private/RemoteControlPreset.cpp
@@ -6,8 +6,11 @@
 #include "Algo/ForEach.h"
 #include "Algo/Transform.h"
 #include "Components/ActorComponent.h"
+#include "Components/StaticMeshComponent.h"
 #include "GameFramework/Actor.h"
+#include "HAL/IConsoleManager.h"
 #include "IRemoteControlModule.h"
+#include "Misc/CoreDelegates.h"
 #include "Misc/Optional.h"
 #include "RemoteControlExposeRegistry.h"
 #include "RemoteControlFieldPath.h"
@@ -28,7 +31,9 @@
 
 URemoteControlPreset::FOnPostLoadRemoteControlPreset URemoteControlPreset::OnPostLoadRemoteControlPreset;
 
-#define LOCTEXT_NAMESPACE "RemoteControlPreset" 
+#define LOCTEXT_NAMESPACE "RemoteControlPreset"
+
+static TAutoConsoleVariable<int32> CVarRemoteControlFramesBetweenPropertyWatch(TEXT("RemoteControl.FramesBetweenPropertyWatch"), 6, TEXT("The number of frames between every property value comparison when manually watching for property changes."));
 
 namespace
 {
@@ -667,6 +672,8 @@ void URemoteControlPreset::PostLoad()
 	InitializeEntitiesMetadata();
 
 	RegisterEntityDelegates();
+
+	CreatePropertyWatchers();
 }
 
 void URemoteControlPreset::BeginDestroy()
@@ -756,7 +763,15 @@ TWeakPtr<FRemoteControlProperty> URemoteControlPreset::ExposeProperty(UObject* O
 	}
 
 	FRemoteControlProperty RCProperty{ this, Registry->GenerateUniqueLabel(DesiredName), MoveTemp(FieldPath), { FindOrAddBinding(Object) } };
-	return StaticCastSharedPtr<FRemoteControlProperty>(Expose(MoveTemp(RCProperty), FRemoteControlProperty::StaticStruct(), Args.GroupId));
+
+	TSharedPtr<FRemoteControlProperty> RCPropertyPtr = StaticCastSharedPtr<FRemoteControlProperty>(Expose(MoveTemp(RCProperty), FRemoteControlProperty::StaticStruct(), Args.GroupId));
+
+	if (PropertyShouldBeWatched(RCPropertyPtr))
+	{
+		CreatePropertyWatcher(RCPropertyPtr);
+	}
+
+	return RCPropertyPtr;
 }
 
 TWeakPtr<FRemoteControlFunction> URemoteControlPreset::ExposeFunction(UObject* Object, UFunction* Function, FRemoteControlPresetExposeArgs Args)
@@ -938,6 +953,54 @@ void URemoteControlPreset::RegisterOnCompileEvent(const TSharedPtr<FRemoteContro
 #endif
 }
 
+void URemoteControlPreset::CreatePropertyWatcher(const TSharedPtr<FRemoteControlProperty>& RCProperty)
+{
+	if (ensure(RCProperty))
+	{
+		if (!PropertyWatchers.Contains(RCProperty->GetId()))
+		{
+			FRCPropertyWatcher Watcher{RCProperty, FSimpleDelegate::CreateLambda([this, WeakProperty = TWeakPtr<FRemoteControlProperty>(RCProperty)]()
+			{
+				if (TSharedPtr<FRemoteControlProperty> PinnedProperty = WeakProperty.Pin())
+				{
+					PerFrameModifiedProperties.Add(PinnedProperty->GetId());
+				}
+			})};
+			
+			PropertyWatchers.Add(RCProperty->GetId(), MoveTemp(Watcher));
+		}
+	}
+}
+
+bool URemoteControlPreset::PropertyShouldBeWatched(const TSharedPtr<FRemoteControlProperty>& RCProperty) const
+{
+	// If we are not running in editor, we need to watch all properties as there is no object modified callback.
+	if (!GIsEditor)
+	{
+		return true;	
+	}
+	
+	static const TSet<FName> WatchedPropertyNames =
+		{
+			UStaticMeshComponent::GetRelativeLocationPropertyName(),
+			UStaticMeshComponent::GetRelativeRotationPropertyName(),
+			UStaticMeshComponent::GetRelativeScale3DPropertyName()
+		};
+	
+	return RCProperty && WatchedPropertyNames.Contains(RCProperty->FieldName);
+}
+
+void URemoteControlPreset::CreatePropertyWatchers()
+{
+	for (const TSharedPtr<FRemoteControlProperty>& ExposedProperty : Registry->GetExposedEntities<FRemoteControlProperty>())
+	{
+		if (PropertyShouldBeWatched(ExposedProperty))
+		{
+			CreatePropertyWatcher(ExposedProperty);
+		}
+	}
+}
+
 TOptional<FRemoteControlFunction> URemoteControlPreset::GetFunction(FName FunctionLabel) const
 {
 PRAGMA_DISABLE_DEPRECATION_WARNINGS
@@ -1020,7 +1083,8 @@ void URemoteControlPreset::Unexpose(const FGuid& EntityId)
 		Registry->RemoveExposedEntity(EntityId);
 		FRCCachedFieldData CachedData = FieldCache.FindChecked(EntityId);
 		Layout.RemoveField(CachedData.LayoutGroupId, EntityId);
-		FieldCache.Remove(EntityId);	
+		FieldCache.Remove(EntityId);
+		PropertyWatchers.Remove(EntityId);
 	}
 }
 
@@ -1098,7 +1162,7 @@ void URemoteControlPreset::NotifyExposedPropertyChanged(FName PropertyLabel)
 {
 	if (TSharedPtr<FRemoteControlProperty> ExposedProperty = GetExposedEntity<FRemoteControlProperty>(GetExposedEntityId(PropertyLabel)).Pin())
 	{
-		OnExposedPropertyChanged().Broadcast(this, *ExposedProperty);
+		PerFrameModifiedProperties.Add(ExposedProperty->GetId());
 	}
 }
 
@@ -1332,7 +1396,7 @@ FName URemoteControlPreset::RenameExposedEntity(const FGuid& ExposedEntityId, FN
 
 bool URemoteControlPreset::IsExposed(const FGuid& ExposedEntityId) const
 {
-	TRACE_CPUPROFILER_EVENT_SCOPE(TEXT("URemoteControlPreset::IsExposed"));
+	TRACE_CPUPROFILER_EVENT_SCOPE(URemoteControlPreset::IsExposed);
 	if (FieldCache.Contains(ExposedEntityId))
 	{
 		return true;
@@ -1412,12 +1476,14 @@ void URemoteControlPreset::OnExpose(const FExposeInfo& Info)
 	{
 		NameToGuidMap.Add(Field->GetLabel(), Info.FieldId);
 		Layout.AddField(FieldGroupId, Info.FieldId);
-		OnPropertyExposed().Broadcast(this, Field->GetLabel());
+		OnEntityExposed().Broadcast(this, Info.FieldId);
 	}
 }
 
 void URemoteControlPreset::OnUnexpose(FGuid UnexposedFieldId)
  {
+	OnEntityUnexposed().Broadcast(this, UnexposedFieldId);
+	
 	FRCCachedFieldData CachedData = FieldCache.FindChecked(UnexposedFieldId);
 
 	Layout.RemoveField(CachedData.LayoutGroupId, UnexposedFieldId);
@@ -1433,11 +1499,6 @@ void URemoteControlPreset::OnUnexpose(FGuid UnexposedFieldId)
 			It.RemoveCurrent();
 		}
 	}
-
-	if (FieldLabel != NAME_None)
-	{
-		OnPropertyUnexposed().Broadcast(this, FieldLabel);
-	}
 }
 
 void URemoteControlPreset::CacheFieldsData()
@@ -1476,6 +1537,7 @@ void URemoteControlPreset::CacheFieldLayoutData()
 void URemoteControlPreset::OnObjectPropertyChanged(UObject* Object, struct FPropertyChangedEvent& Event)
 {
 	// Objects modified should have run through the preobjectmodified. If interesting, they will be cached
+	TRACE_CPUPROFILER_EVENT_SCOPE(URemoteControlPreset::OnObjectPropertyChanged);
 	for (auto Iter = PreObjectsModifiedCache.CreateIterator(); Iter; ++Iter)
 	{
 		FGuid& PropertyId = Iter.Key();
@@ -1487,7 +1549,7 @@ void URemoteControlPreset::OnObjectPropertyChanged(UObject* Object, struct FProp
 			if (TSharedPtr<FRemoteControlProperty> Property = Registry->GetExposedEntity<FRemoteControlProperty>(PropertyId))
 			{
 				UE_LOG(LogRemoteControl, VeryVerbose, TEXT("(%s) Change detected on %s::%s"), *GetName(), *Object->GetName(), *Event.Property->GetName());
-				OnExposedPropertyChanged().Broadcast(this, *Property);
+				PerFrameModifiedProperties.Add(Property->GetId());
 				Iter.RemoveCurrent();
 			}
 		}
@@ -1512,6 +1574,7 @@ void URemoteControlPreset::OnObjectPropertyChanged(UObject* Object, struct FProp
 
 void URemoteControlPreset::OnPreObjectPropertyChanged(UObject* Object, const class FEditPropertyChain& PropertyChain)
 {
+	TRACE_CPUPROFILER_EVENT_SCOPE(URemoteControlPreset::OnPreObjectPropertyChanged);
 	using PropertyNode = TDoubleLinkedList<FProperty*>::TDoubleLinkedListNode;
 
 	//Quick validation of the property chain
@@ -1615,14 +1678,18 @@ void URemoteControlPreset::RegisterDelegates()
 	
 	FCoreUObjectDelegates::OnObjectsReplaced.AddUObject(this, &URemoteControlPreset::OnReplaceObjects);
 
-	FCoreDelegates::OnEndFrame.AddUObject(this, &URemoteControlPreset::OnEndFrame);
-
 	FEditorDelegates::MapChange.AddUObject(this, &URemoteControlPreset::OnMapChange);
 #endif
+
+	FCoreDelegates::OnBeginFrame.AddUObject(this, &URemoteControlPreset::OnBeginFrame);
+	FCoreDelegates::OnEndFrame.AddUObject(this, &URemoteControlPreset::OnEndFrame);
 }
 
 void URemoteControlPreset::UnregisterDelegates()
 {
+	FCoreDelegates::OnBeginFrame.RemoveAll(this);
+	FCoreDelegates::OnEndFrame.RemoveAll(this);
+
 #if WITH_EDITOR
 	for (TWeakObjectPtr<UBlueprint> Blueprint : BlueprintsWithRegisteredDelegates)
 	{
@@ -1634,7 +1701,7 @@ void URemoteControlPreset::UnregisterDelegates()
 	
 	FEditorDelegates::MapChange.RemoveAll(this);
 
-	FCoreDelegates::OnEndFrame.RemoveAll(this);
+
 	FCoreUObjectDelegates::OnObjectsReplaced.RemoveAll(this);
 	
 	FEditorDelegates::EndPIE.RemoveAll(this);
@@ -1736,15 +1803,6 @@ void URemoteControlPreset::OnReplaceObjects(const TMap<UObject*, UObject*>& Repl
 	}
 }
 
-void URemoteControlPreset::OnEndFrame()
-{
-	if (PerFrameUpdatedEntities.Num())
-	{
-		OnEntitiesUpdatedDelegate.Broadcast(this, PerFrameUpdatedEntities);
-		PerFrameUpdatedEntities.Empty();
-	}
-}
-
 void URemoteControlPreset::OnMapChange(uint32)
 {
 	// Delay the refresh in order for the old actors to be invalid.
@@ -1790,5 +1848,98 @@ void URemoteControlPreset::OnBlueprintRecompiled(UBlueprint* Blueprint)
 
 #endif
 
-#undef LOCTEXT_NAMESPACE /* RemoteControlPreset */ 
+void URemoteControlPreset::OnBeginFrame()
+{
+	TRACE_CPUPROFILER_EVENT_SCOPE(URemoteControlPreset::OnBeginFrame);
+	PropertyChangeWatchFrameCounter++;
+
+	if (PropertyChangeWatchFrameCounter == CVarRemoteControlFramesBetweenPropertyWatch.GetValueOnGameThread() - 1)
+	{
+		PropertyChangeWatchFrameCounter = 0;
+		for (TPair<FGuid, FRCPropertyWatcher>& Entry : PropertyWatchers)
+		{
+			Entry.Value.CheckForChange();
+		}
+	}
+}
 
+void URemoteControlPreset::OnEndFrame()
+{
+	TRACE_CPUPROFILER_EVENT_SCOPE(URemoteControlPreset::OnEndFrame);
+	if (PerFrameUpdatedEntities.Num())
+	{
+		OnEntitiesUpdatedDelegate.Broadcast(this, PerFrameUpdatedEntities);
+		PerFrameUpdatedEntities.Empty();
+	}
+
+	if (PerFrameModifiedProperties.Num())
+	{
+		OnPropertyChangedDelegate.Broadcast(this, PerFrameModifiedProperties);
+		PerFrameModifiedProperties.Empty();
+	}
+}
+
+URemoteControlPreset::FRCPropertyWatcher::FRCPropertyWatcher(const TSharedPtr<FRemoteControlProperty>& InWatchedProperty, FSimpleDelegate&& InOnWatchedValueChanged)
+	: OnWatchedValueChanged(MoveTemp(InOnWatchedValueChanged))
+	, WatchedProperty(InWatchedProperty)
+{
+	if (TOptional<FRCFieldResolvedData> ResolvedData = GetWatchedPropertyResolvedData())
+	{
+		SetLastFrameValue(*ResolvedData);
+	}
+}
+
+void URemoteControlPreset::FRCPropertyWatcher::CheckForChange()
+{
+	TRACE_CPUPROFILER_EVENT_SCOPE(FRCPropertyWatcher::CheckForChange);
+	if (TOptional<FRCFieldResolvedData> ResolvedData = GetWatchedPropertyResolvedData())
+	{
+		if (TSharedPtr<FRemoteControlProperty> RCProperty = WatchedProperty.Pin())
+		{
+			const void* NewValueAddress = ResolvedData->Field->ContainerPtrToValuePtr<void>(ResolvedData->ContainerAddress);
+			if (ResolvedData->Field->GetSize() != LastFrameValue.Num() || !ResolvedData->Field->Identical(LastFrameValue.GetData(), NewValueAddress))
+			{
+				SetLastFrameValue(*ResolvedData);
+				OnWatchedValueChanged.ExecuteIfBound();
+			}
+		}
+	}
+}
+
+TOptional<FRCFieldResolvedData> URemoteControlPreset::FRCPropertyWatcher::GetWatchedPropertyResolvedData() const
+{
+	TOptional<FRCFieldResolvedData> ResolvedData;
+	
+	if (TSharedPtr<FRemoteControlProperty> RCProperty = WatchedProperty.Pin())
+	{
+		if (!RCProperty->FieldPathInfo.IsResolved())
+		{
+			// In theory all objects should have the same value if they have an exposed property.
+			TArray<UObject*> Objects = RCProperty->GetBoundObjects();
+			if (Objects.Num() != 0)
+			{
+				RCProperty->FieldPathInfo.Resolve(Objects[0]);
+			}
+		}
+
+		if (RCProperty->FieldPathInfo.IsResolved())
+		{
+			ResolvedData = RCProperty->FieldPathInfo.GetResolvedData();
+		}
+	}
+
+	return ResolvedData;
+}
+
+void URemoteControlPreset::FRCPropertyWatcher::SetLastFrameValue(const FRCFieldResolvedData& ResolvedData)
+{
+	TRACE_CPUPROFILER_EVENT_SCOPE(FRCPropertyWatcher::SetLastFrameValue);
+	checkSlow(ResolvedData.Field);
+	checkSlow(ResolvedData.ContainerAddress);
+	
+	const void* NewValueAddress = ResolvedData.Field->ContainerPtrToValuePtr<void>(ResolvedData.ContainerAddress);
+	LastFrameValue.SetNumUninitialized(ResolvedData.Field->GetSize());
+	ResolvedData.Field->CopyCompleteValue(LastFrameValue.GetData(), NewValueAddress);
+}
+
+#undef LOCTEXT_NAMESPACE /* RemoteControlPreset */ 
diff --git a/Engine/Plugins/VirtualProduction/RemoteControl/Source/RemoteControl/Public/RemoteControlPreset.h b/Engine/Plugins/VirtualProduction/RemoteControl/Source/RemoteControl/Public/RemoteControlPreset.h
--- a/Engine/Plugins/VirtualProduction/RemoteControl/Source/RemoteControl/Public/RemoteControlPreset.h
+++ b/Engine/Plugins/VirtualProduction/RemoteControl/Source/RemoteControl/Public/RemoteControlPreset.h
@@ -709,16 +709,24 @@ public:
 	FOnPresetEntityEvent& OnEntityExposed() { return OnEntityExposedDelegate; }
 	FOnPresetEntityEvent& OnEntityUnexposed() { return OnEntityUnexposedDelegate; }
 
-	DECLARE_MULTICAST_DELEGATE_TwoParams(FOnPresetExposedPropertyChanged, URemoteControlPreset* /*Preset*/, const FRemoteControlProperty& /*Property*/);
-	FOnPresetExposedPropertyChanged& OnExposedPropertyChanged() { return OnPropertyChangedDelegate; }
+	DECLARE_MULTICAST_DELEGATE_TwoParams(FOnPresetExposedPropertiesModified, URemoteControlPreset* /*Preset*/, const TSet<FGuid>& /* ModifiedProperties */);
+	/**
+	 *  Delegate called with the list of exposed property that were modified in the last frame.
+	 */
+	FOnPresetExposedPropertiesModified& OnExposedPropertiesModified() { return OnPropertyChangedDelegate; }
 
 	DECLARE_MULTICAST_DELEGATE_TwoParams(FOnPresetEntitiesUpdatedEvent, URemoteControlPreset* /*Preset*/, const TSet<FGuid>& /* Modified Entities */);
+	/**
+	 * Delegate called when the exposed entity wrapper itself is updated (ie. binding change, rename)  
+	 */
 	FOnPresetEntitiesUpdatedEvent& OnEntitiesUpdated() { return OnEntitiesUpdatedDelegate; }
 	
 	DECLARE_MULTICAST_DELEGATE_TwoParams(FOnPresetPropertyExposed, URemoteControlPreset* /*Preset*/, FName /*ExposedLabel*/);
+	UE_DEPRECATED(4.27, "This delegate is deprecated, use OnEntityExposed instead.")
 	FOnPresetPropertyExposed& OnPropertyExposed() { return OnPropertyExposedDelegate; }
 
 	DECLARE_MULTICAST_DELEGATE_TwoParams(FOnPresetPropertyUnexposed, URemoteControlPreset* /*Preset*/, FName /*ExposedLabel*/);
+	UE_DEPRECATED(4.27, "This delegate is deprecated, use OnEntityUnexposed instead.")
 	FOnPresetPropertyUnexposed& OnPropertyUnexposed() { return OnPropertyUnexposedDelegate; }
 
 	DECLARE_MULTICAST_DELEGATE_ThreeParams(FOnPresetFieldRenamed, URemoteControlPreset* /*Preset*/, FName /*OldExposedLabel*/, FName /**NewExposedLabel*/);
@@ -784,10 +792,13 @@ private:
 	void OnActorDeleted(AActor* Actor);
 	void OnPieEvent(bool);
 	void OnReplaceObjects(const TMap<UObject*, UObject*>& ReplacementObjectMap);
-	void OnEndFrame();
 	void OnMapChange(uint32 /*MapChangeFlags*/);
 	void OnBlueprintRecompiled(UBlueprint* Blueprint);
 #endif
+
+	//~ Frame events handlers.
+	void OnBeginFrame();
+	void OnEndFrame();
 	
 	/** Get a field ptr using it's id. */
 	FRemoteControlField* GetFieldPtr(FGuid FieldId);
@@ -823,6 +834,15 @@ private:
 
 	/** Register an event triggered when the exposed function's owner blueprint is compiled. */
 	void RegisterOnCompileEvent(const TSharedPtr<FRemoteControlFunction>& RCFunction);
+
+	/** Create a property watcher that will trigger the property change delegate upon having a different value from the last frame. */
+	void CreatePropertyWatcher(const TSharedPtr<FRemoteControlProperty>& RCProperty);
+
+	/** Returns whether a given property should have a watcher that checks for property changes across frames. */
+	bool PropertyShouldBeWatched(const TSharedPtr<FRemoteControlProperty>& RCProperty) const;
+
+	/** Create property watchers for exposed properties that need them. */
+	void CreatePropertyWatchers();
 	
 private:
 	/** Preset unique ID */
@@ -852,7 +872,7 @@ private:
 	/** Delegate triggered when entities are modified and may need to be re-resolved. */
 	FOnPresetEntitiesUpdatedEvent OnEntitiesUpdatedDelegate;
 	/** Delegate triggered when an exposed property value has changed. */
-	FOnPresetExposedPropertyChanged OnPropertyChangedDelegate;
+	FOnPresetExposedPropertiesModified OnPropertyChangedDelegate;
 	/** Delegate triggered when a new property has been exposed. */
 	FOnPresetPropertyExposed OnPropertyExposedDelegate;
 	/** Delegate triggered when a property has been unexposed. */
@@ -874,6 +894,10 @@ private:
 	/** Caches object modifications during a frame. */
 	TMap<FGuid, FPreObjectsModifiedCache> PreObjectsModifiedCache;
 	TMap<FGuid, FPreObjectsModifiedCache> PreObjectsModifiedActorCache;
+
+	/** Cache properties that were modified during a frame. */
+	TSet<FGuid> PerFrameModifiedProperties;
+	
 	/** Cache entities updated during a frame. */
 	TSet<FGuid> PerFrameUpdatedEntities; 
 
@@ -883,6 +907,35 @@ private:
 	/** Holds manager that handles rebinding unbound entities upon load or map change. */
 	TPimplPtr<FRemoteControlPresetRebindingManager> RebindingManager;
 
+	/**
+	 * Property watcher that triggers a delegate when the watched property is modified.
+	 */
+	struct FRCPropertyWatcher
+	{
+		FRCPropertyWatcher(const TSharedPtr<FRemoteControlProperty>& InWatchedProperty, FSimpleDelegate&& InOnWatchedValueChanged);
+
+		/** Checks if the property value has changed since the last frame and updates the last frame value. */
+		void CheckForChange();
+
+	private:
+		/** Optionally resolve the property path if possible. */
+		TOptional<FRCFieldResolvedData> GetWatchedPropertyResolvedData() const;
+		/** Store the latest property value in LastFrameValue. */
+		void SetLastFrameValue(const FRCFieldResolvedData& ResolvedData);
+	private:
+		/** Delegate called when the watched property changes values. */
+		FSimpleDelegate OnWatchedValueChanged;
+		/** Weak pointer to the remote control property. */
+		TWeakPtr<FRemoteControlProperty> WatchedProperty;
+		/** Latest property value as bytes. */
+		TArray<uint8> LastFrameValue;
+	};
+	/** Map of property watchers that should trigger the RC property change delegate upon change. */
+	TMap<FGuid, FRCPropertyWatcher> PropertyWatchers;
+
+	/** Frame counter for delaying property change checks. */
+	int8 PropertyChangeWatchFrameCounter = 0;
+
 #if WITH_EDITOR
 	/** List of blueprints for which we have registered events. */
 	TSet<TWeakObjectPtr<UBlueprint>> BlueprintsWithRegisteredDelegates;
@@ -892,4 +945,3 @@ private:
 	friend FRemoteControlPresetLayout;
 	friend FRemoteControlEntity;
 };
-
diff --git a/Engine/Plugins/VirtualProduction/RemoteControl/Source/WebRemoteControl/Private/WebSocketMessageHandler.cpp b/Engine/Plugins/VirtualProduction/RemoteControl/Source/WebRemoteControl/Private/WebSocketMessageHandler.cpp
--- a/Engine/Plugins/VirtualProduction/RemoteControl/Source/WebRemoteControl/Private/WebSocketMessageHandler.cpp
+++ b/Engine/Plugins/VirtualProduction/RemoteControl/Source/WebRemoteControl/Private/WebSocketMessageHandler.cpp
@@ -82,7 +82,7 @@ void FWebSocketMessageHandler::HandleWebSocketPresetRegister(const FRemoteContro
 		ClientIds = &WebSocketNotificationMap.Add(Preset->GetFName());
 
 		//Register to any useful callback for the given preset
-		Preset->OnExposedPropertyChanged().AddRaw(this, &FWebSocketMessageHandler::OnPresetExposedPropertyChanged);
+		Preset->OnExposedPropertiesModified().AddRaw(this, &FWebSocketMessageHandler::OnPresetExposedPropertiesModified);
 		Preset->OnEntityExposed().AddRaw(this, &FWebSocketMessageHandler::OnPropertyExposed);
 		Preset->OnEntityUnexposed().AddRaw(this, &FWebSocketMessageHandler::OnPropertyUnexposed);
 		Preset->OnFieldRenamed().AddRaw(this, &FWebSocketMessageHandler::OnFieldRenamed);
@@ -112,7 +112,7 @@ void FWebSocketMessageHandler::HandleWebSocketPresetUnregister(const FRemoteCont
 void FWebSocketMessageHandler::ProcessChangedProperties()
 {
 	//Go over each property that were changed for each preset
-	for (const TPair<FName, TMap<FGuid, TArray<FRemoteControlProperty>>>& Entry : PerFramePropertyChanged)
+	for (const TPair<FName, TMap<FGuid, TSet<FGuid>>>& Entry : PerFrameModifiedProperties)
 	{
 		if (!ShouldProcessEventForPreset(Entry.Key) || !Entry.Value.Num())
 		{
@@ -128,7 +128,7 @@ void FWebSocketMessageHandler::ProcessChangedProperties()
 		UE_LOG(LogRemoteControl, VeryVerbose, TEXT("(%s) Broadcasting properties changed event."), *Preset->GetName());
 
 		// Each client will have a custom payload that doesnt contain the events it triggered.
-		for (const TPair<FGuid, TArray<FRemoteControlProperty>>& ClientToEventsPair : Entry.Value)
+		for (const TPair<FGuid, TSet<FGuid>>& ClientToEventsPair : Entry.Value)
 		{
 			TArray<uint8> WorkingBuffer;
 			if (ClientToEventsPair.Value.Num() && WritePropertyChangeEventPayload(Preset, ClientToEventsPair.Value, WorkingBuffer))
@@ -140,7 +140,7 @@ void FWebSocketMessageHandler::ProcessChangedProperties()
 		}
 	}
 
-	PerFramePropertyChanged.Empty();
+	PerFrameModifiedProperties.Empty();
 }
 
 void FWebSocketMessageHandler::ProcessChangedActorProperties()
@@ -191,7 +191,7 @@ void FWebSocketMessageHandler::OnPropertyExposed(URemoteControlPreset* Owner, co
 	PerFrameAddedProperties.FindOrAdd(Owner->GetFName()).AddUnique(EntityId);
 }
 
-void FWebSocketMessageHandler::OnPresetExposedPropertyChanged(URemoteControlPreset* Owner, const FRemoteControlProperty& PropertyChanged)
+void FWebSocketMessageHandler::OnPresetExposedPropertiesModified(URemoteControlPreset* Owner, const TSet<FGuid>& ModifiedPropertyIds)
 {
 	if (Owner == nullptr)
 	{
@@ -204,15 +204,16 @@ void FWebSocketMessageHandler::OnPresetExposedPropertyChanged(URemoteControlPres
 	}
 
 	//Cache the property field that changed for end of frame notification
-	TMap<FGuid, TArray<FRemoteControlProperty>>& EventsForClient = PerFramePropertyChanged.FindOrAdd(Owner->GetFName());
-	// Dont send events to the client that triggered it.
+	TMap<FGuid, TSet<FGuid>>& EventsForClient = PerFrameModifiedProperties.FindOrAdd(Owner->GetFName());
+	
+	// Don't send events to the client that triggered it.
 	if (TArray<FGuid>* SubscribedClients = WebSocketNotificationMap.Find(Owner->GetFName()))
 	{
 		for (const FGuid& Client : *SubscribedClients)
 		{
 			if (Client != ActingClientId)
 			{
-				EventsForClient.FindOrAdd(Client).AddUnique(PropertyChanged);
+				EventsForClient.FindOrAdd(Client).Append(ModifiedPropertyIds);
 			}
 		}
 	}
@@ -478,7 +479,7 @@ bool FWebSocketMessageHandler::ShouldProcessEventForPreset(FName PresetName) con
 	return WebSocketNotificationMap.Contains(PresetName) && WebSocketNotificationMap[PresetName].Num() > 0;
 }
 
-bool FWebSocketMessageHandler::WritePropertyChangeEventPayload(URemoteControlPreset* InPreset, const TArray<FRemoteControlProperty>& InEvents, TArray<uint8>& OutBuffer)
+bool FWebSocketMessageHandler::WritePropertyChangeEventPayload(URemoteControlPreset* InPreset, const TSet<FGuid>& InModifiedPropertyIds, TArray<uint8>& OutBuffer)
 {
 	bool bHasProperty = false;
 
@@ -498,37 +499,40 @@ bool FWebSocketMessageHandler::WritePropertyChangeEventPayload(URemoteControlPre
 		//All exposed properties of this preset that changed
 		JsonWriter->WriteArrayStart();
 		{
-			for (const FRemoteControlProperty& Property : InEvents)
+			for (const FGuid& RCPropertyId : InModifiedPropertyIds)
 			{
 				bHasProperty = true;
 
 				FRCObjectReference ObjectRef;
-
-				//Property object
-				JsonWriter->WriteObjectStart();
+				if (TSharedPtr<FRemoteControlProperty> RCProperty = InPreset->GetExposedEntity<FRemoteControlProperty>(RCPropertyId).Pin())
 				{
-					JsonWriter->WriteValue(TEXT("PropertyLabel"), *Property.GetLabel().ToString());
-					JsonWriter->WriteValue(TEXT("Id"), *Property.GetId().ToString());
-
-					for (UObject* Object : Property.GetBoundObjects())
+					//Property object
+					JsonWriter->WriteObjectStart();
 					{
-						bHasProperty = true;
+						JsonWriter->WriteValue(TEXT("PropertyLabel"), *RCProperty->GetLabel().ToString());
+						JsonWriter->WriteValue(TEXT("Id"), *RCProperty->GetId().ToString());
 
-						IRemoteControlModule::Get().ResolveObjectProperty(ERCAccess::READ_ACCESS, Object, Property.FieldPathInfo.ToString(), ObjectRef);
+						for (UObject* Object : RCProperty->GetBoundObjects())
+						{
+							bHasProperty = true;
 
-						JsonWriter->WriteValue(TEXT("ObjectPath"), Object->GetPathName());
-						JsonWriter->WriteIdentifierPrefix(TEXT("PropertyValue"));
+							IRemoteControlModule::Get().ResolveObjectProperty(ERCAccess::READ_ACCESS, Object, RCProperty->FieldPathInfo.ToString(), ObjectRef);
 
-						RemotePayloadSerializer::SerializePartial(
-							[&ObjectRef](FJsonStructSerializerBackend& SerializerBackend)
-							{
-								return IRemoteControlModule::Get().GetObjectProperties(ObjectRef, SerializerBackend);
-							}
-						, Writer);
+							JsonWriter->WriteValue(TEXT("ObjectPath"), Object->GetPathName());
+							JsonWriter->WriteIdentifierPrefix(TEXT("PropertyValue"));
 
+							RemotePayloadSerializer::SerializePartial(
+								[&ObjectRef](FJsonStructSerializerBackend& SerializerBackend)
+								{
+									return IRemoteControlModule::Get().GetObjectProperties(ObjectRef, SerializerBackend);
+								}
+							, Writer);
+
+						}
 					}
+					JsonWriter->WriteObjectEnd();
 				}
-				JsonWriter->WriteObjectEnd();
+
 			}
 		}
 		JsonWriter->WriteArrayEnd();
diff --git a/Engine/Plugins/VirtualProduction/RemoteControl/Source/WebRemoteControl/Private/WebSocketMessageHandler.h b/Engine/Plugins/VirtualProduction/RemoteControl/Source/WebRemoteControl/Private/WebSocketMessageHandler.h
--- a/Engine/Plugins/VirtualProduction/RemoteControl/Source/WebRemoteControl/Private/WebSocketMessageHandler.h
+++ b/Engine/Plugins/VirtualProduction/RemoteControl/Source/WebRemoteControl/Private/WebSocketMessageHandler.h
@@ -40,7 +40,7 @@ private:
 	void HandleWebSocketPresetUnregister(const FRemoteControlWebSocketMessage& WebSocketMessage);
 
 	//Preset callbacks
-	void OnPresetExposedPropertyChanged(URemoteControlPreset* Owner, const FRemoteControlProperty& PropertyChanged);
+	void OnPresetExposedPropertiesModified(URemoteControlPreset* Owner, const TSet<FGuid>& ModifiedPropertyIds);
 	void OnPropertyExposed(URemoteControlPreset* Owner,  const FGuid& EntityId);
 	void OnPropertyUnexposed(URemoteControlPreset* Owner, const FGuid& EntityId);
 	void OnFieldRenamed(URemoteControlPreset* Owner, FName OldFieldLabel, FName NewFieldLabel);
@@ -86,7 +86,7 @@ private:
 	/**
 	 * Write the provided list of events to a buffer.
 	 */
-	bool WritePropertyChangeEventPayload(URemoteControlPreset* InPreset, const TArray<FRemoteControlProperty>& InEvents, TArray<uint8>& OutBuffer);
+	bool WritePropertyChangeEventPayload(URemoteControlPreset* InPreset, const TSet<FGuid>& InModifiedPropertyIds, TArray<uint8>& OutBuffer);
 
 	/**
 	 * Write the provided list of actor modifications to a buffer.
@@ -114,7 +114,7 @@ private:
 	TMap<FGuid, FRCClientConfig> ClientConfigMap;
 
 	/** Properties that changed for a frame, per preset.  */
-	TMap<FName, TMap<FGuid, TArray<FRemoteControlProperty>>> PerFramePropertyChanged;
+	TMap<FName, TMap<FGuid, TSet<FGuid>>> PerFrameModifiedProperties;
 
 	/** Properties that changed on an exposed actor for a given client, for a frame, per preset.  */
 	TMap<FName, TMap<FGuid, TMap<FRemoteControlActor, TArray<FRCObjectReference>>>> PerFrameActorPropertyChanged;
