Put RemoteControlCommon back to Editor it is dependent on editor only types.
#rnx

diff --git a/Engine/Plugins/VirtualProduction/RemoteControl/RemoteControl.uplugin b/Engine/Plugins/VirtualProduction/RemoteControl/RemoteControl.uplugin
--- a/Engine/Plugins/VirtualProduction/RemoteControl/RemoteControl.uplugin
+++ b/Engine/Plugins/VirtualProduction/RemoteControl/RemoteControl.uplugin
@@ -41,7 +41,7 @@
 		},
 		{
 			"Name": "RemoteControlCommon",
-			"Type": "Runtime",
+			"Type": "Editor",
 			"LoadingPhase": "Default"
 		},
 		{
diff --git a/Engine/Plugins/VirtualProduction/RemoteControl/Source/RemoteControl/RemoteControl.Build.cs b/Engine/Plugins/VirtualProduction/RemoteControl/Source/RemoteControl/RemoteControl.Build.cs
--- a/Engine/Plugins/VirtualProduction/RemoteControl/Source/RemoteControl/RemoteControl.Build.cs
+++ b/Engine/Plugins/VirtualProduction/RemoteControl/Source/RemoteControl/RemoteControl.Build.cs
@@ -10,7 +10,6 @@ public class RemoteControl : ModuleRules
 			new string[] {
 				"Core",
 				"CoreUObject",
-				"RemoteControlCommon",
 			}
 		);
 
@@ -26,6 +25,12 @@ public class RemoteControl : ModuleRules
 
 		if (Target.bBuildEditor)
 		{
+			PublicDependencyModuleNames.AddRange(
+				new string[] {
+					"RemoteControlCommon",
+				}
+			);
+
 			PrivateDependencyModuleNames.AddRange(
 				new string[] {
 					"DeveloperSettings",
