RemoteControlCommon type traits/utilities

diff --git a/Engine/Plugins/VirtualProduction/RemoteControl/Source/RemoteControlCommon/Private/RCPropertyContainer.cpp b/Engine/Plugins/VirtualProduction/RemoteControl/Source/RemoteControlCommon/Private/RCPropertyContainer.cpp
--- a/Engine/Plugins/VirtualProduction/RemoteControl/Source/RemoteControlCommon/Private/RCPropertyContainer.cpp
+++ b/Engine/Plugins/VirtualProduction/RemoteControl/Source/RemoteControlCommon/Private/RCPropertyContainer.cpp
@@ -2,17 +2,19 @@
 
 #include "RCPropertyContainer.h"
 
-#include "BlueprintActionDatabase.h"
-#include "Editor.h"
 #include "RemoteControlCommonModule.h"
+#include "Engine/Engine.h"
 #include "Engine/MemberReference.h"
-#include "Kismet2/BlueprintEditorUtils.h"
-#include "Kismet2/CompilerResultsLog.h"
-#include "Kismet2/StructureEditorUtils.h"
 #include "UObject/EnumProperty.h"
 #include "UObject/FieldPathProperty.h"
 #include "UObject/TextProperty.h"
 
+#if WITH_EDITOR
+#include "BlueprintActionDatabase.h"
+#include "Kismet2/BlueprintEditorUtils.h"
+#include "Kismet2/StructureEditorUtils.h"
+#endif
+
 #define LOCTEXT_NAMESPACE "RCPropertyContainer"
 
 namespace PropertyContainers
@@ -23,8 +25,9 @@ namespace PropertyContainers
 
 		check(InOwner);
 		check(InSrcProperty);
+		check(InSrcProperty->IsValidLowLevel());
 
-		URCPropertyContainerRegistry* ContainerRegistry = GEditor->GetEditorSubsystem<URCPropertyContainerRegistry>();
+		URCPropertyContainerRegistry* ContainerRegistry = GEngine->GetEngineSubsystem<URCPropertyContainerRegistry>();
 		FName PropertyTypeName = InSrcProperty->GetClass()->GetFName();
 		if(const FStructProperty* StructProperty = CastField<FStructProperty>(InSrcProperty))
 		{
@@ -431,20 +434,35 @@ namespace PropertyContainers
 #endif
 }
 
-void URCPropertyContainerBase::SetValue(const uint8* InData)
+void URCPropertyContainerBase::SetValue(const uint8* InData, const SIZE_T& InSize)
 {
 	this->Modify();
 
-	FProperty* ValueProperty = GetClass()->FindPropertyByName("Value");
-	uint8* DstData = ValueProperty->ContainerPtrToValuePtr<uint8>(this);
-	FMemory::Memcpy(DstData, InData, ValueProperty->ElementSize);
+	FProperty* Prp = GetValueProperty();
+	uint8* DstData = Prp->ContainerPtrToValuePtr<uint8>(this);
+	FMemory::Memcpy(DstData, InData, InSize > 0 ? InSize : Prp->GetSize());
+}
+
+SIZE_T URCPropertyContainerBase::GetValue(uint8* OutData)
+{
+	FProperty* Prp = GetValueProperty();
+	uint8* SrcData = Prp->ContainerPtrToValuePtr<uint8>(this);
+	FMemory::Memcpy(OutData, SrcData, Prp->GetSize());
+	return Prp->GetSize();
 }
 
-void URCPropertyContainerBase::GetValue(uint8* OutData)
+FProperty* URCPropertyContainerBase::GetValueProperty()
 {
-	FProperty* ValueProperty = GetClass()->FindPropertyByName("Value");
-	uint8* SrcData = ValueProperty->ContainerPtrToValuePtr<uint8>(this);
-	FMemory::Memcpy(OutData, SrcData, ValueProperty->ElementSize);
+	if(ValueProperty.IsValid())
+	{
+		return ValueProperty.Get();
+	}
+
+	static FName ValuePropertyName = TEXT("Value");
+	ValueProperty = FindFProperty<FProperty>(GetClass(), ValuePropertyName);
+	ensure(ValueProperty.IsValid());
+
+	return ValueProperty.Get();
 }
 
 FName FRCPropertyContainerKey::ToClassName() const
@@ -454,11 +472,13 @@ FName FRCPropertyContainerKey::ToClassName() const
 	return FName(FString::Printf(TEXT("PropertyContainer_%s"), *ValueTrimmed));	
 }
 
+
 URCPropertyContainerBase* URCPropertyContainerRegistry::CreateContainer(UObject* InOwner, const FName& InValueTypeName, const FProperty* InValueSrcProperty)
 {
-	check(InValueSrcProperty);
+	check(InValueSrcProperty && InValueSrcProperty->IsValidLowLevel());
 
-	const TSubclassOf<URCPropertyContainerBase> ClassForPropertyType = FindOrAddContainerClass(InValueTypeName, InValueSrcProperty);
+	const FName SanitizedName = FName(InValueTypeName.ToString().Replace(TEXT(":"), TEXT("_")));
+	const TSubclassOf<URCPropertyContainerBase> ClassForPropertyType = FindOrAddContainerClass(SanitizedName, InValueSrcProperty);
 	if (ClassForPropertyType)
 	{
 		return NewObject<URCPropertyContainerBase>(InOwner
@@ -467,7 +487,7 @@ URCPropertyContainerBase* URCPropertyContainerRegistry::CreateContainer(UObject*
 	}
 	else
 	{
-		UE_LOG(LogRemoteControlCommon, Warning, TEXT("Could not create PropertyContainer found for %s"), *InValueTypeName.ToString());
+		UE_LOG(LogRemoteControlCommon, Warning, TEXT("Could not create PropertyContainer found for %s"), *SanitizedName.ToString());
 		return nullptr;
 	}
 }
@@ -476,7 +496,13 @@ TSubclassOf<URCPropertyContainerBase>& URCPropertyContainerRegistry::FindOrAddCo
 {
 	check(InValueSrcProperty);
 
-	const FRCPropertyContainerKey Key = FRCPropertyContainerKey{FName(InValueTypeName.ToString() + "_" + InValueSrcProperty->GetPathName())};
+	static const FString EmptyStr = TEXT("");
+	FString PropertyPathStr = InValueSrcProperty->GetPathName();
+	// format like ObjectName:PropertyName, this ensures theres a unique property container for each defined property rather than for each type
+	PropertyPathStr.ReplaceInline(TEXT(":"), TEXT("_"));
+	PropertyPathStr.ReplaceInline(*(InValueSrcProperty->GetOutermost()->GetPathName() + TEXT(".")), *EmptyStr);
+
+	const FRCPropertyContainerKey Key = FRCPropertyContainerKey{FName(PropertyPathStr)};
 	if(TSubclassOf<URCPropertyContainerBase>* ExistingContainerClass = CachedContainerClasses.Find(Key))
 	{
 		return *ExistingContainerClass;
@@ -488,9 +514,12 @@ TSubclassOf<URCPropertyContainerBase>& URCPropertyContainerRegistry::FindOrAddCo
 	ContainerClass->SetSuperStruct(ParentClass);
 
 	FProperty* ValueProperty = CastField<FProperty>(FField::Duplicate(InValueSrcProperty, ContainerClass, "Value"));
+#if WITH_EDITORONLY_DATA
 	FField::CopyMetaData(InValueSrcProperty, ValueProperty);
+#endif
 	ValueProperty->SetFlags(RF_Transient);
-	ValueProperty->PropertyFlags |= CPF_Edit | CPF_BlueprintVisible;
+	ValueProperty->PropertyFlags |= CPF_Edit; // add this flag(s)
+	//ValueProperty->PropertyFlags &= ~CPF_BlueprintVisible; // clear this flag(s)
 
 	ContainerClass->AddCppProperty(ValueProperty);
 
@@ -498,11 +527,13 @@ TSubclassOf<URCPropertyContainerBase>& URCPropertyContainerRegistry::FindOrAddCo
 	ContainerClass->StaticLink(true);
 	ContainerClass->AssembleReferenceTokenStream(true);
 
+#if WITH_EDITOR
 	if (FBlueprintActionDatabase* ActionDB = FBlueprintActionDatabase::TryGet())
 	{
 		// Notify Blueprints that there is a new class to add to the action list
 		ActionDB->RefreshClassActions(ContainerClass);
 	}
+#endif
 
 	return CachedContainerClasses.Add(Key, MoveTemp(ContainerClass));	
 }
diff --git a/Engine/Plugins/VirtualProduction/RemoteControl/Source/RemoteControlCommon/Private/RCPropertyUtilities.cpp b/Engine/Plugins/VirtualProduction/RemoteControl/Source/RemoteControlCommon/Private/RCPropertyUtilities.cpp
new file mode 100644
--- /dev/null
+++ b/Engine/Plugins/VirtualProduction/RemoteControl/Source/RemoteControlCommon/Private/RCPropertyUtilities.cpp
@@ -0,0 +1,170 @@
+﻿// Copyright Epic Games, Inc. All Rights Reserved.
+
+#include "RCPropertyUtilities.h"
+
+#include "CoreMinimal.h"
+
+#include "RCTypeUtilities.h"
+#include "UObject/TextProperty.h"
+#include "UObject/UnrealType.h"
+
+#if WITH_EDITOR
+
+template <>
+bool RemoteControlPropertyUtilities::FromBytes<FStructProperty>(const FRCPropertyVariant& InSrcData, FRCPropertyVariant& InDst)
+{
+	InDst.GetProperty()->CopyCompleteValue((uint8*)InDst.GetPropertyData(), (uint8*)InSrcData.GetPropertyData());
+	
+	return true;
+}
+
+template <>
+bool RemoteControlPropertyUtilities::FromBytes<FArrayProperty>(const FRCPropertyVariant& InSrcData, FRCPropertyVariant& InDst)
+{
+	const FArrayProperty* Property = InDst.GetProperty<FArrayProperty>();
+
+	FScriptArrayHelper SrcArrayHelper(Property, Property->ContainerPtrToValuePtr<void>(InSrcData.GetPropertyData()));
+	FScriptArrayHelper DstArrayHelper(Property, Property->ContainerPtrToValuePtr<void>(InDst.GetPropertyData()));
+
+	const int32 Num = SrcArrayHelper.Num();	
+	const FProperty* InnerProperty = Property->Inner;
+	if(InnerProperty->HasAnyPropertyFlags(CPF_IsPlainOldData))
+	{
+		DstArrayHelper.EmptyAndAddValues(Num);
+	}
+	else
+	{
+		DstArrayHelper.EmptyAndAddUninitializedValues(Num);
+	}
+
+	for(int32 Idx = 0; Idx < Num; ++Idx)
+	{
+		FRCPropertyVariant Dst{InnerProperty, DstArrayHelper.GetRawPtr(Idx)};
+		FromBytes<FProperty>({InnerProperty, SrcArrayHelper.GetRawPtr(Idx)}, Dst);
+	}
+
+	return true;
+}
+
+template <>
+bool RemoteControlPropertyUtilities::FromBytes<FSetProperty>(const FRCPropertyVariant& InSrcData, FRCPropertyVariant& InDst)
+{
+	const FSetProperty* Property = InDst.GetProperty<FSetProperty>();
+	
+	FScriptSetHelper SrcSetHelper(Property, Property->ContainerPtrToValuePtr<void>(InSrcData.GetPropertyData()));
+	FScriptSetHelper DstSetHelper(Property, Property->ContainerPtrToValuePtr<void>(InDst.GetPropertyData()));
+
+	const FProperty* InnerProperty = Property->ElementProp;
+	
+	DstSetHelper.EmptyElements(SrcSetHelper.Num());
+
+	for(int32 Idx = 0; Idx < SrcSetHelper.Num(); ++Idx)
+	{
+		DstSetHelper.AddUninitializedValue();
+		FRCPropertyVariant Dst{InnerProperty, DstSetHelper.GetElementPtr(Idx)};
+		FromBytes<FProperty>({InnerProperty, SrcSetHelper.GetElementPtr(Idx)}, Dst);
+	}
+
+	return true;
+}
+
+template <>
+bool RemoteControlPropertyUtilities::FromBytes<FMapProperty>(const FRCPropertyVariant& InSrcData, FRCPropertyVariant& InDst)
+{
+	const FMapProperty* Property = InDst.GetProperty<FMapProperty>();
+	
+	FScriptMapHelper SrcMapHelper(Property, Property->ContainerPtrToValuePtr<void>(InSrcData.GetPropertyData()));
+	FScriptMapHelper DstMapHelper(Property, Property->ContainerPtrToValuePtr<void>(InDst.GetPropertyData()));
+
+	const FProperty* KeyProperty = Property->KeyProp;
+	const FProperty* ValueProperty = Property->ValueProp;
+
+	DstMapHelper.EmptyValues(SrcMapHelper.Num());
+
+	for(int32 Idx = 0; Idx < SrcMapHelper.Num(); ++Idx)
+	{
+		DstMapHelper.AddUninitializedValue();
+
+		FRCPropertyVariant KeyDst{KeyProperty, (void*)DstMapHelper.GetKeyPtr(Idx)};
+		FromBytes<FProperty>({KeyProperty, SrcMapHelper.GetKeyPtr(Idx)}, KeyDst);
+
+		FRCPropertyVariant ValueDst{ValueProperty, (void*)DstMapHelper.GetValuePtr(Idx)};
+		FromBytes<FProperty>({ValueProperty, SrcMapHelper.GetValuePtr(Idx)}, ValueDst);
+	}
+
+	return true;
+}
+
+template <>
+bool RemoteControlPropertyUtilities::ToBytes<FStructProperty>(const FRCPropertyVariant& InSrc, FRCPropertyVariant& OutDst)
+{
+	InSrc.GetProperty()->CopyCompleteValue((uint8*)OutDst.GetPropertyData(), (uint8*)InSrc.GetPropertyData());
+
+	return true;
+}
+
+template <>
+bool RemoteControlPropertyUtilities::ToBytes<FArrayProperty>(const FRCPropertyVariant& InSrc, FRCPropertyVariant& OutDst)
+{
+	const FArrayProperty* Property = InSrc.GetProperty<FArrayProperty>();
+
+	// Dst wasn't yet set to anything
+	if(OutDst.Num() == 0)
+	{
+		OutDst.Init(1);
+		Property->InitializeValueInternal(OutDst.GetPropertyData());	
+	}
+	
+	FScriptArrayHelper SrcArrayHelper(Property, Property->ContainerPtrToValuePtr<void>(InSrc.GetPropertyData()));
+	FScriptArrayHelper DstArrayHelper(Property, Property->ContainerPtrToValuePtr<void>(OutDst.GetPropertyData()));
+
+	const int32 Num = SrcArrayHelper.Num();
+	const FProperty* InnerProperty = Property->Inner;
+	if(InnerProperty->HasAnyPropertyFlags(CPF_IsPlainOldData))
+	{
+		DstArrayHelper.EmptyAndAddValues(Num);
+	}
+	else
+	{
+		DstArrayHelper.EmptyAndAddUninitializedValues(Num);
+	}
+
+	for(int32 Idx = 0; Idx < Num; ++Idx)
+	{
+		FRCPropertyVariant Dst{InnerProperty, (uint8*)DstArrayHelper.GetRawPtr(Idx)};
+		ToBytes<FProperty>({InnerProperty, SrcArrayHelper.GetRawPtr(Idx)}, Dst);
+	}		
+
+	return true;
+}
+
+template <>
+bool RemoteControlPropertyUtilities::ToBytes<FSetProperty>(const FRCPropertyVariant& InSrc, FRCPropertyVariant& OutDst)
+{
+	return false;
+}
+
+template <>
+bool RemoteControlPropertyUtilities::ToBytes<FMapProperty>(const FRCPropertyVariant& InSrc, FRCPropertyVariant& OutDst)
+{
+	return false;
+}
+
+#endif
+template <>
+bool RemoteControlPropertyUtilities::FromBytes<FProperty>(const FRCPropertyVariant& InSrcData, FRCPropertyVariant& InDst)
+{
+	const FProperty* Property = InDst.GetProperty();
+	FOREACH_CAST_PROPERTY(Property, FromBytes<CastPropertyType>(InSrcData, InDst))
+
+	return true;
+}
+
+template <>
+bool RemoteControlPropertyUtilities::ToBytes<FProperty>(const FRCPropertyVariant& InSrc, FRCPropertyVariant& OutDst)
+{
+	const FProperty* Property = InSrc.GetProperty();
+	FOREACH_CAST_PROPERTY(Property, ToBytes<CastPropertyType>(InSrc, OutDst))
+
+	return true;
+}
diff --git a/Engine/Plugins/VirtualProduction/RemoteControl/Source/RemoteControlCommon/Private/RCTypeUtilities.cpp b/Engine/Plugins/VirtualProduction/RemoteControl/Source/RemoteControlCommon/Private/RCTypeUtilities.cpp
new file mode 100644
--- /dev/null
+++ b/Engine/Plugins/VirtualProduction/RemoteControl/Source/RemoteControlCommon/Private/RCTypeUtilities.cpp
@@ -0,0 +1,63 @@
+﻿// Copyright Epic Games, Inc. All Rights Reserved.
+
+#include "RCTypeUtilities.h"
+
+#include "UObject/UnrealType.h"
+
+#if WITH_EDITOR
+#include "PropertyHandle.h"
+#endif
+
+SIZE_T RemoteControlTypeUtilities::GetPropertySize(const FProperty* InProperty, void* InData)
+{
+	check(InProperty);
+
+	const int32 ElementSize = InProperty->ElementSize;
+	if (const FArrayProperty* ArrayProperty = CastField<FArrayProperty>(InProperty))
+	{
+		return FScriptArrayHelper(ArrayProperty, InData).Num() * ElementSize;
+	}
+	if (const FSetProperty* SetProperty = CastField<FSetProperty>(InProperty))
+	{
+		return FScriptSetHelper(SetProperty, InData).Num() * ElementSize;
+	}
+	if (const FMapProperty* MapProperty = CastField<FMapProperty>(InProperty))
+	{
+		return FScriptMapHelper(MapProperty, InData).Num() * ElementSize;
+	}
+
+	return ElementSize;
+}
+
+#if WITH_EDITOR
+SIZE_T RemoteControlTypeUtilities::GetPropertySize(const TSharedPtr<IPropertyHandle>& InPropertyHandle)
+{
+	check(InPropertyHandle.IsValid() && InPropertyHandle->IsValidHandle());
+
+	TArray<void*> Data;
+	InPropertyHandle->AccessRawData(Data);
+	check(Data.IsValidIndex(0));
+
+	return GetPropertySize(InPropertyHandle->GetProperty(), Data[0]);
+
+	FString StrValue;
+	if (InPropertyHandle->GetValue(StrValue) == FPropertyAccess::Success)
+	{
+		return StrValue.Len();
+	}
+
+	FName NameValue;
+	if (InPropertyHandle->GetValue(NameValue) == FPropertyAccess::Success)
+	{
+		return NameValue.GetStringLength();
+	}
+
+	FText TextValue;
+	if (InPropertyHandle->GetValue(TextValue) == FPropertyAccess::Success)
+	{
+		return TextValue.ToString().Len();
+	}
+
+	return InPropertyHandle->GetProperty()->GetSize();
+}
+#endif
diff --git a/Engine/Plugins/VirtualProduction/RemoteControl/Source/RemoteControlCommon/Private/Tests/RCPropertyContainerTestData.h b/Engine/Plugins/VirtualProduction/RemoteControl/Source/RemoteControlCommon/Private/Tests/RCPropertyContainerTestData.h
--- a/Engine/Plugins/VirtualProduction/RemoteControl/Source/RemoteControlCommon/Private/Tests/RCPropertyContainerTestData.h
+++ b/Engine/Plugins/VirtualProduction/RemoteControl/Source/RemoteControlCommon/Private/Tests/RCPropertyContainerTestData.h
@@ -13,6 +13,9 @@ class UPropertyContainerTestObject : public UObject
 	GENERATED_BODY()
 
 public:
+	UPROPERTY()
+	bool bSomeBool;
+	
 	UPROPERTY()
 	uint32 SomeUInt32;
 
@@ -22,11 +25,18 @@ public:
 	UPROPERTY()
 	FVector SomeVector;
 
-	UPROPERTY(meta = (ClampMin = -5.0f, ClampMax = 99.0f))
-	float ClampedFloat = 56.0f;
+	UPROPERTY()
+	FRotator SomeRotator;
+
+	UPROPERTY(meta = (ClampMin = 20, ClampMax = 145))
+	int32 SomeClampedInt = 5;
+
+	UPROPERTY(meta = (ClampMin = 0.2f, ClampMax = 0.92f))
+	float SomeClampedFloat = 0.25f;
 	
 	UPropertyContainerTestObject()
 		: SomeUInt32(44),
 		SomeFloat(45.0f),
-		SomeVector(FVector(0.2f, 0.3f, 0.6f)) { }
+		SomeVector(FVector(0.2f, 0.3f, 0.6f)),
+		SomeRotator(FRotator::ZeroRotator) { }
 };
diff --git a/Engine/Plugins/VirtualProduction/RemoteControl/Source/RemoteControlCommon/Private/Tests/RCPropertyContainerTests.cpp b/Engine/Plugins/VirtualProduction/RemoteControl/Source/RemoteControlCommon/Private/Tests/RCPropertyContainerTests.cpp
--- a/Engine/Plugins/VirtualProduction/RemoteControl/Source/RemoteControlCommon/Private/Tests/RCPropertyContainerTests.cpp
+++ b/Engine/Plugins/VirtualProduction/RemoteControl/Source/RemoteControlCommon/Private/Tests/RCPropertyContainerTests.cpp
@@ -1,15 +1,16 @@
 ﻿// Copyright Epic Games, Inc. All Rights Reserved.
 
-#include "Misc/AutomationTest.h"
-#include "Editor.h"
+#include "Engine/Engine.h"
 #include "RCPropertyContainer.h"
-#include "RCPropertyContainerTestData.h"
 #include "Camera/CameraComponent.h"
+#include "Misc/AutomationTest.h"
 
+#if WITH_EDITOR
+#include "RCPropertyContainerTestData.h"
 #if WITH_DEV_AUTOMATION_TESTS
 
 BEGIN_DEFINE_SPEC(FPropertyContainerSpec,
-	"PropertyContainer",
+	"Plugin.RemoteControl.PropertyContainer",
 	EAutomationTestFlags::EditorContext | EAutomationTestFlags::EngineFilter | EAutomationTestFlags::ApplicationContextMask)
 END_DEFINE_SPEC(FPropertyContainerSpec)
 
@@ -19,7 +20,7 @@ void FPropertyContainerSpec::Define()
 	{
 		It("Returns_Valid_Subsystem", [this]
 		{
-			URCPropertyContainerRegistry* PropertyContainerRegistry = GEditor->GetEditorSubsystem<URCPropertyContainerRegistry>();
+			URCPropertyContainerRegistry* PropertyContainerRegistry = GEngine->GetEngineSubsystem<URCPropertyContainerRegistry>();
 			TestNotNull("PropertyContainerRegistry", PropertyContainerRegistry);
 		});
 	});
@@ -29,16 +30,16 @@ void FPropertyContainerSpec::Define()
 	{
 		It("Clamp_Min_Max", [this]
 		{
-			FProperty* FloatProperty = UPropertyContainerTestObject::StaticClass()->FindPropertyByName(GET_MEMBER_NAME_CHECKED(UPropertyContainerTestObject, ClampedFloat));
+			FProperty* FloatProperty = UPropertyContainerTestObject::StaticClass()->FindPropertyByName(GET_MEMBER_NAME_CHECKED(UPropertyContainerTestObject, SomeClampedFloat));
 			check(FloatProperty);
 
 			TestTrue("ClampMin", FloatProperty->HasMetaData("ClampMin"));
 			TestTrue("ClampMax", FloatProperty->HasMetaData("ClampMax"));
 
-			TestEqual("ClampMin Value", FloatProperty->GetFloatMetaData("ClampMin"), -5.0f);
-			TestEqual("ClampMax Value", FloatProperty->GetFloatMetaData("ClampMax"), 99.0f);
+			TestEqual("ClampMin Value", FloatProperty->GetFloatMetaData("ClampMin"), 0.2f);
+			TestEqual("ClampMax Value", FloatProperty->GetFloatMetaData("ClampMax"), 0.92f);
 
-			URCPropertyContainerRegistry* PropertyContainerRegistry = GEditor->GetEditorSubsystem<URCPropertyContainerRegistry>();
+			URCPropertyContainerRegistry* PropertyContainerRegistry = GEngine->GetEngineSubsystem<URCPropertyContainerRegistry>();
 			URCPropertyContainerBase* Container = PropertyContainerRegistry->CreateContainer(GetTransientPackage(), NAME_FloatProperty, FloatProperty);
 
 			TestNotNull("Container", Container);
@@ -52,7 +53,7 @@ void FPropertyContainerSpec::Define()
 			FProperty* FloatProperty = UCameraComponent::StaticClass()->FindPropertyByName(GET_MEMBER_NAME_CHECKED(UCameraComponent, PostProcessBlendWeight));
 			check(FloatProperty);
 		
-			URCPropertyContainerRegistry* PropertyContainerRegistry = GEditor->GetEditorSubsystem<URCPropertyContainerRegistry>();
+			URCPropertyContainerRegistry* PropertyContainerRegistry = GEngine->GetEngineSubsystem<URCPropertyContainerRegistry>();
 			URCPropertyContainerBase* Container = PropertyContainerRegistry->CreateContainer(GetTransientPackage(), NAME_FloatProperty, FloatProperty);
 
 			TestNotNull("Container", Container);
@@ -64,7 +65,7 @@ void FPropertyContainerSpec::Define()
 		    FProperty* ValueProperty = UPropertyContainerTestObject::StaticClass()->FindPropertyByName(GET_MEMBER_NAME_CHECKED(UPropertyContainerTestObject, SomeFloat));
 		    check(ValueProperty);
 				
-		    URCPropertyContainerRegistry* PropertyContainerRegistry = GEditor->GetEditorSubsystem<URCPropertyContainerRegistry>();
+		    URCPropertyContainerRegistry* PropertyContainerRegistry = GEngine->GetEngineSubsystem<URCPropertyContainerRegistry>();
 		    URCPropertyContainerBase* Container = PropertyContainerRegistry->CreateContainer(GetTransientPackage(), NAME_FloatProperty, ValueProperty);
 
             void* ValuePtr = ValueProperty->ContainerPtrToValuePtr<void>(ObjectInstance);
@@ -78,7 +79,7 @@ void FPropertyContainerSpec::Define()
             FProperty* ValueProperty = UPropertyContainerTestObject::StaticClass()->FindPropertyByName(GET_MEMBER_NAME_CHECKED(UPropertyContainerTestObject, SomeFloat));
             check(ValueProperty);
 				
-            URCPropertyContainerRegistry* PropertyContainerRegistry = GEditor->GetEditorSubsystem<URCPropertyContainerRegistry>();
+            URCPropertyContainerRegistry* PropertyContainerRegistry = GEngine->GetEngineSubsystem<URCPropertyContainerRegistry>();
             URCPropertyContainerBase* Container = PropertyContainerRegistry->CreateContainer(GetTransientPackage(), NAME_FloatProperty, ValueProperty);
 
 			float InputValue = 84.3f;
@@ -103,7 +104,7 @@ void FPropertyContainerSpec::Define()
 			FProperty* ValueProperty = UPropertyContainerTestObject::StaticClass()->FindPropertyByName(GET_MEMBER_NAME_CHECKED(UPropertyContainerTestObject, SomeVector));
 			check(ValueProperty);
 
-			URCPropertyContainerRegistry* PropertyContainerRegistry = GEditor->GetEditorSubsystem<URCPropertyContainerRegistry>();
+			URCPropertyContainerRegistry* PropertyContainerRegistry = GEngine->GetEngineSubsystem<URCPropertyContainerRegistry>();
 			URCPropertyContainerBase* Container = PropertyContainerRegistry->CreateContainer(GetTransientPackage(), NAME_VectorProperty, ValueProperty);
 
 			TestNotNull("Container", Container);
@@ -115,7 +116,7 @@ void FPropertyContainerSpec::Define()
 		    FProperty* ValueProperty = UPropertyContainerTestObject::StaticClass()->FindPropertyByName(GET_MEMBER_NAME_CHECKED(UPropertyContainerTestObject, SomeVector));
 		    check(ValueProperty);
 				
-		    URCPropertyContainerRegistry* PropertyContainerRegistry = GEditor->GetEditorSubsystem<URCPropertyContainerRegistry>();
+		    URCPropertyContainerRegistry* PropertyContainerRegistry = GEngine->GetEngineSubsystem<URCPropertyContainerRegistry>();
 		    URCPropertyContainerBase* Container = PropertyContainerRegistry->CreateContainer(GetTransientPackage(), NAME_VectorProperty, ValueProperty);
 
             void* ValuePtr = ValueProperty->ContainerPtrToValuePtr<void>(ObjectInstance);
@@ -129,7 +130,7 @@ void FPropertyContainerSpec::Define()
             FProperty* ValueProperty = UPropertyContainerTestObject::StaticClass()->FindPropertyByName(GET_MEMBER_NAME_CHECKED(UPropertyContainerTestObject, SomeVector));
             check(ValueProperty);
 				
-            URCPropertyContainerRegistry* PropertyContainerRegistry = GEditor->GetEditorSubsystem<URCPropertyContainerRegistry>();
+            URCPropertyContainerRegistry* PropertyContainerRegistry = GEngine->GetEngineSubsystem<URCPropertyContainerRegistry>();
             URCPropertyContainerBase* Container = PropertyContainerRegistry->CreateContainer(GetTransientPackage(), NAME_VectorProperty, ValueProperty);
 
 			FVector InputValue(0.45f, 0.65f, -1.24f);
@@ -145,4 +146,6 @@ void FPropertyContainerSpec::Define()
         });
 	});
 }
+
+#endif
 #endif
diff --git a/Engine/Plugins/VirtualProduction/RemoteControl/Source/RemoteControlCommon/Public/RCPropertyContainer.h b/Engine/Plugins/VirtualProduction/RemoteControl/Source/RemoteControlCommon/Public/RCPropertyContainer.h
--- a/Engine/Plugins/VirtualProduction/RemoteControl/Source/RemoteControlCommon/Public/RCPropertyContainer.h
+++ b/Engine/Plugins/VirtualProduction/RemoteControl/Source/RemoteControlCommon/Public/RCPropertyContainer.h
@@ -3,23 +3,31 @@
 #pragma once
 
 #include "CoreMinimal.h"
-#include "EditorSubsystem.h"
+
+#include "Subsystems/EngineSubsystem.h"
 #include "Templates/SubclassOf.h"
 #include "UObject/Object.h"
+#include "UObject/WeakFieldPtr.h"
 
 #include "RCPropertyContainer.generated.h"
 
 UCLASS(Transient, Abstract)
-class URCPropertyContainerBase : public UObject
+class REMOTECONTROLCOMMON_API URCPropertyContainerBase : public UObject
 {
 	GENERATED_BODY()
 
 public:
-	/** Sets the value from the incoming raw data */
-	void SetValue(const uint8* InData);
+	/** Sets the value from the incoming raw data. Provide size for array, string, etc. */
+	void SetValue(const uint8* InData, const SIZE_T& InSize = 0);
+
+	/** Writes to the provided raw data pointer. Returns size for array, string, etc. */
+	SIZE_T GetValue(uint8* OutData);
 
-	/** Writes to the provided raw data pointer */
-	void GetValue(uint8* OutData);
+	/** Returns the Property for Value */
+	virtual FProperty* GetValueProperty();
+
+private:
+	TWeakFieldPtr<FProperty> ValueProperty;
 };
 
 /** Minimal information needed to lookup a unique property container class */
@@ -40,7 +48,7 @@ inline bool operator!=(const FRCPropertyContainerKey& Lhs, const FRCPropertyCont
 
 /** A subsystem to provide and cache dynamically created PropertyContainer classes. */
 UCLASS()
-class REMOTECONTROLCOMMON_API URCPropertyContainerRegistry final : public UEditorSubsystem
+class REMOTECONTROLCOMMON_API URCPropertyContainerRegistry : public UEngineSubsystem
 {
 	GENERATED_BODY()
 	
diff --git a/Engine/Plugins/VirtualProduction/RemoteControl/Source/RemoteControlCommon/Public/RCPropertyUtilities.h b/Engine/Plugins/VirtualProduction/RemoteControl/Source/RemoteControlCommon/Public/RCPropertyUtilities.h
new file mode 100644
--- /dev/null
+++ b/Engine/Plugins/VirtualProduction/RemoteControl/Source/RemoteControlCommon/Public/RCPropertyUtilities.h
@@ -0,0 +1,216 @@
+﻿// Copyright Epic Games, Inc. All Rights Reserved.
+
+#pragma once
+
+#include "CoreMinimal.h"
+
+#include "RCTypeTraits.h"
+
+#if WITH_EDITOR
+#include "PropertyHandle.h"
+#endif
+
+namespace RemoteControlPropertyUtilities
+{
+#if WITH_EDITOR
+	/** Container that can hold either a PropertyHandle, or Property/Data pair. Similar to FFieldVariant */
+	class FRCPropertyVariant
+	{
+		TSharedPtr<IPropertyHandle> PropertyHandle = nullptr;
+
+		const FProperty* Property = nullptr;
+		void* PropertyData = nullptr;
+		TArray<uint8>* PropertyContainer = nullptr;
+
+		bool bIsHandle = false;
+		int32 NumElements = -1;
+
+	public:
+		explicit FRCPropertyVariant() = default;
+
+		virtual ~FRCPropertyVariant()
+		{
+			auto o =123;
+			auto n = o;
+		}
+
+		FRCPropertyVariant(const TSharedPtr<IPropertyHandle>& InPropertyHandle)
+			: bIsHandle(true)
+		{
+			PropertyHandle = InPropertyHandle;
+		}
+
+		FRCPropertyVariant(const FProperty* InProperty, const void* InPropertyData, const int32& InNumElements = -1)
+			: bIsHandle(false)
+			, NumElements(InNumElements)
+		{
+			Property = InProperty;
+			// @todo: make more elegant
+			PropertyData = const_cast<void*>(InPropertyData);
+		}
+
+		FRCPropertyVariant(const FProperty* InProperty, TArray<uint8>& InPropertyData)
+			: bIsHandle(false)
+		{
+			Property = InProperty;
+
+			PropertyData = InPropertyData.GetData();
+			PropertyContainer = &InPropertyData;
+			NumElements = InPropertyData.Num();
+		}
+		
+		/** Gets the property. */
+		const FProperty* GetProperty() const
+		{
+			return bIsHandle ? PropertyHandle->GetProperty() : Property;
+		}
+
+		/** Gets the typed property, returns nullptr if not cast. */
+		template <typename PropertyType>
+		const PropertyType* GetProperty() const
+		{
+			static_assert(TIsDerivedFrom<PropertyType, FProperty>::Value, "PropertyType must derive from FProperty");
+
+			return CastField<PropertyType>(GetProperty());
+		}
+
+		/** Gets the data pointer */
+		void* GetPropertyData() const
+		{
+			if(bIsHandle)
+			{
+				TArray<void*> Data;
+				PropertyHandle->AccessRawData(Data);
+				check(Data.IsValidIndex(0)); // check that there's at least one set of data
+
+				return Data[0];
+			}			
+
+			return PropertyData;
+		}
+
+		bool IsHandle() const { return bIsHandle; }
+
+		void Init(const int32 InSize)
+		{
+			if(!bIsHandle && PropertyContainer)
+			{
+				PropertyContainer->SetNumUninitialized(InSize);
+				PropertyData = PropertyContainer->GetData();
+				NumElements = InSize;
+				return;
+			}
+		}
+
+		/** Gets the size of the allocated data. Not always available. */
+		int32 Num() const { return NumElements; }
+
+		/** Comparison is by property, not by property instance value. */
+		bool operator==(const FRCPropertyVariant& InOther) const
+		{
+			return GetProperty() == InOther.GetProperty();
+		}
+		
+		bool operator!=(const FRCPropertyVariant& InOther) const
+		{
+			return GetProperty() != InOther.GetProperty();
+		}
+	};
+
+	/** Gets from raw data to input property handle. */
+
+	// Non 'String-like' property types
+	template <typename PropertyType>
+	typename TEnableIf<TNot<RemoteControlTypeTraits::TIsStringLikeProperty<PropertyType>>::Value, bool>::Type
+	FromBytes(const FRCPropertyVariant& InSrc, FRCPropertyVariant& InDstProperty)
+	{
+		InDstProperty.GetProperty()->CopyCompleteValue((uint8*)InDstProperty.GetPropertyData(), (uint8*)InSrc.GetPropertyData());
+
+		return true;
+	}
+
+	template <>
+	REMOTECONTROLCOMMON_API bool FromBytes<FStructProperty>(const FRCPropertyVariant& InSrc, FRCPropertyVariant& InDstProperty);
+
+	template <>
+	REMOTECONTROLCOMMON_API bool FromBytes<FArrayProperty>(const FRCPropertyVariant& InSrc, FRCPropertyVariant& InDstProperty);
+
+	template <>
+	REMOTECONTROLCOMMON_API bool FromBytes<FSetProperty>(const FRCPropertyVariant& InSrc, FRCPropertyVariant& InDstProperty);
+
+	// Specialization for 'String-like' property types
+	template <typename PropertyType>
+	typename TEnableIf<RemoteControlTypeTraits::TIsStringLikeProperty<PropertyType>::Value, bool>::Type
+	FromBytes(const FRCPropertyVariant& InSrc, FRCPropertyVariant& OutDst)
+	{
+		const int32 StrLen = InSrc.Num();
+		FString Str = BytesToString((uint8*)InSrc.GetPropertyData(), StrLen);
+
+		if(!OutDst.IsHandle())
+		{
+			OutDst.Init(StrLen);
+		}
+
+		OutDst.GetProperty()->InitializeValue((uint8*)OutDst.GetPropertyData());
+		OutDst.GetProperty()->ImportText(*Str, (uint8*)OutDst.GetPropertyData(), 0, nullptr);
+
+#if !UE_BUILD_SHIPPING && UE_BUILD_DEBUG
+		// const typename PropertyType::TCppType* Value = reinterpret_cast<const typename PropertyType::TCppType*>((uint8*)OutDst.GetPropertyData());
+		// auto o = Value;
+#endif
+
+		return true;
+	}
+
+	/** Specialization for FProperty casts and forwards to specializations. */
+	template <> 
+	REMOTECONTROLCOMMON_API bool FromBytes<FProperty>(const FRCPropertyVariant& InSrc, FRCPropertyVariant& InDstProperty);
+
+	/** Sets raw data from input property handle. */
+
+	// Non 'String-like' property types
+	template <typename PropertyType>
+	typename TEnableIf<TNot<RemoteControlTypeTraits::TIsStringLikeProperty<PropertyType>>::Value, bool>::Type
+	ToBytes(const FRCPropertyVariant& InSrc, FRCPropertyVariant& OutDst)
+	{
+		InSrc.GetProperty()->CopyCompleteValue((uint8*)OutDst.GetPropertyData(), (uint8*)InSrc.GetPropertyData());
+
+		return true;
+	}
+
+	template <>
+	REMOTECONTROLCOMMON_API bool ToBytes<FStructProperty>(const FRCPropertyVariant& InSrc, FRCPropertyVariant& OutDst);
+
+	template <>
+	REMOTECONTROLCOMMON_API bool ToBytes<FArrayProperty>(const FRCPropertyVariant& InSrc, FRCPropertyVariant& OutDst);
+
+	template <>
+	REMOTECONTROLCOMMON_API bool ToBytes<FSetProperty>(const FRCPropertyVariant& InSrc, FRCPropertyVariant& OutDst);
+
+	template <>
+	REMOTECONTROLCOMMON_API bool ToBytes<FMapProperty>(const FRCPropertyVariant& InSrc, FRCPropertyVariant& OutDst);
+
+	// Specialization for 'String-like' property types
+	template <typename PropertyType>
+	typename TEnableIf<RemoteControlTypeTraits::TIsStringLikeProperty<PropertyType>::Value, bool>::Type
+	ToBytes(const FRCPropertyVariant& InSrc, FRCPropertyVariant& OutDst)
+	{
+		FString Str;
+		InSrc.GetProperty()->ExportTextItem(Str, (uint8*)InSrc.GetPropertyData(), nullptr, nullptr, PPF_None, nullptr);
+
+		TArray<uint8> Buffer;
+		Buffer.SetNumUninitialized(Str.Len());
+		StringToBytes(Str, Buffer.GetData(), Buffer.Num());
+
+		OutDst.Init(Str.Len());
+		FMemory::Memcpy((uint8*)OutDst.GetPropertyData(), Buffer.GetData(), Buffer.Num());
+
+		return true;
+	}
+
+	/** Specialization for FProperty casts and forwards to specializations. */
+	template <>
+	REMOTECONTROLCOMMON_API bool ToBytes<FProperty>(const FRCPropertyVariant& InSrc, FRCPropertyVariant& OutDst);
+
+#endif
+}
diff --git a/Engine/Plugins/VirtualProduction/RemoteControl/Source/RemoteControlCommon/Public/RCTypeTraits.h b/Engine/Plugins/VirtualProduction/RemoteControl/Source/RemoteControlCommon/Public/RCTypeTraits.h
--- a/Engine/Plugins/VirtualProduction/RemoteControl/Source/RemoteControlCommon/Public/RCTypeTraits.h
+++ b/Engine/Plugins/VirtualProduction/RemoteControl/Source/RemoteControlCommon/Public/RCTypeTraits.h
@@ -2,213 +2,575 @@
 
 #pragma once
 
-#include "CoreTypes.h"
-#include "Templates/IsArithmetic.h"
-#include "Templates/UnrealTypeTraits.h"
+#include "CoreMinimal.h"
+
 #include "Math/NumericLimits.h"
 #include "Math/UnrealMathUtility.h"
+#include "Templates/IsArithmetic.h"
+#include "Templates/UnrealTypeTraits.h"
+#include "UObject/UnrealType.h"
 
-/** Concept to check if T has NumericLimits */
-struct CHasNumericLimits {
-	template <typename T>
-	auto Requires() -> decltype(TNumericLimits<T>::NumericType);
-};
+namespace RemoteControlTypeTraits
+{
+	namespace Concepts
+	{
+		/** Concept to check if T has NumericLimits */
+		struct CNumerical
+		{
+			template <typename T>
+			auto Requires() -> decltype(
+				TAnd<
+					TIsSame<typename TNumericLimits<T>::NumericType, T>,
+					TNot<TIsSame<T, bool>>>::Value);
+		};
+
+		/** Concept to check if PropertyType::TCppType matches ValueType */
+		struct CPropertyValuePair
+		{
+			template <typename PropertyType, typename ValueType>
+			auto Requires() -> decltype(TIsDerivedFrom<PropertyType, TProperty<ValueType, PropertyType>>::Value);
+		};	
+	}
+
+	/** Catch-all for string-like property types */
+	template <typename PropertyType, typename Enable = void>
+	struct TIsStringLikeProperty
+	{
+		enum { Value = false };
+	};
+
+	template <typename PropertyType>
+	struct TIsStringLikeProperty<
+			PropertyType,
+			typename TEnableIf<
+				TOr<
+					TIsSame<PropertyType, FStrProperty>,
+					TIsSame<PropertyType, FNameProperty>,
+					TIsSame<PropertyType, FTextProperty>>::Value>::Type>
+	{
+		enum
+		{
+			Value = true
+		};
+	};
+
+	/** Ensures ValueType is a numeric type. */
+	template <typename ValueType>
+	using TNumericValueConstraint = TAnd<TModels<RemoteControlTypeTraits::Concepts::CNumerical, ValueType>>;
+}
 
 /** Various RemoteControl type traits */
-template <typename ValueType>
+template <typename ValueType, typename Enable = void>
 struct TRemoteControlTypeTraits;
 
-/** RemoteControlTypeTraits for const types */
-template <typename ValueType>
-struct TRemoteControlTypeTraits<const ValueType> 
-    : public TRemoteControlTypeTraits<ValueType>
-{ };
+/** Various RemoteControl property type traits */
+template <typename PropertyType, typename Enable = void>
+struct TRemoteControlPropertyTypeTraits;
+
+#pragma region Numeric Types
 
-/** RemoteControlTypeTraits for volatile types */
+/** RemoteControlTypeTraits for integers */
 template <typename ValueType>
-struct TRemoteControlTypeTraits<volatile ValueType> 
-    : public TRemoteControlTypeTraits<ValueType>
-{ };
+struct TRemoteControlTypeTraits<ValueType,
+	typename TEnableIf<
+		TAnd<
+			TIsArithmetic<ValueType>,
+			TNot<TIsFloatingPoint<ValueType>>,
+			TModels<RemoteControlTypeTraits::Concepts::CNumerical, ValueType>>::Value, void>::Type>
+{
+	using Type = ValueType;
+
+	/** Is ValueType supported as a range (protocol input) value? */
+	static constexpr bool IsSupportedRangeType() { return true; }
+
+	/** Is ValueType supported as a mapping (property output) value? */
+	static constexpr bool IsSupportedMappingType() { return true; }
+
+	/** The default minimum value for newly created range (protocol, input). */
+	static constexpr Type DefaultRangeValueMin() { return TNumericLimits<Type>::Min(); }
+
+	/** The default maximum value for newly created range. */
+	static constexpr Type DefaultRangeValueMax() { return TNumericLimits<Type>::Max(); }
 
-/** RemoteControlTypeTraits for const volatile types */
+	/** The default minimum value for newly created mapping. */
+	static constexpr Type DefaultMappingValueMin() { return static_cast<Type>(0); }
+
+	/** The default maximum value for newly created mapping. */
+	static constexpr Type DefaultMappingValueMax() { return static_cast<Type>(1); }
+};
+
+/** RemoteControlTypeTraits for floats */
 template <typename ValueType>
-struct TRemoteControlTypeTraits<const volatile ValueType> 
-    : public TRemoteControlTypeTraits<ValueType>
-{ };
+struct TRemoteControlTypeTraits<ValueType,
+	typename TEnableIf<TIsFloatingPoint<ValueType>::Value, void>::Type>
+{
+	using Type = ValueType;
+
+	/** Is ValueType supported as a range (protocol input) value? */
+	static constexpr bool IsSupportedRangeType() { return true; }
+
+	/** Is ValueType supported as a mapping (property output) value? */
+	static constexpr bool IsSupportedMappingType() { return true; }
+
+	static constexpr Type DefaultRangeValueMin() { return static_cast<Type>(0.0f); }
 
-/** RemoteControlTypeTraits for numeric types */
-template <typename ValueType, typename TEnableIf<TModels<CHasNumericLimits, ValueType>::Value>::Type* = nullptr>
-struct TRemoteControlTypeTraits<ValueType>
+	static constexpr Type DefaultRangeValueMax() { return static_cast<Type>(1.0f); }
+
+	static constexpr Type DefaultMappingValueMin() { return static_cast<Type>(0.0f); }
+
+	static constexpr Type DefaultMappingValueMax() { return static_cast<Type>(1.0f); }
+};
+
+/**
+ * RemoteControlPropertyTypeTraits for FNumericProperty
+ * Currently all numeric types are supported so we can shortcut the above (rather than doing a series of CastField's)
+ */
+template <>
+struct TRemoteControlPropertyTypeTraits<FNumericProperty>
+{
+	using Type = FNumericProperty;
+
+	/** Is ValueType supported as a range (protocol input) value? */
+	static constexpr bool IsSupportedRangeType() { return true; }
+
+	/** Is ValueType supported as a mapping (property output) value? */
+	static constexpr bool IsSupportedMappingType() { return true; }
+};
+
+/** RemoteControlPropertyTypeTraits for FEnumProperty */
+template <>
+struct TRemoteControlPropertyTypeTraits<FEnumProperty>
+{
+	using PropertyType = FEnumProperty;
+	using ValueType = uint8;
+
+	/** Is ValueType supported as a range (protocol input) value? */
+	static constexpr bool IsSupportedRangeType() { return false; }
+
+	/** Is ValueType supported as a mapping (property output) value? */
+	static constexpr bool IsSupportedMappingType() { return true; }
+};
+
+#pragma endregion Numeric Types
+
+/** RemoteControlTypeTraits for bool */
+template <>
+struct TRemoteControlTypeTraits<bool>
 {
-	typedef ValueType Type;
+	using Type = bool;
+
+	/** Is ValueType supported as a range (protocol input) value? */
+	static constexpr bool IsSupportedRangeType() { return false; }
+
+	/** Is ValueType supported as a mapping (property output) value? */
+	static constexpr bool IsSupportedMappingType() { return true; }
+
+	static constexpr Type DefaultRangeValueMin() { return false; }
+
+	static constexpr Type DefaultRangeValueMax() { return true; }
 	
-	static constexpr Type DefaultMin()
-	{
-		return TNumericLimits<Type>::Min();
-	}
+	static constexpr Type DefaultMappingValueMin() { return false; }
+	
+	static constexpr Type DefaultMappingValueMax() { return true; }
+};
 
-	static constexpr Type DefaultMax()
-	{
-		return TNumericLimits<Type>::Max();
-	}
+/** RemoteControlPropertyTypeTraits for FBoolProperty */
+template <>
+struct TRemoteControlPropertyTypeTraits<FBoolProperty>
+{
+	using PropertyType = FBoolProperty;
+	using ValueType = bool;
+
+	/** Is ValueType supported as a range (protocol input) value? */
+	static constexpr bool IsSupportedRangeType() { return TRemoteControlTypeTraits<bool>::IsSupportedRangeType(); }
+
+	/** Is ValueType supported as a mapping (property output) value? */
+	static constexpr bool IsSupportedMappingType() { return TRemoteControlTypeTraits<bool>::IsSupportedMappingType(); }
+
+	/** The default minimum value for newly created range (protocol, input). */
+	static constexpr ValueType DefaultRangeValueMin() { return TRemoteControlTypeTraits<bool>::DefaultRangeValueMin(); }
+
+	/** The default maximum value for newly created range. */
+	static constexpr ValueType DefaultRangeValueMax() { return TRemoteControlTypeTraits<bool>::DefaultRangeValueMax(); }
+
+	/** The default minimum value for newly created mapping. */
+	static constexpr ValueType DefaultMappingValueMin() { return TRemoteControlTypeTraits<bool>::DefaultMappingValueMin(); }
+
+	/** The default maximum value for newly created mapping. */
+	static constexpr ValueType DefaultMappingValueMax() { return TRemoteControlTypeTraits<bool>::DefaultMappingValueMax(); }
+};
+
+#pragma region String Types
+
+/** RemoteControlTypeTraits for FString */
+template <>
+struct TRemoteControlTypeTraits<FString>
+{
+	using Type = FString;
+	
+	/** Is ValueType supported as a range (protocol input) value? */
+	static constexpr bool IsSupportedRangeType() { return false; }
+
+	/** Is ValueType supported as a mapping (property output) value? */
+	static constexpr bool IsSupportedMappingType() { return true; }
+
+	static Type DefaultRangeValueMin() { return {}; }
+
+	static Type DefaultRangeValueMax() { return {}; }
+	
+	static Type DefaultMappingValueMin() { return {}; }
+	
+	static Type DefaultMappingValueMax() { return {}; }
+};
+
+/** RemoteControlPropertyTypeTraits for FStrProperty */
+template <>
+struct TRemoteControlPropertyTypeTraits<FStrProperty>
+{
+	using PropertyType = FStrProperty;
+	using ValueType = FString;
+
+	/** Is ValueType supported as a range (protocol input) value? */
+	static constexpr bool IsSupportedRangeType() { return TRemoteControlTypeTraits<FString>::IsSupportedRangeType(); }
+
+	/** Is ValueType supported as a mapping (property output) value? */
+	static constexpr bool IsSupportedMappingType() { return TRemoteControlTypeTraits<FString>::IsSupportedMappingType(); }
+
+	/** The default minimum value for newly created range (protocol, input). */
+	static ValueType DefaultRangeValueMin() { return TRemoteControlTypeTraits<FString>::DefaultRangeValueMin(); }
+
+	/** The default maximum value for newly created range. */
+	static ValueType DefaultRangeValueMax() { return TRemoteControlTypeTraits<FString>::DefaultRangeValueMax(); }
+
+	/** The default minimum value for newly created mapping. */
+	static ValueType DefaultMappingValueMin() { return TRemoteControlTypeTraits<FString>::DefaultMappingValueMin(); }
+
+	/** The default maximum value for newly created mapping. */
+	static ValueType DefaultMappingValueMax() { return TRemoteControlTypeTraits<FString>::DefaultMappingValueMax(); }
+};
+
+/** RemoteControlTypeTraits for FName */
+template <>
+struct TRemoteControlTypeTraits<FName>
+{
+	using Type = FName;
+
+	/** Is ValueType supported as a range (protocol input) value? */
+	static constexpr bool IsSupportedRangeType() { return false; }
+
+	/** Is ValueType supported as a mapping (property output) value? */
+	static constexpr bool IsSupportedMappingType() { return true; }
+
+	static Type DefaultRangeValueMin() { return {}; }
+
+	static Type DefaultRangeValueMax() { return {}; }
+	
+	static Type DefaultMappingValueMin() { return {}; }
+	
+	static Type DefaultMappingValueMax() { return {}; }
+};
+
+/** RemoteControlPropertyTypeTraits for FNameProperty */
+template <>
+struct TRemoteControlPropertyTypeTraits<FNameProperty>
+{
+	using PropertyType = FNameProperty;
+	using ValueType = FName;
+
+	/** Is ValueType supported as a range (protocol input) value? */
+	static constexpr bool IsSupportedRangeType() { return TRemoteControlTypeTraits<ValueType>::IsSupportedRangeType(); }
+
+	/** Is ValueType supported as a mapping (property output) value? */
+	static constexpr bool IsSupportedMappingType() { return TRemoteControlTypeTraits<ValueType>::IsSupportedMappingType(); }
+
+	/** The default minimum value for newly created range (protocol, input). */
+	static ValueType DefaultRangeValueMin() { return TRemoteControlTypeTraits<ValueType>::DefaultRangeValueMin(); }
+
+	/** The default maximum value for newly created range. */
+	static ValueType DefaultRangeValueMax() { return TRemoteControlTypeTraits<ValueType>::DefaultRangeValueMax(); }
+
+	/** The default minimum value for newly created mapping. */
+	static ValueType DefaultMappingValueMin() { return TRemoteControlTypeTraits<ValueType>::DefaultMappingValueMin(); }
+
+	/** The default maximum value for newly created mapping. */
+	static ValueType DefaultMappingValueMax() { return TRemoteControlTypeTraits<ValueType>::DefaultMappingValueMax(); }
+};
+
+/** RemoteControlTypeTraits for FText */
+template <>
+struct TRemoteControlTypeTraits<FText>
+{
+	using Type = FText;
+
+	/** Is ValueType supported as a range (protocol input) value? */
+	static constexpr bool IsSupportedRangeType() { return false; }
+
+	/** Is ValueType supported as a mapping (property output) value? */
+	static constexpr bool IsSupportedMappingType() { return true; }
+
+	static Type DefaultRangeValueMin() { return {}; }
+
+	static Type DefaultRangeValueMax() { return {}; }
+	
+	static Type DefaultMappingValueMin() { return {}; }
+	
+	static Type DefaultMappingValueMax() { return {}; }
+};
+
+/** RemoteControlPropertyTypeTraits for FTextProperty */
+template <>
+struct TRemoteControlPropertyTypeTraits<FTextProperty>
+{
+	using PropertyType = FTextProperty;
+	using ValueType = FText;
+
+	/** Is ValueType supported as a range (protocol input) value? */
+	static constexpr bool IsSupportedRangeType() { return TRemoteControlTypeTraits<ValueType>::IsSupportedRangeType(); }
+
+	/** Is ValueType supported as a mapping (property output) value? */
+	static constexpr bool IsSupportedMappingType() { return TRemoteControlTypeTraits<ValueType>::IsSupportedMappingType(); }
+
+	/** The default minimum value for newly created range (protocol, input). */
+	static ValueType DefaultRangeValueMin() { return TRemoteControlTypeTraits<ValueType>::DefaultRangeValueMin(); }
+
+	/** The default maximum value for newly created range. */
+	static ValueType DefaultRangeValueMax() { return TRemoteControlTypeTraits<ValueType>::DefaultRangeValueMax(); }
+
+	/** The default minimum value for newly created mapping. */
+	static ValueType DefaultMappingValueMin() { return TRemoteControlTypeTraits<ValueType>::DefaultMappingValueMin(); }
+
+	/** The default maximum value for newly created mapping. */
+	static ValueType DefaultMappingValueMax() { return TRemoteControlTypeTraits<ValueType>::DefaultMappingValueMax(); }
+};
+
+#pragma endregion String Types
+
+/** RemoteControlPropertyTypeTraits for FArrayProperty */
+template <>
+struct TRemoteControlPropertyTypeTraits<FArrayProperty>
+{
+	using PropertyType = FArrayProperty;
+
+	/** Is ValueType supported as a range (protocol input) value? */
+	static constexpr bool IsSupportedRangeType() { return false; }
+
+	/** Is ValueType supported as a mapping (property output) value? */
+	static constexpr bool IsSupportedMappingType() { return true; }
+};
+
+/** RemoteControlPropertyTypeTraits for FSetProperty */
+template <>
+struct TRemoteControlPropertyTypeTraits<FSetProperty>
+{
+	using PropertyType = FSetProperty;
+
+	/** Is ValueType supported as a range (protocol input) value? */
+	static constexpr bool IsSupportedRangeType() { return false; }
+
+	/** Is ValueType supported as a mapping (property output) value? */
+	static constexpr bool IsSupportedMappingType() { return true; }
+};
+
+/** RemoteControlPropertyTypeTraits for FMapProperty */
+template <>
+struct TRemoteControlPropertyTypeTraits<FMapProperty>
+{
+	using PropertyType = FMapProperty;
+
+	/** Is ValueType supported as a range (protocol input) value? */
+	static constexpr bool IsSupportedRangeType() { return false; }
+
+	/** Is ValueType supported as a mapping (property output) value? */
+	static constexpr bool IsSupportedMappingType() { return true; }
+};
+
+#pragma region Structs (Built-in)
+
+/** RemoteControlPropertyTypeTraits for FStructProperty */
+template <>
+struct TRemoteControlPropertyTypeTraits<FStructProperty>
+{
+	using PropertyType = FNameProperty;
+
+	/** Is ValueType supported as a range (protocol input) value? */
+	static constexpr bool IsSupportedRangeType() { return false; }
+
+	/** Is ValueType supported as a mapping (property output) value? */
+	static constexpr bool IsSupportedMappingType() { return true; }
 };
 
 /** RemoteControlTypeTraits for FVector */
 template <>
 struct TRemoteControlTypeTraits<FVector>
 {
-	typedef FVector Type;
+	using Type = FVector;
 	
-	static constexpr Type DefaultMin()
-	{
-		return Type::ZeroVector;
-	}
+	/** Is ValueType supported as a range (protocol input) value? */
+	static constexpr bool IsSupportedRangeType() { return false; }
 
-	static constexpr Type DefaultMax()
-	{
-		return Type::OneVector;
-	}
+	/** Is ValueType supported as a mapping (property output) value? */
+	static constexpr bool IsSupportedMappingType() { return true; }
+	
+	static Type DefaultMappingValueMin() { return Type::ZeroVector; }
+
+	static Type DefaultMappingValueMax() { return Type::OneVector; }
 };
 
 /** RemoteControlTypeTraits for FVector2D */
 template <>
 struct TRemoteControlTypeTraits<FVector2D>
 {
-	typedef FVector2D Type;
+	using Type = FVector2D;
 	
-	static constexpr Type DefaultMin()
-	{
-		return Type::ZeroVector;
-	}
+	/** Is ValueType supported as a range (protocol input) value? */
+	static constexpr bool IsSupportedRangeType() { return false; }
 
-	static constexpr Type DefaultMax()
-	{
-		return Type::UnitVector;
-	}
+	/** Is ValueType supported as a mapping (property output) value? */
+	static constexpr bool IsSupportedMappingType() { return true; }
+
+	static Type DefaultMappingValueMin() { return Type::ZeroVector; }
+	
+	static Type DefaultMappingValueMax() { return Type::UnitVector; }
 };
 
 /** RemoteControlTypeTraits for FVector4 */
 template <>
 struct TRemoteControlTypeTraits<FVector4>
 {
-	typedef FVector4 Type;
+	using Type = FVector4;
 	
-	static constexpr Type DefaultMin()
-	{
-		return Type(EForceInit::ForceInitToZero);
-	}
+	/** Is ValueType supported as a range (protocol input) value? */
+	static constexpr bool IsSupportedRangeType() { return false; }
 
-	static constexpr Type DefaultMax()
-	{
-		return Type(1.0f, 1.0f, 1.0f, 1.0f);
-	}
+	/** Is ValueType supported as a mapping (property output) value? */
+	static constexpr bool IsSupportedMappingType() { return true; }
+	
+	static Type DefaultMappingValueMin() { return Type(EForceInit::ForceInitToZero); }
+
+	static Type DefaultMappingValueMax() { return Type(1.0f, 1.0f, 1.0f, 1.0f); }
 };
 
 /** RemoteControlTypeTraits for FRotator */
 template <>
 struct TRemoteControlTypeTraits<FRotator>
 {
-	typedef FRotator Type;
+	using Type = FRotator;
+
+	/** Is ValueType supported as a range (protocol input) value? */
+	static constexpr bool IsSupportedRangeType() { return false; }
+
+	/** Is ValueType supported as a mapping (property output) value? */
+	static constexpr bool IsSupportedMappingType() { return true; }
 	
-	static constexpr Type DefaultMin()
-	{
-		return FRotator::ZeroRotator;
-	}
+	static Type DefaultMappingValueMin() { return Type::ZeroRotator; }
 
-	static constexpr Type DefaultMax()
-	{
-		return FRotator(90.0f, 90.0f, 90.0f);
-	}
+	static Type DefaultMappingValueMax() { return Type(90.0f, 90.0f, 90.0f); }
 };
 
 /** RemoteControlTypeTraits for FQuat */
 template <>
 struct TRemoteControlTypeTraits<FQuat>
 {
-	typedef FQuat Type;
+	using Type = FQuat;
 	
-	static constexpr Type DefaultMin()
-	{
-		return FQuat(0.0f, 0.0f, 0.0f, 0.0f);
-	}
+	/** Is ValueType supported as a range (protocol input) value? */
+	static constexpr bool IsSupportedRangeType() { return false; }
 
-	static constexpr Type DefaultMax()
-	{
-		return Type::Identity;
-	}
+	/** Is ValueType supported as a mapping (property output) value? */
+	static constexpr bool IsSupportedMappingType() { return true; }
+	
+	static Type DefaultMappingValueMin() { return Type(0.0f, 0.0f, 0.0f, 0.0f); }
+
+	static Type DefaultMappingValueMax() { return Type::Identity; }
 };
 
 /** RemoteControlTypeTraits for FTransform */
 template <>
 struct TRemoteControlTypeTraits<FTransform>
 {
-	typedef FTransform Type;
+	using Type = FTransform;
+	
+	/** Is ValueType supported as a range (protocol input) value? */
+	static constexpr bool IsSupportedRangeType() { return false; }
+
+	/** Is ValueType supported as a mapping (property output) value? */
+	static constexpr bool IsSupportedMappingType() { return true; }
 	
-	static constexpr Type DefaultMin()
+	static Type DefaultMappingValueMin()
 	{
-		return FTransform(
-			TRemoteControlTypeTraits<FRotator>::DefaultMin(),
-			TRemoteControlTypeTraits<FVector>::DefaultMin(),
-			TRemoteControlTypeTraits<FVector>::DefaultMax()); // scale is Max cause it shouldn't be zero		
+		return Type(
+			TRemoteControlTypeTraits<FRotator>::DefaultMappingValueMin(),
+			TRemoteControlTypeTraits<FVector>::DefaultMappingValueMin(),
+			TRemoteControlTypeTraits<FVector>::DefaultMappingValueMax()); // scale is Max cause it shouldn't be zero
 	}
 
-	static constexpr Type DefaultMax()
+	static Type DefaultMappingValueMax()
 	{
-		return FTransform(
-	TRemoteControlTypeTraits<FRotator>::DefaultMax(),
-	TRemoteControlTypeTraits<FVector>::DefaultMax(),
-	TRemoteControlTypeTraits<FVector>::DefaultMax()); // scale is Max cause it shouldn't be zero	
-	}
+		return Type(
+			TRemoteControlTypeTraits<FRotator>::DefaultMappingValueMax(),
+			TRemoteControlTypeTraits<FVector>::DefaultMappingValueMax(),
+			TRemoteControlTypeTraits<FVector>::DefaultMappingValueMax()); // scale is Max cause it shouldn't be zero
+	}		
 };
 
 /** RemoteControlTypeTraits for FIntPoint */
 template <>
 struct TRemoteControlTypeTraits<FIntPoint>
 {
-	typedef FIntPoint Type;
+	using Type = FIntPoint;
+
+	/** Is ValueType supported as a range (protocol input) value? */
+	static constexpr bool IsSupportedRangeType() { return false; }
+
+	/** Is ValueType supported as a mapping (property output) value? */
+	static constexpr bool IsSupportedMappingType() { return true; }
 	
-	static constexpr Type DefaultMin()
-	{
-		return Type::ZeroValue;
-	}
+	static Type DefaultMappingValueMin() { return Type::ZeroValue; }
 
-	static constexpr Type DefaultMax()
-	{
-		return Type(1, 1);
-	}
+	static Type DefaultMappingValueMax() { return Type(1, 1); }
 };
 
 /** RemoteControlTypeTraits for FIntVector */
 template <>
 struct TRemoteControlTypeTraits<FIntVector>
 {
-	typedef FIntVector Type;
-	
-	static constexpr Type DefaultMin()
-	{
-		return Type::ZeroValue;
-	}
+	using Type = FIntVector;
 
-	static constexpr Type DefaultMax()
-	{
-		return Type(1, 1, 1);
-	}
+	/** Is ValueType supported as a range (protocol input) value? */
+	static constexpr bool IsSupportedRangeType() { return false; }
+
+	/** Is ValueType supported as a mapping (property output) value? */
+	static constexpr bool IsSupportedMappingType() { return true; }
+	
+	static Type DefaultMappingValueMin() { return Type::ZeroValue; }
+	
+	static Type DefaultMappingValueMax() { return Type(1, 1, 1); }
 };
 
 /** RemoteControlTypeTraits for FBox */
 template <>
 struct TRemoteControlTypeTraits<FBox>
 {
-	typedef FBox Type;
+	using Type = FBox;
 	
-	static constexpr Type DefaultMin()
+	/** Is ValueType supported as a range (protocol input) value? */
+	static constexpr bool IsSupportedRangeType() { return false; }
+
+	/** Is ValueType supported as a mapping (property output) value? */
+	static constexpr bool IsSupportedMappingType() { return true; }
+	
+	static Type DefaultMappingValueMin()
 	{
-		return Type(TRemoteControlTypeTraits<FVector>::DefaultMin(), TRemoteControlTypeTraits<FVector>::DefaultMax());
+		return Type(
+			TRemoteControlTypeTraits<FVector>::DefaultMappingValueMin(),
+			TRemoteControlTypeTraits<FVector>::DefaultMappingValueMax());
 	}
 
-	static constexpr Type DefaultMax()
+	static Type DefaultMappingValueMax()
 	{
-		return Type(TRemoteControlTypeTraits<FVector>::DefaultMax(), TRemoteControlTypeTraits<FVector>::DefaultMax() * 2.0f);
+		return Type(
+			TRemoteControlTypeTraits<FVector>::DefaultMappingValueMax(),
+			TRemoteControlTypeTraits<FVector>::DefaultMappingValueMax() * 2.0f);
 	}
 };
 
@@ -216,16 +578,26 @@ struct TRemoteControlTypeTraits<FBox>
 template <>
 struct TRemoteControlTypeTraits<FBox2D>
 {
-	typedef FBox2D Type;
+	using Type = FBox2D;
+	
+	/** Is ValueType supported as a range (protocol input) value? */
+	static constexpr bool IsSupportedRangeType() { return false; }
+
+	/** Is ValueType supported as a mapping (property output) value? */
+	static constexpr bool IsSupportedMappingType() { return true; }
 	
-	static constexpr Type DefaultMin()
+	static Type DefaultMappingValueMin()
 	{
-		return Type(TRemoteControlTypeTraits<FVector2D>::DefaultMin(), TRemoteControlTypeTraits<FVector2D>::DefaultMax());
+		return Type(
+			TRemoteControlTypeTraits<FVector2D>::DefaultMappingValueMin(),
+			TRemoteControlTypeTraits<FVector2D>::DefaultMappingValueMax());
 	}
 
-	static constexpr Type DefaultMax()
+	static Type DefaultMappingValueMax()
 	{
-		return Type(TRemoteControlTypeTraits<FVector2D>::DefaultMax(), TRemoteControlTypeTraits<FVector2D>::DefaultMax() * 2.0f);
+		return Type(
+			TRemoteControlTypeTraits<FVector2D>::DefaultMappingValueMax(),
+			TRemoteControlTypeTraits<FVector2D>::DefaultMappingValueMax() * 2.0f);
 	}
 };
 
@@ -233,39 +605,43 @@ struct TRemoteControlTypeTraits<FBox2D>
 template <>
 struct TRemoteControlTypeTraits<FBoxSphereBounds>
 {
-	typedef FBoxSphereBounds Type;
+	using Type = FBoxSphereBounds;
 	
-	static constexpr Type DefaultMin()
-	{
-		return Type(TRemoteControlTypeTraits<FBox>::DefaultMin());
-	}
+	/** Is ValueType supported as a range (protocol input) value? */
+	static constexpr bool IsSupportedRangeType() { return false; }
 
-	static constexpr Type DefaultMax()
-	{
-		return Type(TRemoteControlTypeTraits<FBox>::DefaultMax());
-	}
+	/** Is ValueType supported as a mapping (property output) value? */
+	static constexpr bool IsSupportedMappingType() { return true; }
+	
+	static Type DefaultMappingValueMin() { return Type(TRemoteControlTypeTraits<FBox>::DefaultMappingValueMin()); }
+
+	static Type DefaultMappingValueMax() { return Type(TRemoteControlTypeTraits<FBox>::DefaultMappingValueMax()); }
 };
 
 /** RemoteControlTypeTraits for FColor */
 template <>
 struct TRemoteControlTypeTraits<FColor>
 {
-	typedef FColor Type;
+	using Type = FColor;
 	
-	static constexpr Type DefaultMin()
+	/** Is ValueType supported as a range (protocol input) value? */
+	static constexpr bool IsSupportedRangeType() { return false; }
+
+	/** Is ValueType supported as a mapping (property output) value? */
+	static constexpr bool IsSupportedMappingType() { return true; }
+	
+	static Type DefaultMappingValueMin()
 	{
-		return Type(
-			TRemoteControlTypeTraits<uint8>::DefaultMin(),
-			TRemoteControlTypeTraits<uint8>::DefaultMin(),
-			TRemoteControlTypeTraits<uint8>::DefaultMin());
+		return Type(TRemoteControlTypeTraits<uint8>::DefaultMappingValueMin(),
+			TRemoteControlTypeTraits<uint8>::DefaultMappingValueMin(),
+			TRemoteControlTypeTraits<uint8>::DefaultMappingValueMin());
 	}
-
-	static constexpr Type DefaultMax()
+		
+	static Type DefaultMappingValueMax()
 	{
-		return Type(
-			 TRemoteControlTypeTraits<uint8>::DefaultMax(),
-			TRemoteControlTypeTraits<uint8>::DefaultMax(),
-			TRemoteControlTypeTraits<uint8>::DefaultMax());
+		return Type(TRemoteControlTypeTraits<uint8>::DefaultMappingValueMax(),
+			TRemoteControlTypeTraits<uint8>::DefaultMappingValueMax(),
+			TRemoteControlTypeTraits<uint8>::DefaultMappingValueMax());
 	}
 };
 
@@ -273,21 +649,27 @@ struct TRemoteControlTypeTraits<FColor>
 template <>
 struct TRemoteControlTypeTraits<FLinearColor>
 {
-	typedef FLinearColor Type;
+	using Type = FLinearColor;
+	
+	/** Is ValueType supported as a range (protocol input) value? */
+	static constexpr bool IsSupportedRangeType() { return false; }
+
+	/** Is ValueType supported as a mapping (property output) value? */
+	static constexpr bool IsSupportedMappingType() { return true; }
 	
-	static constexpr Type DefaultMin()
+	static Type DefaultMappingValueMin()
 	{
-		return Type(
-			TRemoteControlTypeTraits<float>::DefaultMin(),
-			TRemoteControlTypeTraits<float>::DefaultMin(),
-			TRemoteControlTypeTraits<float>::DefaultMin());
+		return Type(TRemoteControlTypeTraits<float>::DefaultMappingValueMin(),
+			TRemoteControlTypeTraits<float>::DefaultMappingValueMin(),
+			TRemoteControlTypeTraits<float>::DefaultMappingValueMin());
 	}
 
-	static constexpr Type DefaultMax()
+	static Type DefaultMappingValueMax()
 	{
-		return Type(
-			TRemoteControlTypeTraits<float>::DefaultMax(),
-			TRemoteControlTypeTraits<float>::DefaultMax(),
-			TRemoteControlTypeTraits<float>::DefaultMax());
+		return Type(TRemoteControlTypeTraits<float>::DefaultMappingValueMax(),
+			TRemoteControlTypeTraits<float>::DefaultMappingValueMax(),
+			TRemoteControlTypeTraits<float>::DefaultMappingValueMax());
 	}
 };
+
+#pragma endregion Structs (Built-in)
diff --git a/Engine/Plugins/VirtualProduction/RemoteControl/Source/RemoteControlCommon/Public/RCTypeUtilities.h b/Engine/Plugins/VirtualProduction/RemoteControl/Source/RemoteControlCommon/Public/RCTypeUtilities.h
new file mode 100644
--- /dev/null
+++ b/Engine/Plugins/VirtualProduction/RemoteControl/Source/RemoteControlCommon/Public/RCTypeUtilities.h
@@ -0,0 +1,505 @@
+﻿// Copyright Epic Games, Inc. All Rights Reserved.
+
+#pragma once
+
+#include "CoreMinimal.h"
+
+#include "RCTypeTraits.h"
+#include "Math/NumericLimits.h"
+#include "Math/UnrealMathUtility.h"
+#include "Templates/UnrealTypeTraits.h"
+#include "UObject/StructOnScope.h"
+#include "UObject/TextProperty.h"
+#include "UObject/UnrealType.h"
+
+class IPropertyHandle;
+
+namespace RemoteControlTypeUtilities
+{
+	#pragma region FOREACH_CAST_PROPERTY
+	#ifndef FOREACH_CAST_PROPERTY
+	
+	/**
+	 * Iterate and attempt to cast Property to different types (as CastPropertyType), where the instance will be CastProperty
+	 * Usage:
+	 * FProperty* SomePropertyInstance;
+	 *
+	 * template <>
+	 * void SomeFunc<FNumericProperty>(const FNumericProperty*);
+	 * etc..
+	 * 
+	 * FOREACH_CAST_PROPERTY(SomePropertyInstance, SomeFunc<CastPropertyType>(CastProperty))
+	 */
+	#define FOREACH_CAST_PROPERTY(Property, Func)											\
+		if(const FBoolProperty* CastProperty = CastField<FBoolProperty>(Property))			\
+		{																					\
+			using CastPropertyType = FBoolProperty;											\
+			return Func;																	\
+		}																					\
+		if(const FEnumProperty* CastProperty = CastField<FEnumProperty>(Property))			\
+		{																					\
+			using CastPropertyType = FEnumProperty;											\
+			return Func;																	\
+		}																					\
+		if(const FNumericProperty* CastProperty = CastField<FNumericProperty>(Property))	\
+		{																					\
+			using CastPropertyType = FNumericProperty;										\
+			return Func;																	\
+		}																					\
+		if(const FStructProperty* CastProperty = CastField<FStructProperty>(Property))		\
+		{																					\
+			using CastPropertyType = FStructProperty;										\
+			return Func;																	\
+		}																					\
+																							\
+		if(const FStrProperty* CastProperty = CastField<FStrProperty>(Property))			\
+		{																					\
+			using CastPropertyType = FStrProperty;											\
+			return Func;																	\
+		}																					\
+		if(const FNameProperty* CastProperty = CastField<FNameProperty>(Property))			\
+		{																					\
+			using CastPropertyType = FNameProperty;											\
+			return Func;																	\
+		}																					\
+		if(const FTextProperty* CastProperty = CastField<FTextProperty>(Property))			\
+		{																					\
+			using CastPropertyType = FTextProperty;											\
+			return Func;																	\
+		}																					\
+																							\
+		if(const FArrayProperty* CastProperty = CastField<FArrayProperty>(Property))		\
+		{																					\
+			using CastPropertyType = FArrayProperty;										\
+			return Func;																	\
+		}																					\
+		if(const FSetProperty* CastProperty = CastField<FSetProperty>(Property))			\
+		{																					\
+			using CastPropertyType = FSetProperty;											\
+			return Func;																	\
+		}																					\
+		if(const FMapProperty* CastProperty = CastField<FMapProperty>(Property))			\
+		{																					\
+			using CastPropertyType = FMapProperty;											\
+			return Func;																	\
+		}																					\
+
+	#endif
+	#pragma endregion FOREACH_CAST_PROPERTY
+			
+	/** Metadata Keys. */
+	
+	static FName ClampMinKey = "ClampMin";
+	static FName ClampMaxKey = "ClampMax";
+
+	static FName UIMinKey = "UIMin";
+	static FName UIMaxKey = "UIMax";
+
+	/** Delegated Min operator. */
+	template <typename ValueType>
+	static auto MinOp = [](auto&& InValues) { return FMath::Min(InValues); };
+
+	/** Delegated Max operator. */
+	template <typename ValueType>
+	static auto MaxOp = [](auto&& InValues) { return FMath::Max(InValues); };
+
+	/** Returns typed metadata for numeric property if it exists, otherwise returns the DefaultValue. */
+	template <typename PropertyType, typename ValueType>
+	constexpr static typename TEnableIf<
+		TAnd<
+			TIsDerivedFrom<PropertyType, FProperty>,
+			RemoteControlTypeTraits::TNumericValueConstraint<ValueType>>::Value, ValueType>::Type
+	GetMetadataValue(const PropertyType* InProperty, const FName& InKey, const ValueType& InDefaultValue)
+	{
+		if(InProperty->HasMetaData(InKey))
+		{
+			if(const FNumericProperty* NumericProperty = CastField<FNumericProperty>(InProperty))
+			{
+				if(NumericProperty->IsInteger())
+				{
+					return static_cast<ValueType>(NumericProperty->GetIntMetaData(InKey));
+				}
+			
+				if(NumericProperty->IsFloatingPoint())
+				{
+					return static_cast<ValueType>(NumericProperty->GetFloatMetaData(InKey));
+				}
+			}
+		}
+
+		return InDefaultValue;
+	}
+
+	/** Returns typed metadata if it exists, otherwise returns the DefaultValue. */
+	template <typename PropertyType, typename ValueType>
+	constexpr static typename TEnableIf<
+		TAnd<
+			TIsDerivedFrom<PropertyType, FProperty>,
+			TNot<RemoteControlTypeTraits::TNumericValueConstraint<ValueType>>>::Value, ValueType>::Type
+	GetMetadataValue(const PropertyType* InProperty, const FName& InKey, const ValueType InDefaultValue)
+	{
+		return InDefaultValue;
+	}
+
+	template <class Func, class PropertyType, class ValueType>
+	static ValueType GetClampedValue(Func InFunc, const PropertyType* InProperty, ValueType InValue, const TArray<FName>& InMetaKeys)
+	{
+		TArray<ValueType> Params;
+		Params.Reserve(InMetaKeys.Num());
+		for(const FName& Key : InMetaKeys)
+		{
+			Params.Add(
+				RemoteControlTypeUtilities::GetMetadataValue(
+					InProperty,
+					Key,
+					InValue));
+		}
+
+		return InFunc(Params);
+	}
+
+	/** Clamps the value to the PropertyType, for example if an int64 is requested but the PropertyType is UInt16Property. */
+	template <typename PropertyType, typename ValueType>
+	static ValueType ClampToPropertyType(const FNumericProperty* InProperty, ValueType& InOutValue)
+	{
+		check(InProperty);
+
+		const FName PropertyTypeName = InProperty->GetClass()->GetFName();
+		if(InProperty->IsInteger())
+		{
+			if(PropertyTypeName == NAME_ByteProperty)
+			{
+				using T = uint8;
+				InOutValue = FMath::Clamp<T>(InOutValue, TNumericLimits<T>::Min(), TNumericLimits<T>::Max());
+			}
+			else if(PropertyTypeName == NAME_Int8Property)
+			{
+				using T = int8;
+				InOutValue = FMath::Clamp<T>(InOutValue, TNumericLimits<T>::Min(), TNumericLimits<T>::Max());
+			}
+			else if(PropertyTypeName == NAME_Int16Property)
+			{
+				using T = int16;
+				InOutValue = FMath::Clamp<T>(InOutValue, TNumericLimits<T>::Min(), TNumericLimits<T>::Max());
+			}
+			else if(PropertyTypeName == NAME_UInt16Property)
+			{
+				using T = uint16;
+				InOutValue = FMath::Clamp<T>(InOutValue, TNumericLimits<T>::Min(), TNumericLimits<T>::Max());
+			}
+			else if(PropertyTypeName == NAME_Int32Property)
+			{
+				using T = int32;
+				InOutValue = FMath::Clamp<T>(InOutValue, TNumericLimits<T>::Min(), TNumericLimits<T>::Max());
+			}
+			else if(PropertyTypeName == NAME_UInt32Property)
+			{
+				using T = uint32;
+				InOutValue = FMath::Clamp<T>(InOutValue, TNumericLimits<T>::Min(), TNumericLimits<T>::Max());
+			}
+			else if(PropertyTypeName == NAME_Int64Property)
+			{
+				using T = int64;
+				InOutValue = FMath::Clamp<T>(InOutValue, TNumericLimits<T>::Min(), TNumericLimits<T>::Max());
+			}
+			else if(PropertyTypeName == NAME_UInt64Property)
+			{
+				using T = uint64;
+				InOutValue = FMath::Clamp<T>(InOutValue, TNumericLimits<T>::Min(), TNumericLimits<T>::Max());
+			}
+		}
+		else if(InProperty->IsFloatingPoint())
+		{
+			if(PropertyTypeName == NAME_FloatProperty)
+			{
+				using T = float;
+				InOutValue = FMath::Clamp<T>(InOutValue, TNumericLimits<T>::Min(), TNumericLimits<T>::Max());
+			}
+			else if(PropertyTypeName == NAME_DoubleProperty)
+			{
+				using T = double;
+				InOutValue = FMath::Clamp<T>(InOutValue, TNumericLimits<T>::Min(), TNumericLimits<T>::Max());
+			}
+		}
+
+		return InOutValue;
+	}
+	
+	/** The utilities below wrap TRemoteControlTypeTraits */
+
+	template <typename PropertyType>
+	static bool IsSupportedRangeType(const PropertyType* InProperty);
+
+	/** Is ValueType supported as a range (protocol input) value? */
+	template <typename PropertyType>
+	static bool IsSupportedRangeType(const PropertyType* InProperty)
+	{
+		check(InProperty);
+
+		if(const FArrayProperty* ArrayProperty = CastField<FArrayProperty>(InProperty))
+		{
+			return IsSupportedRangeType(ArrayProperty->Inner);
+		}
+		
+		if(const FSetProperty* SetProperty = CastField<FSetProperty>(InProperty))
+		{
+			return IsSupportedRangeType(SetProperty->ElementProp);
+		}
+		
+		if(const FMapProperty* MapProperty = CastField<FMapProperty>(InProperty))
+		{
+			return IsSupportedRangeType(MapProperty->KeyProp);
+		}
+
+		FOREACH_CAST_PROPERTY(InProperty, TRemoteControlPropertyTypeTraits<CastPropertyType>::IsSupportedRangeType())
+		
+		return false;
+	}
+
+	/** Is ValueType supported as a mapping (property output) value? */
+	template <typename PropertyType>
+	static constexpr bool IsSupportedMappingType(const PropertyType* InProperty);
+
+	/** Is ValueType supported as a mapping (property output) value? */
+	template <typename PropertyType>
+	static constexpr bool IsSupportedMappingType(const PropertyType* InProperty)
+	{
+		check(InProperty);
+
+		if(const FArrayProperty* ArrayProperty = CastField<FArrayProperty>(InProperty))
+		{
+			return IsSupportedMappingType(ArrayProperty->Inner);
+		}
+		
+		if(const FSetProperty* SetProperty = CastField<FSetProperty>(InProperty))
+		{
+			return IsSupportedMappingType(SetProperty->ElementProp);
+		}
+		
+		if(const FMapProperty* MapProperty = CastField<FMapProperty>(InProperty))
+		{
+			return IsSupportedMappingType(MapProperty->ValueProp);
+		}
+
+		FOREACH_CAST_PROPERTY(InProperty, TRemoteControlPropertyTypeTraits<CastPropertyType>::IsSupportedMappingType())
+
+		return false;
+	}
+
+	template <typename ValueType>
+	static constexpr ValueType GetDefaultRangeValueMin()
+	{
+		return TRemoteControlTypeTraits<ValueType>::DefaultRangeValueMin();
+	}
+
+	template <typename ValueType>
+	static constexpr typename TEnableIf<TNot<RemoteControlTypeTraits::TNumericValueConstraint<ValueType>>::Value, ValueType>::Type
+	GetDefaultRangeValueMin(const FProperty* InProperty)
+	{
+		check(InProperty);
+
+		return TRemoteControlTypeTraits<ValueType>::DefaultRangeValueMin();
+	}
+
+	template <typename ValueType>
+	static constexpr typename TEnableIf<RemoteControlTypeTraits::TNumericValueConstraint<ValueType>::Value, ValueType>::Type
+	GetDefaultRangeValueMin(const FProperty* InProperty)
+	{
+		check(InProperty);
+
+		// Choose greater than the 3 values: (Traits defined) Default, ClampMin, or UIMin
+		ValueType ValueResult = GetClampedValue(MaxOp<ValueType>, InProperty, TRemoteControlTypeTraits<ValueType>::DefaultRangeValueMin(), { ClampMinKey, UIMinKey });
+
+		// ValueType might not match property, so clamp to min/max of property type
+		if(const FNumericProperty* NumericProperty = CastField<FNumericProperty>(InProperty))
+		{
+			ClampToPropertyType<ValueType>(NumericProperty, ValueResult);
+		}
+
+		return ValueResult;
+	}
+
+	template <typename ValueType>
+	static constexpr ValueType GetDefaultRangeValueMax()
+	{
+		return TRemoteControlTypeTraits<ValueType>::DefaultRangeValueMax();
+	}
+
+	template <typename ValueType>
+	static constexpr typename TEnableIf<TNot<RemoteControlTypeTraits::TNumericValueConstraint<ValueType>>::Value, ValueType>::Type
+	GetDefaultRangeValueMax(const FProperty* InProperty)
+	{
+		check(InProperty);
+
+		return TRemoteControlTypeTraits<ValueType>::DefaultRangeValueMax();
+	}
+
+	template <typename ValueType>
+	static constexpr typename TEnableIf<RemoteControlTypeTraits::TNumericValueConstraint<ValueType>::Value, ValueType>::Type
+	GetDefaultRangeValueMax(const FProperty* InProperty)
+	{
+		check(InProperty);
+		
+		// Choose lesser than the 3 values: (Traits defined) Default, ClampMin, or UIMin
+		ValueType ValueResult =	GetClampedValue(MinOp<ValueType>, InProperty, TRemoteControlTypeTraits<ValueType>::DefaultRangeValueMax(), { ClampMaxKey, UIMaxKey });
+
+		// ValueType might not match property, so clamp to min/max of property type
+		if(const FNumericProperty* NumericProperty = CastField<FNumericProperty>(InProperty))
+		{
+			ClampToPropertyType<ValueType>(NumericProperty, ValueResult);
+		}
+
+		return ValueResult;
+	}
+
+	template <typename ValueType>
+	static constexpr ValueType GetDefaultMappingValueMin()
+	{
+		return TRemoteControlTypeTraits<ValueType>::DefaultMappingValueMin();
+	}
+
+	template <typename ValueType>
+	static constexpr typename TEnableIf<TNot<RemoteControlTypeTraits::TNumericValueConstraint<ValueType>>::Value, ValueType>::Type
+	GetDefaultMappingValueMin(const FProperty* InProperty)
+	{
+		check(InProperty);
+
+		return TRemoteControlTypeTraits<ValueType>::DefaultMappingValueMin();
+	}
+
+	template <typename ValueType>
+	static constexpr typename TEnableIf<RemoteControlTypeTraits::TNumericValueConstraint<ValueType>::Value, ValueType>::Type
+	GetDefaultMappingValueMin(const FProperty* InProperty)
+	{
+		check(InProperty);
+		
+		ValueType ValueResult = GetClampedValue(MaxOp<ValueType>, InProperty, TRemoteControlTypeTraits<ValueType>::DefaultMappingValueMin(), { ClampMinKey, UIMinKey });
+
+		// ValueType might not match property, so clamp to min/max of property type
+		if(const FNumericProperty* NumericProperty = CastField<FNumericProperty>(InProperty))
+		{
+			ClampToPropertyType<ValueType>(NumericProperty, ValueResult);
+		}
+
+		return ValueResult;
+	}
+
+	template <typename ValueType>
+	static typename TEnableIf<TIsSame<ValueType, FStructOnScope>::Value, TSharedPtr<ValueType>>::Type
+	GetDefaultMappingValueMin(const FStructProperty* InProperty)
+	{
+		check(InProperty);
+
+		TSharedPtr<FStructOnScope> PopulatedStruct = MakeShared<FStructOnScope>(InProperty->Struct);
+
+		/*
+		// See UETOOL-3502
+		for (TFieldIterator<FProperty> It(InProperty->Struct, EFieldIteratorFlags::IncludeSuper); It; ++It)
+		{
+			FProperty* Property = *It;
+			uint8* PropertyAddr = Property->ContainerPtrToValuePtr<uint8>(PopulatedStruct->GetStructMemory());
+		}
+		*/
+
+		return PopulatedStruct;
+	}
+
+	template <typename ValueType>
+	static constexpr ValueType GetDefaultMappingValueMax()
+	{
+		return TRemoteControlTypeTraits<ValueType>::DefaultMappingValueMax();		
+	}
+
+	template <typename ValueType>
+	static constexpr typename TEnableIf<TNot<RemoteControlTypeTraits::TNumericValueConstraint<ValueType>>::Value, ValueType>::Type
+	GetDefaultMappingValueMax(const FProperty* InProperty)
+	{
+		check(InProperty);
+
+		return TRemoteControlTypeTraits<ValueType>::DefaultMappingValueMax();
+	}
+
+	template <typename ValueType>
+	static constexpr typename TEnableIf<RemoteControlTypeTraits::TNumericValueConstraint<ValueType>::Value, ValueType>::Type
+	GetDefaultMappingValueMax(const FProperty* InProperty)
+	{
+		check(InProperty);
+		
+		ValueType ValueResult = GetClampedValue(MinOp<ValueType>, InProperty, TRemoteControlTypeTraits<ValueType>::DefaultMappingValueMax(), { ClampMaxKey, UIMaxKey });
+
+		// ValueType might not match property, so clamp to min/max of property type
+		if(const FNumericProperty* NumericProperty = CastField<FNumericProperty>(InProperty))
+		{
+			ClampToPropertyType<ValueType>(NumericProperty, ValueResult);
+		}
+
+		return ValueResult;
+	}
+
+	template <typename ValueType>
+	static constexpr typename TEnableIf<TIsSame<ValueType, FStructOnScope>::Value, TSharedPtr<ValueType>>::Type
+	GetDefaultMappingValueMax(const FStructProperty* InProperty)
+	{
+		check(InProperty);
+
+		return TRemoteControlTypeTraits<ValueType>::DefaultMappingValueMax();
+	}
+
+	/** Various RemoteControl property type Traits (wraps internal implementations). This resolves the property type for you. */
+
+	template <typename PropertyType, typename ValueType>
+	using TPropertyConstraint = typename TEnableIf<
+		TAnd<TIsDerivedFrom<PropertyType, FProperty>>::Value, ValueType>::Type;
+	
+	template <typename PropertyType, typename ValueType>
+	constexpr static TPropertyConstraint<PropertyType, ValueType>
+	GetDefaultRangeValueMin(const PropertyType* InProperty)
+	{
+		check(InProperty);
+		
+		using FTraits = TRemoteControlTypeTraits<ValueType>;
+		
+		return GetClampedValue(MaxOp<ValueType>, InProperty, TRemoteControlTypeTraits<ValueType>::DefaultRangeValueMin(), { ClampMinKey, UIMinKey });
+	}
+
+	template <typename PropertyType, typename ValueType>
+	constexpr static TPropertyConstraint<PropertyType, ValueType>
+	GetDefaultRangeValueMax(const FProperty* InProperty)
+	{
+		check(InProperty);
+
+		using FTraits = TRemoteControlTypeTraits<ValueType>;
+
+		return GetClampedValue(MinOp<ValueType>, InProperty, TRemoteControlTypeTraits<ValueType>::DefaultRangeValueMax(), { ClampMaxKey, UIMaxKey });
+	}
+
+	template <typename PropertyType, typename ValueType>
+	constexpr static TPropertyConstraint<PropertyType, ValueType>
+	GetDefaultMappingValueMin(const FProperty* InProperty)
+	{
+		check(InProperty);
+		
+		using FTraits = TRemoteControlTypeTraits<ValueType>;
+		
+		return GetClampedValue(MaxOp<ValueType>, InProperty, TRemoteControlTypeTraits<ValueType>::DefaultMappingValueMin(), { ClampMinKey, UIMinKey });
+	}
+
+	template <typename PropertyType, typename ValueType>
+	constexpr static TPropertyConstraint<PropertyType, ValueType>
+	GetDefaultMappingValueMax(const FProperty* InProperty)
+	{
+		check(InProperty);
+		
+		using FTraits = TRemoteControlTypeTraits<ValueType>;
+		
+		return GetClampedValue(MinOp<ValueType>, InProperty, TRemoteControlTypeTraits<ValueType>::DefaultMappingValueMax(), { ClampMaxKey, UIMaxKey });
+	}
+
+	/** Size of given property + data, accounting for item count (arrays, strings, etc.) */
+	REMOTECONTROLCOMMON_API SIZE_T GetPropertySize(const FProperty* InProperty, void* InData);
+
+#if WITH_EDITOR
+	/** Size of given property + data, accounting for item count (arrays, strings, etc.) */
+	REMOTECONTROLCOMMON_API SIZE_T GetPropertySize(const TSharedPtr<IPropertyHandle>& InPropertyHandle);
+#endif
+}
