compilation fix (constexpr)

diff --git a/Engine/Plugins/VirtualProduction/RemoteControl/Source/RemoteControlCommon/Public/RCTypeUtilities.h b/Engine/Plugins/VirtualProduction/RemoteControl/Source/RemoteControlCommon/Public/RCTypeUtilities.h
--- a/Engine/Plugins/VirtualProduction/RemoteControl/Source/RemoteControlCommon/Public/RCTypeUtilities.h
+++ b/Engine/Plugins/VirtualProduction/RemoteControl/Source/RemoteControlCommon/Public/RCTypeUtilities.h
@@ -262,7 +262,7 @@ namespace RemoteControlTypeUtilities
 
 	/** Is ValueType supported as a mapping (property output) value? */
 	template <typename PropertyType>
-	static constexpr bool IsSupportedMappingType(const PropertyType* InProperty)
+	static bool IsSupportedMappingType(const PropertyType* InProperty)
 	{
 		check(InProperty);
 
@@ -293,7 +293,7 @@ namespace RemoteControlTypeUtilities
 	}
 
 	template <typename ValueType>
-	static constexpr typename TEnableIf<TNot<RemoteControlTypeTraits::TNumericValueConstraint<ValueType>>::Value, ValueType>::Type
+	static typename TEnableIf<TNot<RemoteControlTypeTraits::TNumericValueConstraint<ValueType>>::Value, ValueType>::Type
 	GetDefaultRangeValueMin(const FProperty* InProperty)
 	{
 		check(InProperty);
@@ -302,7 +302,7 @@ namespace RemoteControlTypeUtilities
 	}
 
 	template <typename ValueType>
-	static constexpr typename TEnableIf<RemoteControlTypeTraits::TNumericValueConstraint<ValueType>::Value, ValueType>::Type
+	static typename TEnableIf<RemoteControlTypeTraits::TNumericValueConstraint<ValueType>::Value, ValueType>::Type
 	GetDefaultRangeValueMin(const FProperty* InProperty)
 	{
 		check(InProperty);
@@ -326,7 +326,7 @@ namespace RemoteControlTypeUtilities
 	}
 
 	template <typename ValueType>
-	static constexpr typename TEnableIf<TNot<RemoteControlTypeTraits::TNumericValueConstraint<ValueType>>::Value, ValueType>::Type
+	static typename TEnableIf<TNot<RemoteControlTypeTraits::TNumericValueConstraint<ValueType>>::Value, ValueType>::Type
 	GetDefaultRangeValueMax(const FProperty* InProperty)
 	{
 		check(InProperty);
@@ -335,7 +335,7 @@ namespace RemoteControlTypeUtilities
 	}
 
 	template <typename ValueType>
-	static constexpr typename TEnableIf<RemoteControlTypeTraits::TNumericValueConstraint<ValueType>::Value, ValueType>::Type
+	static typename TEnableIf<RemoteControlTypeTraits::TNumericValueConstraint<ValueType>::Value, ValueType>::Type
 	GetDefaultRangeValueMax(const FProperty* InProperty)
 	{
 		check(InProperty);
@@ -359,7 +359,7 @@ namespace RemoteControlTypeUtilities
 	}
 
 	template <typename ValueType>
-	static constexpr typename TEnableIf<TNot<RemoteControlTypeTraits::TNumericValueConstraint<ValueType>>::Value, ValueType>::Type
+	static typename TEnableIf<TNot<RemoteControlTypeTraits::TNumericValueConstraint<ValueType>>::Value, ValueType>::Type
 	GetDefaultMappingValueMin(const FProperty* InProperty)
 	{
 		check(InProperty);
@@ -368,7 +368,7 @@ namespace RemoteControlTypeUtilities
 	}
 
 	template <typename ValueType>
-	static constexpr typename TEnableIf<RemoteControlTypeTraits::TNumericValueConstraint<ValueType>::Value, ValueType>::Type
+	static typename TEnableIf<RemoteControlTypeTraits::TNumericValueConstraint<ValueType>::Value, ValueType>::Type
 	GetDefaultMappingValueMin(const FProperty* InProperty)
 	{
 		check(InProperty);
@@ -411,7 +411,7 @@ namespace RemoteControlTypeUtilities
 	}
 
 	template <typename ValueType>
-	static constexpr typename TEnableIf<TNot<RemoteControlTypeTraits::TNumericValueConstraint<ValueType>>::Value, ValueType>::Type
+	static typename TEnableIf<TNot<RemoteControlTypeTraits::TNumericValueConstraint<ValueType>>::Value, ValueType>::Type
 	GetDefaultMappingValueMax(const FProperty* InProperty)
 	{
 		check(InProperty);
@@ -420,7 +420,7 @@ namespace RemoteControlTypeUtilities
 	}
 
 	template <typename ValueType>
-	static constexpr typename TEnableIf<RemoteControlTypeTraits::TNumericValueConstraint<ValueType>::Value, ValueType>::Type
+	static typename TEnableIf<RemoteControlTypeTraits::TNumericValueConstraint<ValueType>::Value, ValueType>::Type
 	GetDefaultMappingValueMax(const FProperty* InProperty)
 	{
 		check(InProperty);
@@ -437,7 +437,7 @@ namespace RemoteControlTypeUtilities
 	}
 
 	template <typename ValueType>
-	static constexpr typename TEnableIf<TIsSame<ValueType, FStructOnScope>::Value, TSharedPtr<ValueType>>::Type
+	static typename TEnableIf<TIsSame<ValueType, FStructOnScope>::Value, TSharedPtr<ValueType>>::Type
 	GetDefaultMappingValueMax(const FStructProperty* InProperty)
 	{
 		check(InProperty);
@@ -452,7 +452,7 @@ namespace RemoteControlTypeUtilities
 		TAnd<TIsDerivedFrom<PropertyType, FProperty>>::Value, ValueType>::Type;
 	
 	template <typename PropertyType, typename ValueType>
-	constexpr static TPropertyConstraint<PropertyType, ValueType>
+	static TPropertyConstraint<PropertyType, ValueType>
 	GetDefaultRangeValueMin(const PropertyType* InProperty)
 	{
 		check(InProperty);
@@ -463,7 +463,7 @@ namespace RemoteControlTypeUtilities
 	}
 
 	template <typename PropertyType, typename ValueType>
-	constexpr static TPropertyConstraint<PropertyType, ValueType>
+	static TPropertyConstraint<PropertyType, ValueType>
 	GetDefaultRangeValueMax(const FProperty* InProperty)
 	{
 		check(InProperty);
@@ -474,7 +474,7 @@ namespace RemoteControlTypeUtilities
 	}
 
 	template <typename PropertyType, typename ValueType>
-	constexpr static TPropertyConstraint<PropertyType, ValueType>
+	static TPropertyConstraint<PropertyType, ValueType>
 	GetDefaultMappingValueMin(const FProperty* InProperty)
 	{
 		check(InProperty);
@@ -485,7 +485,7 @@ namespace RemoteControlTypeUtilities
 	}
 
 	template <typename PropertyType, typename ValueType>
-	constexpr static TPropertyConstraint<PropertyType, ValueType>
+	static TPropertyConstraint<PropertyType, ValueType>
 	GetDefaultMappingValueMax(const FProperty* InProperty)
 	{
 		check(InProperty);
