fix for previous fix

diff --git a/Engine/Plugins/VirtualProduction/RemoteControl/Source/RemoteControlCommon/Public/RCTypeUtilities.h b/Engine/Plugins/VirtualProduction/RemoteControl/Source/RemoteControlCommon/Public/RCTypeUtilities.h
--- a/Engine/Plugins/VirtualProduction/RemoteControl/Source/RemoteControlCommon/Public/RCTypeUtilities.h
+++ b/Engine/Plugins/VirtualProduction/RemoteControl/Source/RemoteControlCommon/Public/RCTypeUtilities.h
@@ -258,7 +258,7 @@ namespace RemoteControlTypeUtilities
 
 	/** Is ValueType supported as a mapping (property output) value? */
 	template <typename PropertyType>
-	static constexpr bool IsSupportedMappingType(const PropertyType* InProperty);
+	static bool IsSupportedMappingType(const PropertyType* InProperty);
 
 	/** Is ValueType supported as a mapping (property output) value? */
 	template <typename PropertyType>
