Merging CL 17884606 from release stream

diff --git a/Engine/Plugins/VirtualProduction/RemoteControl/Source/RemoteControlUI/Private/UI/SRCPanelFunctionPicker.cpp b/Engine/Plugins/VirtualProduction/RemoteControl/Source/RemoteControlUI/Private/UI/SRCPanelFunctionPicker.cpp
--- a/Engine/Plugins/VirtualProduction/RemoteControl/Source/RemoteControlUI/Private/UI/SRCPanelFunctionPicker.cpp
+++ b/Engine/Plugins/VirtualProduction/RemoteControl/Source/RemoteControlUI/Private/UI/SRCPanelFunctionPicker.cpp
@@ -241,7 +241,6 @@ void SRCPanelFunctionPicker::Tick(const FGeometry& AllottedGeometry, const doubl
 	constexpr double MinimumTimeToSwitchFocus = 0.2;
 	if (LastTimeSinceTick != 0.0 && InCurrentTime - LastTimeSinceTick > MinimumTimeToSwitchFocus)
 	{
-		ObjectsTreeView->ClearSearchBox();
 		ObjectsTreeView->Focus();
 	}
 
@@ -252,7 +251,6 @@ FReply SRCPanelFunctionPicker::OnFocusReceived(const FGeometry& MyGeometry, cons
 {
 	if (InFocusEvent.GetCause() == EFocusCause::Navigation)
 	{
-		ObjectsTreeView->ClearSearchBox();
 		ObjectsTreeView->Focus();
 		return FReply::Handled();
 	}
diff --git a/Engine/Plugins/VirtualProduction/RemoteControl/Source/RemoteControlUI/Private/UI/SSearchableTreeView.h b/Engine/Plugins/VirtualProduction/RemoteControl/Source/RemoteControlUI/Private/UI/SSearchableTreeView.h
--- a/Engine/Plugins/VirtualProduction/RemoteControl/Source/RemoteControlUI/Private/UI/SSearchableTreeView.h
+++ b/Engine/Plugins/VirtualProduction/RemoteControl/Source/RemoteControlUI/Private/UI/SSearchableTreeView.h
@@ -73,6 +73,7 @@ public:
 				TreeView.ToSharedRef()
 			]
 		];
+		SearchBox->SetSelectAllTextWhenFocused(false);
 	}
 
 	virtual void Tick(const FGeometry&, const double, const float DeltaTime) override
@@ -106,6 +107,7 @@ public:
 	void ClearSearchBox()
 	{
 		SearchBox->SetText(FText());
+		Refresh();
 	}
 private:
 
