Merging CL 17890913 from Release 5.0

diff --git a/Engine/Plugins/Media/AudioCaptureTimecodeProvider/Source/AudioCaptureTimecodeProvider/Private/AudioCaptureTimecodeProvider.cpp b/Engine/Plugins/Media/AudioCaptureTimecodeProvider/Source/AudioCaptureTimecodeProvider/Private/AudioCaptureTimecodeProvider.cpp
--- a/Engine/Plugins/Media/AudioCaptureTimecodeProvider/Source/AudioCaptureTimecodeProvider/Private/AudioCaptureTimecodeProvider.cpp
+++ b/Engine/Plugins/Media/AudioCaptureTimecodeProvider/Source/AudioCaptureTimecodeProvider/Private/AudioCaptureTimecodeProvider.cpp
@@ -5,6 +5,7 @@
 
 #include "AudioCapture.h"
 #include "HAL/CriticalSection.h"
+#include "Misc/CommandLine.h"
 #include "LinearTimecodeDecoder.h"
 #include "Stats/StatsMisc.h"
 
@@ -238,6 +239,12 @@ FFrameRate UAudioCaptureTimecodeProvider::GetFrameRateInternal() const
 
 bool UAudioCaptureTimecodeProvider::Initialize(class UEngine* InEngine)
 {
+	if (IsRunningCommandlet() && !FParse::Param(FCommandLine::Get(), TEXT("useaudiocapturetimecode")))
+	{
+		UE_LOG(LogAudioCaptureTimecodeProvider, Display, TEXT("Audio Capture Timecode Provided initilization was skipped because UE is running in a commandlet. Use -useaudiocapturetimecode to force the initialization."));		
+		return false;
+	}
+	
 	check(Implementation == nullptr);
 	delete Implementation; // in case
 
@@ -255,9 +262,12 @@ bool UAudioCaptureTimecodeProvider::Initialize(class UEngine* InEngine)
 
 void UAudioCaptureTimecodeProvider::Shutdown(class UEngine* InEngine)
 {
-	SynchronizationState = ETimecodeProviderSynchronizationState::Closed;
-	delete Implementation;
-	Implementation = nullptr;
+	if (Implementation)
+	{
+		SynchronizationState = ETimecodeProviderSynchronizationState::Closed;
+		delete Implementation;
+		Implementation = nullptr;
+	}
 }
 
 void UAudioCaptureTimecodeProvider::BeginDestroy()
diff --git a/Engine/Plugins/VirtualProduction/RemoteControl/Source/RemoteControl/Private/RemoteControlBinding.cpp b/Engine/Plugins/VirtualProduction/RemoteControl/Source/RemoteControl/Private/RemoteControlBinding.cpp
--- a/Engine/Plugins/VirtualProduction/RemoteControl/Source/RemoteControl/Private/RemoteControlBinding.cpp
+++ b/Engine/Plugins/VirtualProduction/RemoteControl/Source/RemoteControl/Private/RemoteControlBinding.cpp
@@ -13,6 +13,7 @@
 #include "Components/ActorComponent.h"
 #include "Editor.h"
 #include "Editor/UnrealEd/Public/Editor.h"
+#include "Misc/App.h"
 #endif
 
 namespace
@@ -184,7 +185,7 @@ UWorld* URemoteControlLevelDependantBinding::GetCurrentWorld() const
 	UWorld* World = nullptr;
 
 #if WITH_EDITOR
-	if (GEditor)
+	if (GEditor && FApp::CanEverRender())
 	{
 		World = GEditor->GetEditorWorldContext(false).World();
 	}
@@ -195,12 +196,15 @@ UWorld* URemoteControlLevelDependantBinding::GetCurrentWorld() const
 		return World;
 	}
 
-	for (const FWorldContext& WorldContext : GEngine->GetWorldContexts())
+	if (GEngine)
 	{
-		if (WorldContext.WorldType == EWorldType::Game)
+		for (const FWorldContext& WorldContext : GEngine->GetWorldContexts())
 		{
-			World = WorldContext.World();
-			break;
+			if (WorldContext.WorldType == EWorldType::Game)
+			{
+				World = WorldContext.World();
+				break;
+			}
 		}
 	}
 		
diff --git a/Engine/Plugins/VirtualProduction/RemoteControl/Source/RemoteControl/Private/RemoteControlPreset.cpp b/Engine/Plugins/VirtualProduction/RemoteControl/Source/RemoteControl/Private/RemoteControlPreset.cpp
--- a/Engine/Plugins/VirtualProduction/RemoteControl/Source/RemoteControl/Private/RemoteControlPreset.cpp
+++ b/Engine/Plugins/VirtualProduction/RemoteControl/Source/RemoteControl/Private/RemoteControlPreset.cpp
@@ -16,6 +16,7 @@
 #include "RemoteControlFieldPath.h"
 #include "RemoteControlActor.h"
 #include "RemoteControlBinding.h"
+#include "RemoteControlLogger.h"
 #include "RemoteControlEntityFactory.h"
 #include "RemoteControlObjectVersion.h"
 #include "RemoteControlPresetRebindingManager.h"
@@ -739,6 +740,9 @@ TWeakPtr<FRemoteControlActor> URemoteControlPreset::ExposeActor(AActor* Actor, F
 	const TCHAR* DesiredName = Args.Label.IsEmpty() ? *Actor->GetName() : *Args.Label;
 #endif
 
+	FText LogText = FText::Format(LOCTEXT("ExposedActor", "Exposed actor ({0})"), FText::FromString(Actor->GetName()));
+    FRemoteControlLogger::Get().Log(TEXT("RemoteControlPreset"), [Text = MoveTemp(LogText)](){ return Text; });
+
 	FRemoteControlActor RCActor{this, Registry->GenerateUniqueLabel(DesiredName), { FindOrAddBinding(Actor)} };
 	return StaticCastSharedPtr<FRemoteControlActor>(Expose(MoveTemp(RCActor), FRemoteControlActor::StaticStruct(), Args.GroupId));
 }
@@ -794,6 +798,9 @@ TWeakPtr<FRemoteControlProperty> URemoteControlPreset::ExposeProperty(UObject* O
 		CreatePropertyWatcher(RCPropertyPtr);
 	}
 
+	FText LogText = FText::Format(LOCTEXT("ExposedProperty", "Exposed property ({0}) on object {1}"), FText::FromString(RCPropertyPtr->FieldPathInfo.ToString()), FText::FromString(Object->GetPathName()));
+	FRemoteControlLogger::Get().Log(TEXT("RemoteControlPreset"), [Text = MoveTemp(LogText)](){ return Text; });
+
 	return RCPropertyPtr;
 }
 
@@ -822,6 +829,9 @@ TWeakPtr<FRemoteControlFunction> URemoteControlPreset::ExposeFunction(UObject* O
 	TSharedPtr<FRemoteControlFunction> RCFunctionPtr = StaticCastSharedPtr<FRemoteControlFunction>(Expose(MoveTemp(RCFunction), FRemoteControlFunction::StaticStruct(), Args.GroupId));
 
 	RegisterOnCompileEvent(RCFunctionPtr);
+
+	FText LogText = FText::Format(LOCTEXT("ExposedFunction", "Exposed function ({0}) on object {1}"), FText::FromString(Function->GetPathName()), FText::FromString(Object->GetPathName()));
+	FRemoteControlLogger::Get().Log(TEXT("RemoteControlPreset"), [Text = MoveTemp(LogText)](){ return Text; });
 	
 	return RCFunctionPtr;
 }
diff --git a/Engine/Plugins/VirtualProduction/RemoteControl/Source/RemoteControlUI/Private/UI/RemoteControlPanelStyle.cpp b/Engine/Plugins/VirtualProduction/RemoteControl/Source/RemoteControlUI/Private/UI/RemoteControlPanelStyle.cpp
--- a/Engine/Plugins/VirtualProduction/RemoteControl/Source/RemoteControlUI/Private/UI/RemoteControlPanelStyle.cpp
+++ b/Engine/Plugins/VirtualProduction/RemoteControl/Source/RemoteControlUI/Private/UI/RemoteControlPanelStyle.cpp
@@ -82,18 +82,23 @@ void FRemoteControlPanelStyle::Initialize()
 	NewSelectionColor.R *= 1.8;
 	NewSelectionColor.G *= 1.8;
 	NewSelectionColor.B *= 1.8;
-	NewSelectionColor.A = 0.3;
+	NewSelectionColor.A = 0.4;
+
+
 	StyleSet->Set("RemoteControlPanel.GroupRowSelected", new BOX_BRUSH("Common/GroupBorderLight", FMargin(4.0f / 16.0f), NewSelectionColor));
+	StyleSet->Set("RemoteControlPanel.TransparentBorder", new FSlateColorBrush(FColor::Transparent));
 
 	FTableRowStyle GroupRowStyle = AppStyle.GetWidgetStyle<FTableRowStyle>("TableView.Row");
+	GroupRowStyle.SetSelectorFocusedBrush(*FRemoteControlPanelStyle::Get()->GetBrush("RemoteControlPanel.GroupRowSelected"));
+
 	GroupRowStyle.SetEvenRowBackgroundBrush(*FRemoteControlPanelStyle::Get()->GetBrush("RemoteControlPanel.GroupBorder"));
 	GroupRowStyle.SetOddRowBackgroundBrush(*FRemoteControlPanelStyle::Get()->GetBrush("RemoteControlPanel.GroupBorder"));
 	GroupRowStyle.SetEvenRowBackgroundHoveredBrush(*FRemoteControlPanelStyle::Get()->GetBrush("RemoteControlPanel.GroupBorder"));
 	GroupRowStyle.SetOddRowBackgroundHoveredBrush(*FRemoteControlPanelStyle::Get()->GetBrush("RemoteControlPanel.GroupBorder"));
-	GroupRowStyle.SetActiveHoveredBrush(*FRemoteControlPanelStyle::Get()->GetBrush("RemoteControlPanel.GroupBorder"));
-	GroupRowStyle.SetInactiveHoveredBrush(*FRemoteControlPanelStyle::Get()->GetBrush("RemoteControlPanel.GroupBorder"));
-	GroupRowStyle.SetActiveBrush(*FRemoteControlPanelStyle::Get()->GetBrush("RemoteControlPanel.GroupBorder"));
-	GroupRowStyle.SetInactiveBrush(*FRemoteControlPanelStyle::Get()->GetBrush("RemoteControlPanel.GroupBorder"));
+	GroupRowStyle.SetActiveHoveredBrush(*FRemoteControlPanelStyle::Get()->GetBrush("RemoteControlPanel.GroupRowSelected"));
+	GroupRowStyle.SetInactiveHoveredBrush(*FRemoteControlPanelStyle::Get()->GetBrush("RemoteControlPanel.GroupRowSelected"));
+	GroupRowStyle.SetActiveBrush(*FRemoteControlPanelStyle::Get()->GetBrush("RemoteControlPanel.GroupRowSelected"));
+	GroupRowStyle.SetInactiveBrush(*FRemoteControlPanelStyle::Get()->GetBrush("RemoteControlPanel.GroupRowSelected"));
 	
 	StyleSet->Set("RemoteControlPanel.GroupRow", GroupRowStyle);
 	
diff --git a/Engine/Plugins/VirtualProduction/RemoteControl/Source/RemoteControlUI/Private/UI/SRCPanelExposedEntitiesList.cpp b/Engine/Plugins/VirtualProduction/RemoteControl/Source/RemoteControlUI/Private/UI/SRCPanelExposedEntitiesList.cpp
--- a/Engine/Plugins/VirtualProduction/RemoteControl/Source/RemoteControlUI/Private/UI/SRCPanelExposedEntitiesList.cpp
+++ b/Engine/Plugins/VirtualProduction/RemoteControl/Source/RemoteControlUI/Private/UI/SRCPanelExposedEntitiesList.cpp
@@ -524,6 +524,8 @@ void SRCPanelExposedEntitiesList::OnGroupAdded(const FRemoteControlPresetGroup&
 			FieldGroup->GetNodes().Add(*Widget);
 		}
 	}
+
+	FieldGroup->EnterRenameMode();
 	TreeView->SetSelection(FieldGroup);
 	TreeView->ScrollToBottom();
 	TreeView->RequestListRefresh();
diff --git a/Engine/Plugins/VirtualProduction/RemoteControl/Source/RemoteControlUI/Private/UI/SRCPanelFieldGroup.cpp b/Engine/Plugins/VirtualProduction/RemoteControl/Source/RemoteControlUI/Private/UI/SRCPanelFieldGroup.cpp
--- a/Engine/Plugins/VirtualProduction/RemoteControl/Source/RemoteControlUI/Private/UI/SRCPanelFieldGroup.cpp
+++ b/Engine/Plugins/VirtualProduction/RemoteControl/Source/RemoteControlUI/Private/UI/SRCPanelFieldGroup.cpp
@@ -150,6 +150,11 @@ void SRCPanelGroup::SetName(FName InName)
 	NameTextBox->SetText(FText::FromName(Name));
 }
 
+void SRCPanelGroup::EnterRenameMode()
+{
+	bNeedsRename = true;
+}
+
 void SRCPanelGroup::GetNodeChildren(TArray<TSharedPtr<SRCPanelTreeNode>>& OutChildren) const
 {
 	OutChildren.Append(Nodes);
@@ -231,7 +236,7 @@ FSlateColor SRCPanelGroup::GetGroupNameTextColor() const
 
 const FSlateBrush* SRCPanelGroup::GetBorderImage() const
 {
-	return FRemoteControlPanelStyle::Get()->GetBrush("RemoteControlPanel.GroupBorder");
+	return FRemoteControlPanelStyle::Get()->GetBrush("RemoteControlPanel.TransparentBorder");
 }
 
 EVisibility SRCPanelGroup::GetVisibilityAccordingToEditMode(EVisibility NonEditModeVisibility) const
diff --git a/Engine/Plugins/VirtualProduction/RemoteControl/Source/RemoteControlUI/Private/UI/SRCPanelFieldGroup.h b/Engine/Plugins/VirtualProduction/RemoteControl/Source/RemoteControlUI/Private/UI/SRCPanelFieldGroup.h
--- a/Engine/Plugins/VirtualProduction/RemoteControl/Source/RemoteControlUI/Private/UI/SRCPanelFieldGroup.h
+++ b/Engine/Plugins/VirtualProduction/RemoteControl/Source/RemoteControlUI/Private/UI/SRCPanelFieldGroup.h
@@ -56,6 +56,9 @@ public:
 	/** Get raw access to this group's child nodes. */
 	TArray<TSharedPtr<SRCPanelTreeNode>>& GetNodes() { return Nodes; }
 
+	/** Make the group name's text box editable. */
+	void EnterRenameMode();
+
 	//~ SRCPanelTreeNode Interface
 	virtual void GetNodeChildren(TArray<TSharedPtr<SRCPanelTreeNode>>& OutChildren) const override;
 	virtual FGuid GetRCId() const override;
diff --git a/Engine/Plugins/VirtualProduction/RemoteControl/Source/WebRemoteControl/Private/Serialization/RCJsonStructDeserializerBackend.cpp b/Engine/Plugins/VirtualProduction/RemoteControl/Source/WebRemoteControl/Private/Serialization/RCJsonStructDeserializerBackend.cpp
--- a/Engine/Plugins/VirtualProduction/RemoteControl/Source/WebRemoteControl/Private/Serialization/RCJsonStructDeserializerBackend.cpp
+++ b/Engine/Plugins/VirtualProduction/RemoteControl/Source/WebRemoteControl/Private/Serialization/RCJsonStructDeserializerBackend.cpp
@@ -5,6 +5,8 @@
 #include "UObject/EnumProperty.h"
 #include "UObject/UnrealType.h"
 
+#include <limits>
+
 namespace
 {
 	static void* GetPropertyValuePtr(FProperty* Property, FProperty* Outer, void* Data, int32 ArrayIndex)
@@ -23,10 +25,13 @@ namespace
 
 			return ArrayHelper.GetRawPtr(Index);
 		}
-		else if (Property->IsA<FMapProperty>() || Property->IsA<FSetProperty>())
+		else if (Outer && (Outer->IsA<FMapProperty>() || Outer->IsA<FSetProperty>()))
 		{
-			ensureAlwaysMsgf(false, TEXT("Set or Map property containing enums is not supported by Remote Control."));
-			return nullptr;
+			if (Property->IsA<FEnumProperty>())
+			{
+				ensureAlwaysMsgf(false, TEXT("Set or Map property containing enums is not supported by Remote Control."));
+				return nullptr;
+			}
 		}
 
 		if (ArrayIndex >= Property->ArrayDim)
@@ -128,6 +133,26 @@ bool FRCJsonStructDeserializerBackend::ReadProperty(FProperty* Property, FProper
 			return bReadResult;
 		}
 	}
+	else if (GetLastNotation() == EJsonNotation::Number)
+	{
+		if (FFloatProperty* FloatProperty = CastField<FFloatProperty>(Property))
+		{
+			if (void* Ptr = GetPropertyValuePtr(Property, Outer, Data, ArrayIndex))
+			{
+				const double& DoubleValue = GetReader()->GetValueAsNumber();
+				if (DoubleValue > std::numeric_limits<float>::max())
+				{
+					FloatProperty->SetPropertyValue(Ptr, std::numeric_limits<float>::max());
+					return true;
+				}
+			}
+			else
+			{
+				return false;
+			}
+		}
+	}
+
 
 	return FJsonStructDeserializerBackend::ReadProperty(Property, Outer, Data, ArrayIndex);
 }
diff --git a/Engine/Plugins/VirtualProduction/RemoteControl/Source/WebRemoteControl/Private/Serialization/RCJsonStructSerializerBackend.cpp b/Engine/Plugins/VirtualProduction/RemoteControl/Source/WebRemoteControl/Private/Serialization/RCJsonStructSerializerBackend.cpp
--- a/Engine/Plugins/VirtualProduction/RemoteControl/Source/WebRemoteControl/Private/Serialization/RCJsonStructSerializerBackend.cpp
+++ b/Engine/Plugins/VirtualProduction/RemoteControl/Source/WebRemoteControl/Private/Serialization/RCJsonStructSerializerBackend.cpp
@@ -2,6 +2,7 @@
 
 #include "RCJsonStructSerializerBackend.h"
 #include "UObject/EnumProperty.h"
+#include <limits>
 
 void FRCJsonStructSerializerBackend::WriteProperty(const FStructSerializerState& State, int32 ArrayIndex /*= 0*/)
 {
@@ -17,7 +18,7 @@ void FRCJsonStructSerializerBackend::WriteProperty(const FStructSerializerState&
 			return;
 		}
 	}
-	else if(State.FieldType == FEnumProperty::StaticClass())
+	else if (State.FieldType == FEnumProperty::StaticClass())
 	{
 		FEnumProperty* EnumProperty = CastFieldChecked<FEnumProperty>(State.ValueProperty);
 		const void* PropertyValuePtr = EnumProperty->ContainerPtrToValuePtr<void>(State.ValueData, ArrayIndex);
@@ -26,6 +27,20 @@ void FRCJsonStructSerializerBackend::WriteProperty(const FStructSerializerState&
 		WritePropertyValue(State, DisplayName.ToString());
 		return;
 	}
+	else if (State.FieldType == FFloatProperty::StaticClass())
+	{
+		FFloatProperty* FloatProperty = CastFieldChecked<FFloatProperty>(State.ValueProperty);
+		const float Value = FloatProperty->GetPropertyValue_InContainer(State.ValueData, ArrayIndex);
+		if (Value >= std::numeric_limits<float>::max())
+		{
+			WritePropertyValue(State, std::numeric_limits<float>::max());
+		}
+		else
+		{
+			WritePropertyValue(State, Value);
+		}
+		return;
+	}
 
 	FJsonStructSerializerBackend::WriteProperty(State, ArrayIndex);
 }
diff --git a/Engine/Plugins/VirtualProduction/RemoteControl/Source/WebRemoteControl/Private/WebRemoteControl.cpp b/Engine/Plugins/VirtualProduction/RemoteControl/Source/WebRemoteControl/Private/WebRemoteControl.cpp
--- a/Engine/Plugins/VirtualProduction/RemoteControl/Source/WebRemoteControl/Private/WebRemoteControl.cpp
+++ b/Engine/Plugins/VirtualProduction/RemoteControl/Source/WebRemoteControl/Private/WebRemoteControl.cpp
@@ -194,10 +194,7 @@ namespace WebRemoteControl
 		bool bIsEditor = false;
 
 #if WITH_EDITOR
-		if (GIsEditor)
-		{
-			bIsEditor = true;
-		}
+		bIsEditor = GIsEditor;
 #endif
 
 		// By default, web remote control is disabled in -game and packaged game.
@@ -1041,7 +1038,7 @@ bool FWebRemoteControlModule::HandlePresetSetPropertyRoute(const FHttpServerRequ
 			// Don't manually trigger a property change modification if this request gets converted to a function call.
 			if (ObjectRef.IsValid() && !RemoteControlPropertyUtilities::FindSetterFunction(ObjectRef.Property.Get(), ObjectRef.Object->GetClass()))
 			{
-				WebSocketHandler->NotifyPropertyChangedRemotely(ActingClientId, Preset->GetFName(), RemoteControlProperty->GetId());
+				WebSocketHandler->NotifyPropertyChangedRemotely(ActingClientId, Preset->GetPresetId(), RemoteControlProperty->GetId());
 			}
 		}
 		if (SetPropertyRequest.ResetToDefault)
diff --git a/Engine/Plugins/VirtualProduction/RemoteControl/Source/WebRemoteControl/Private/WebSocketMessageHandler.cpp b/Engine/Plugins/VirtualProduction/RemoteControl/Source/WebRemoteControl/Private/WebSocketMessageHandler.cpp
--- a/Engine/Plugins/VirtualProduction/RemoteControl/Source/WebRemoteControl/Private/WebSocketMessageHandler.cpp
+++ b/Engine/Plugins/VirtualProduction/RemoteControl/Source/WebRemoteControl/Private/WebSocketMessageHandler.cpp
@@ -247,9 +247,9 @@ void FWebSocketMessageHandler::UnregisterRoutes(FWebRemoteControlModule* WebRemo
 	}
 }
 
-void FWebSocketMessageHandler::NotifyPropertyChangedRemotely(const FGuid& OriginClientId, FName PresetName, const FGuid& ExposedPropertyId)
+void FWebSocketMessageHandler::NotifyPropertyChangedRemotely(const FGuid& OriginClientId, const FGuid& PresetId, const FGuid& ExposedPropertyId)
 {
-	if (TArray<FGuid>* SubscribedClients = WebSocketNotificationMap.Find(PresetName))
+	if (TArray<FGuid>* SubscribedClients = WebSocketNotificationMap.Find(PresetId))
 	{
 		if (SubscribedClients->Contains(OriginClientId))
 		{
@@ -262,11 +262,11 @@ void FWebSocketMessageHandler::NotifyPropertyChangedRemotely(const FGuid& Origin
 
 			if (!bIgnoreIncomingNotification)
 			{
-				PerFrameModifiedProperties.FindOrAdd(PresetName).FindOrAdd(OriginClientId).Add(ExposedPropertyId);
+				PerFrameModifiedProperties.FindOrAdd(PresetId).FindOrAdd(OriginClientId).Add(ExposedPropertyId);
 			}
 			else
 			{
-				for (TPair<FGuid, TSet<FGuid>>& Entry : PerFrameModifiedProperties.FindOrAdd(PresetName))
+				for (TPair<FGuid, TSet<FGuid>>& Entry : PerFrameModifiedProperties.FindOrAdd(PresetId))
 				{
 					if (Entry.Key != OriginClientId)
 					{
@@ -310,12 +310,12 @@ void FWebSocketMessageHandler::HandleWebSocketPresetRegister(const FRemoteContro
 
 	ClientConfigMap.FindOrAdd(WebSocketMessage.ClientId).bIgnoreRemoteChanges = Body.IgnoreRemoteChanges;
 	
-	TArray<FGuid>* ClientIds = WebSocketNotificationMap.Find(Preset->GetFName());
+	TArray<FGuid>* ClientIds = WebSocketNotificationMap.Find(Preset->GetPresetId());
 
 	// Don't register delegates for a preset more than once.
 	if (!ClientIds)
 	{
-		ClientIds = &WebSocketNotificationMap.Add(Preset->GetFName());
+		ClientIds = &WebSocketNotificationMap.Add(Preset->GetPresetId());
 
 		//Register to any useful callback for the given preset
 		Preset->OnExposedPropertiesModified().AddRaw(this, &FWebSocketMessageHandler::OnPresetExposedPropertiesModified);
@@ -355,7 +355,7 @@ void FWebSocketMessageHandler::HandleWebSocketPresetUnregister(const FRemoteCont
 
 	if (Preset)
 	{
-		if (TArray<FGuid>* RegisteredClients = WebSocketNotificationMap.Find(Preset->GetFName()))
+		if (TArray<FGuid>* RegisteredClients = WebSocketNotificationMap.Find(Preset->GetPresetId()))
 		{
 		RegisteredClients->Remove(WebSocketMessage.ClientId);
 	}
@@ -365,7 +365,7 @@ void FWebSocketMessageHandler::HandleWebSocketPresetUnregister(const FRemoteCont
 void FWebSocketMessageHandler::ProcessChangedProperties()
 {
 	//Go over each property that were changed for each preset
-	for (const TPair<FName, TMap<FGuid, TSet<FGuid>>>& Entry : PerFrameModifiedProperties)
+	for (const TPair<FGuid, TMap<FGuid, TSet<FGuid>>>& Entry : PerFrameModifiedProperties)
 	{
 		if (!ShouldProcessEventForPreset(Entry.Key) || !Entry.Value.Num())
 		{
@@ -399,7 +399,7 @@ void FWebSocketMessageHandler::ProcessChangedProperties()
 void FWebSocketMessageHandler::ProcessChangedActorProperties()
 {
 	//Go over each property that were changed for each preset
-	for (const TPair<FName, TMap<FGuid, TMap<FRemoteControlActor, TArray<FRCObjectReference>>>>& Entry : PerFrameActorPropertyChanged)
+	for (const TPair<FGuid, TMap<FGuid, TMap<FRemoteControlActor, TArray<FRCObjectReference>>>>& Entry : PerFrameActorPropertyChanged)
 	{
 		if (!ShouldProcessEventForPreset(Entry.Key) || !Entry.Value.Num())
 		{
@@ -443,7 +443,7 @@ void FWebSocketMessageHandler::OnPropertyExposed(URemoteControlPreset* Owner, co
 	}
 
 	//Cache the property field that was removed for end of frame notification
-	PerFrameAddedProperties.FindOrAdd(Owner->GetFName()).AddUnique(EntityId);
+	PerFrameAddedProperties.FindOrAdd(Owner->GetPresetId()).AddUnique(EntityId);
 }
 
 void FWebSocketMessageHandler::OnPresetExposedPropertiesModified(URemoteControlPreset* Owner, const TSet<FGuid>& ModifiedPropertyIds)
@@ -459,9 +459,9 @@ void FWebSocketMessageHandler::OnPresetExposedPropertiesModified(URemoteControlP
 	}
 
 	//Cache the property field that changed for end of frame notification
-	TMap<FGuid, TSet<FGuid>>& EventsForClient = PerFrameModifiedProperties.FindOrAdd(Owner->GetFName());
+	TMap<FGuid, TSet<FGuid>>& EventsForClient = PerFrameModifiedProperties.FindOrAdd(Owner->GetPresetId());
 	
-	if (TArray<FGuid>* SubscribedClients = WebSocketNotificationMap.Find(Owner->GetFName()))
+	if (TArray<FGuid>* SubscribedClients = WebSocketNotificationMap.Find(Owner->GetPresetId()))
 	{
 		for (const FGuid& ModifiedPropertyId : ModifiedPropertyIds)
 		{
@@ -503,8 +503,8 @@ void FWebSocketMessageHandler::OnPropertyUnexposed(URemoteControlPreset* Owner,
 	TSharedPtr<FRemoteControlEntity> Entity = Owner->GetExposedEntity(EntityId).Pin();
 	check(Entity);
 
-	//Cache the property field that was removed for end of frame notification
-	TTuple<TArray<FGuid>, TArray<FName>>& Entries = PerFrameRemovedProperties.FindOrAdd(Owner->GetFName());
+	// Cache the property field that was removed for end of frame notification
+	TTuple<TArray<FGuid>, TArray<FName>>& Entries = PerFrameRemovedProperties.FindOrAdd(Owner->GetPresetId());
 	Entries.Key.AddUnique(EntityId);
 	Entries.Value.AddUnique(Entity->GetLabel());
 }
@@ -522,7 +522,7 @@ void FWebSocketMessageHandler::OnFieldRenamed(URemoteControlPreset* Owner, FName
 	}
 
 	//Cache the field that was renamed for end of frame notification
-	PerFrameRenamedFields.FindOrAdd(Owner->GetFName()).AddUnique(TTuple<FName, FName>(OldFieldLabel, NewFieldLabel));
+	PerFrameRenamedFields.FindOrAdd(Owner->GetPresetId()).AddUnique(TTuple<FName, FName>(OldFieldLabel, NewFieldLabel));
 }
 
 void FWebSocketMessageHandler::OnMetadataModified(URemoteControlPreset* Owner)
@@ -539,7 +539,7 @@ void FWebSocketMessageHandler::OnMetadataModified(URemoteControlPreset* Owner)
 	}
 
 	//Cache the field that was renamed for end of frame notification
-	PerFrameModifiedMetadata.Add(Owner->GetFName());
+	PerFrameModifiedMetadata.Add(Owner->GetPresetId());
 }
 
 void FWebSocketMessageHandler::OnActorPropertyChanged(URemoteControlPreset* Owner, FRemoteControlActor& Actor, UObject* ModifiedObject, FProperty* ModifiedProperty)
@@ -570,10 +570,10 @@ void FWebSocketMessageHandler::OnActorPropertyChanged(URemoteControlPreset* Owne
 
 
 	//Cache the property field that changed for end of frame notification
-	TMap<FGuid, TMap<FRemoteControlActor, TArray<FRCObjectReference>>>& EventsForClient = PerFrameActorPropertyChanged.FindOrAdd(Owner->GetFName());
+	TMap<FGuid, TMap<FRemoteControlActor, TArray<FRCObjectReference>>>& EventsForClient = PerFrameActorPropertyChanged.FindOrAdd(Owner->GetPresetId());
 
 	// Dont send events to the client that triggered it.
-	if (TArray<FGuid>* SubscribedClients = WebSocketNotificationMap.Find(Owner->GetFName()))
+	if (TArray<FGuid>* SubscribedClients = WebSocketNotificationMap.Find(Owner->GetPresetId()))
 	{
 		for (const FGuid& Client : *SubscribedClients)
 		{
@@ -596,7 +596,7 @@ void FWebSocketMessageHandler::OnEntitiesModified(URemoteControlPreset* Owner, c
 	
 	TArray<uint8> Payload;
 	WebRemoteControlUtils::SerializeResponse(FRCPresetEntitiesModifiedEvent{Owner, ModifiedEntities.Array()}, Payload);
-	BroadcastToListeners(Owner->GetFName(), Payload);
+	BroadcastToListeners(Owner->GetPresetId(), Payload);
 }
 
 void FWebSocketMessageHandler::OnLayoutModified(URemoteControlPreset* Owner)
@@ -612,7 +612,7 @@ void FWebSocketMessageHandler::OnLayoutModified(URemoteControlPreset* Owner)
 	}
 
 	//Cache the field that was renamed for end of frame notification
-	PerFrameModifiedPresetLayouts.Add(Owner->GetFName());
+	PerFrameModifiedPresetLayouts.Add(Owner->GetPresetId());
 }
 
 void FWebSocketMessageHandler::OnConnectionClosedCallback(FGuid ClientId)
@@ -652,7 +652,7 @@ void FWebSocketMessageHandler::OnEndFrame()
 
 void FWebSocketMessageHandler::ProcessAddedProperties()
 {
-	for (const TPair<FName, TArray<FGuid>>& Entry : PerFrameAddedProperties)
+	for (const TPair<FGuid, TArray<FGuid>>& Entry : PerFrameAddedProperties)
 	{
 		if (Entry.Value.Num() <= 0 || !ShouldProcessEventForPreset(Entry.Key))
 		{
@@ -686,7 +686,7 @@ void FWebSocketMessageHandler::ProcessAddedProperties()
 		}
 
 		TArray<uint8> Payload;
-		WebRemoteControlUtils::SerializeResponse(FRCPresetFieldsAddedEvent{ Entry.Key, Preset->GetPresetId(), AddedPropertiesDescription }, Payload);
+		WebRemoteControlUtils::SerializeResponse(FRCPresetFieldsAddedEvent{ Preset->GetFName(), Preset->GetPresetId(), AddedPropertiesDescription }, Payload);
 		BroadcastToListeners(Entry.Key, Payload);
 	}
 
@@ -695,7 +695,7 @@ void FWebSocketMessageHandler::ProcessAddedProperties()
 
 void FWebSocketMessageHandler::ProcessRemovedProperties()
 {
-	for (const TPair<FName, TTuple<TArray<FGuid>, TArray<FName>>>& Entry : PerFrameRemovedProperties)
+	for (const TPair<FGuid, TTuple<TArray<FGuid>, TArray<FName>>>& Entry : PerFrameRemovedProperties)
 	{
 		if (Entry.Value.Key.Num() <= 0 || !ShouldProcessEventForPreset(Entry.Key))
 		{
@@ -711,7 +711,7 @@ void FWebSocketMessageHandler::ProcessRemovedProperties()
 		ensure(Entry.Value.Key.Num() == Entry.Value.Value.Num());
 		
 		TArray<uint8> Payload;
-		WebRemoteControlUtils::SerializeResponse(FRCPresetFieldsRemovedEvent{ Entry.Key, Preset->GetPresetId(), Entry.Value.Value, Entry.Value.Key }, Payload);
+		WebRemoteControlUtils::SerializeResponse(FRCPresetFieldsRemovedEvent{ Preset->GetFName(), Preset->GetPresetId(), Entry.Value.Value, Entry.Value.Key }, Payload);
 		BroadcastToListeners(Entry.Key, Payload);
 	}
 	
@@ -720,7 +720,7 @@ void FWebSocketMessageHandler::ProcessRemovedProperties()
 
 void FWebSocketMessageHandler::ProcessRenamedFields()
 {
-	for (const TPair<FName, TArray<TTuple<FName, FName>>>& Entry : PerFrameRenamedFields)
+	for (const TPair<FGuid, TArray<TTuple<FName, FName>>>& Entry : PerFrameRenamedFields)
 	{
 		if (Entry.Value.Num() <= 0 || !ShouldProcessEventForPreset(Entry.Key))
 		{
@@ -734,7 +734,7 @@ void FWebSocketMessageHandler::ProcessRenamedFields()
 		}
 
 		TArray<uint8> Payload;
-		WebRemoteControlUtils::SerializeResponse(FRCPresetFieldsRenamedEvent{Entry.Key, Preset->GetPresetId(), Entry.Value}, Payload);
+		WebRemoteControlUtils::SerializeResponse(FRCPresetFieldsRenamedEvent{Preset->GetFName(), Preset->GetPresetId(), Entry.Value}, Payload);
 		BroadcastToListeners(Entry.Key, Payload);
 	}
 
@@ -743,7 +743,7 @@ void FWebSocketMessageHandler::ProcessRenamedFields()
 
 void FWebSocketMessageHandler::ProcessModifiedMetadata()
 {
-	for (const FName& Entry : PerFrameModifiedMetadata)
+	for (const FGuid& Entry : PerFrameModifiedMetadata)
 	{
 		if (ShouldProcessEventForPreset(Entry))
 		{
@@ -761,7 +761,7 @@ void FWebSocketMessageHandler::ProcessModifiedMetadata()
 
 void FWebSocketMessageHandler::ProcessModifiedPresetLayouts()
 {
-	for (const FName& Entry : PerFrameModifiedPresetLayouts)
+	for (const FGuid& Entry : PerFrameModifiedPresetLayouts)
 	{
 		if (ShouldProcessEventForPreset(Entry))
 		{
@@ -777,18 +777,18 @@ void FWebSocketMessageHandler::ProcessModifiedPresetLayouts()
 	PerFrameModifiedPresetLayouts.Empty();
 }
 
-void FWebSocketMessageHandler::BroadcastToListeners(FName TargetPresetName, const TArray<uint8>& Payload)
+void FWebSocketMessageHandler::BroadcastToListeners(const FGuid& TargetPresetId, const TArray<uint8>& Payload)
 {
-	const TArray<FGuid>& Listeners = WebSocketNotificationMap.FindChecked(TargetPresetName);
+	const TArray<FGuid>& Listeners = WebSocketNotificationMap.FindChecked(TargetPresetId);
 	for (const FGuid& Listener : Listeners)
 	{
 		Server->Send(Listener, Payload);
 	}
 }
 
-bool FWebSocketMessageHandler::ShouldProcessEventForPreset(FName PresetName) const
+bool FWebSocketMessageHandler::ShouldProcessEventForPreset(const FGuid& PresetId) const
 {
-	return WebSocketNotificationMap.Contains(PresetName) && WebSocketNotificationMap[PresetName].Num() > 0;
+	return WebSocketNotificationMap.Contains(PresetId) && WebSocketNotificationMap[PresetId].Num() > 0;
 }
 
 bool FWebSocketMessageHandler::WritePropertyChangeEventPayload(URemoteControlPreset* InPreset, const TSet<FGuid>& InModifiedPropertyIds, TArray<uint8>& OutBuffer)
diff --git a/Engine/Plugins/VirtualProduction/RemoteControl/Source/WebRemoteControl/Private/WebSocketMessageHandler.h b/Engine/Plugins/VirtualProduction/RemoteControl/Source/WebRemoteControl/Private/WebSocketMessageHandler.h
--- a/Engine/Plugins/VirtualProduction/RemoteControl/Source/WebRemoteControl/Private/WebSocketMessageHandler.h
+++ b/Engine/Plugins/VirtualProduction/RemoteControl/Source/WebRemoteControl/Private/WebSocketMessageHandler.h
@@ -32,7 +32,7 @@ public:
 	void UnregisterRoutes(FWebRemoteControlModule* WebRemoteControl);
 
 	/** Notify that a property was modified by a web client. */
-	void NotifyPropertyChangedRemotely(const FGuid& OriginClientId, FName PresetName, const FGuid& ExposedPropertyId);
+	void NotifyPropertyChangedRemotely(const FGuid& OriginClientId, const FGuid& PresetId, const FGuid& ExposedPropertyId);
 
 private:
 	
@@ -83,12 +83,12 @@ private:
 	 * Send a payload to all clients bound to a certain preset.
 	 * @note: TargetPresetName must be in the WebSocketNotificationMap.
 	 */
-	void BroadcastToListeners(FName TargetPresetName, const TArray<uint8>& Payload);
+	void BroadcastToListeners(const FGuid& TargetPresetId, const TArray<uint8>& Payload);
 
 	/**
 	 * Returns whether an event targeting a particular preset should be processed.
 	 */
-	bool ShouldProcessEventForPreset(FName PresetName) const;
+	bool ShouldProcessEventForPreset(const FGuid& PresetId) const;
 
 	/**
 	 * Write the provided list of events to a buffer.
@@ -108,7 +108,7 @@ private:
 	TArray<TUniquePtr<FRemoteControlWebsocketRoute>> Routes;
 
 	/** All websockets connections associated to a preset notifications */
-	TMap<FName, TArray<FGuid>> WebSocketNotificationMap;
+	TMap<FGuid, TArray<FGuid>> WebSocketNotificationMap;
 
 	/** Configuration for a given client related to how events should be handled. */
 	struct FRCClientConfig
@@ -121,7 +121,7 @@ private:
 	TMap<FGuid, FRCClientConfig> ClientConfigMap;
 
 	/** Properties that changed for a frame, per preset.  */
-	TMap<FName, TMap<FGuid, TSet<FGuid>>> PerFrameModifiedProperties;
+	TMap<FGuid, TMap<FGuid, TSet<FGuid>>> PerFrameModifiedProperties;
 
 	/** 
 	 * List of properties modified remotely this frame, used to not trigger a 
@@ -130,22 +130,22 @@ private:
 	TSet<FGuid> PropertiesManuallyNotifiedThisFrame;
 
 	/** Properties that changed on an exposed actor for a given client, for a frame, per preset.  */
-	TMap<FName, TMap<FGuid, TMap<FRemoteControlActor, TArray<FRCObjectReference>>>> PerFrameActorPropertyChanged;
+	TMap<FGuid, TMap<FGuid, TMap<FRemoteControlActor, TArray<FRCObjectReference>>>> PerFrameActorPropertyChanged;
 
 	/** Properties that were exposed for a frame, per preset */
-	TMap<FName, TArray<FGuid>> PerFrameAddedProperties;
+	TMap<FGuid, TArray<FGuid>> PerFrameAddedProperties;
 
 	/** Properties that were unexposed for a frame, per preset */
-	TMap<FName, TTuple<TArray<FGuid>, TArray<FName>>> PerFrameRemovedProperties;
+	TMap<FGuid, TTuple<TArray<FGuid>, TArray<FName>>> PerFrameRemovedProperties;
 
 	/** Fields that were renamed for a frame, per preset */
-	TMap<FName, TArray<TTuple<FName, FName>>> PerFrameRenamedFields;
+	TMap<FGuid, TArray<TTuple<FName, FName>>> PerFrameRenamedFields;
 
 	/** Presets that had their metadata modified for a frame */
-	TSet<FName> PerFrameModifiedMetadata;
+	TSet<FGuid> PerFrameModifiedMetadata;
 
 	/** Presets that had their layout modified for a frame. */
-	TSet<FName> PerFrameModifiedPresetLayouts;
+	TSet<FGuid> PerFrameModifiedPresetLayouts;
 	
 	/** Holds the ID of the client currently making a request. Used to prevent sending back notifications to it. */
 	const FGuid& ActingClientId;
diff --git a/Engine/Plugins/VirtualProduction/RemoteControlWebInterface/Source/RemoteControlWebInterface/Private/RCWebInterface.cpp b/Engine/Plugins/VirtualProduction/RemoteControlWebInterface/Source/RemoteControlWebInterface/Private/RCWebInterface.cpp
--- a/Engine/Plugins/VirtualProduction/RemoteControlWebInterface/Source/RemoteControlWebInterface/Private/RCWebInterface.cpp
+++ b/Engine/Plugins/VirtualProduction/RemoteControlWebInterface/Source/RemoteControlWebInterface/Private/RCWebInterface.cpp
@@ -24,10 +24,7 @@ namespace RCWebInterface
 		bool bIsEditor = false;
 
 #if WITH_EDITOR
-		if (GEditor)
-		{
-			bIsEditor = true;
-		}
+		bIsEditor = GIsEditor;
 #endif
 
 		// By default, remote control web interface is disabled in -game and packaged game.
