Fix CIS warning in remote control

diff --git a/Engine/Plugins/VirtualProduction/RemoteControl/Source/RemoteControlUI/Private/UI/SRCPanelDragHandle.h b/Engine/Plugins/VirtualProduction/RemoteControl/Source/RemoteControlUI/Private/UI/SRCPanelDragHandle.h
--- a/Engine/Plugins/VirtualProduction/RemoteControl/Source/RemoteControlUI/Private/UI/SRCPanelDragHandle.h
+++ b/Engine/Plugins/VirtualProduction/RemoteControl/Source/RemoteControlUI/Private/UI/SRCPanelDragHandle.h
@@ -66,7 +66,9 @@ public:
 	{
 		if (TSharedPtr<WidgetType> WidgetPtr = Widget.Pin())
 		{
-			return MakeShared<DragDropType>(Widget.Pin(), Id);
+			TSharedPtr<FDecoratedDragDropOp> DragDropOperation = MakeShared<DragDropType>(Widget.Pin(), Id);
+			DragDropOperation->Construct();
+			return DragDropOperation;
 		}
 		return nullptr;
 	}
diff --git a/Engine/Plugins/VirtualProduction/RemoteControl/Source/RemoteControlUI/Private/UI/SRCPanelExposedEntity.h b/Engine/Plugins/VirtualProduction/RemoteControl/Source/RemoteControlUI/Private/UI/SRCPanelExposedEntity.h
--- a/Engine/Plugins/VirtualProduction/RemoteControl/Source/RemoteControlUI/Private/UI/SRCPanelExposedEntity.h
+++ b/Engine/Plugins/VirtualProduction/RemoteControl/Source/RemoteControlUI/Private/UI/SRCPanelExposedEntity.h
@@ -87,8 +87,6 @@ public:
 			[
 				InWidget.ToSharedRef()
 			];
-
-		Construct();
 	}
 
 	/** Get the ID of the represented entity or group. */
diff --git a/Engine/Plugins/VirtualProduction/RemoteControl/Source/RemoteControlUI/Private/UI/SRCPanelFieldGroup.h b/Engine/Plugins/VirtualProduction/RemoteControl/Source/RemoteControlUI/Private/UI/SRCPanelFieldGroup.h
--- a/Engine/Plugins/VirtualProduction/RemoteControl/Source/RemoteControlUI/Private/UI/SRCPanelFieldGroup.h
+++ b/Engine/Plugins/VirtualProduction/RemoteControl/Source/RemoteControlUI/Private/UI/SRCPanelFieldGroup.h
@@ -124,8 +124,6 @@ public:
 			[
 				InWidget.ToSharedRef()
 			];
-
-		Construct();
 	}
 
 	FGuid GetGroupId() const
