Merge  change from partner stream (17724474)

diff --git a/Engine/Plugins/VirtualProduction/RemoteControl/Source/RemoteControl/Private/RemoteControlModule.cpp b/Engine/Plugins/VirtualProduction/RemoteControl/Source/RemoteControl/Private/RemoteControlModule.cpp
--- a/Engine/Plugins/VirtualProduction/RemoteControl/Source/RemoteControl/Private/RemoteControlModule.cpp
+++ b/Engine/Plugins/VirtualProduction/RemoteControl/Source/RemoteControl/Private/RemoteControlModule.cpp
@@ -345,8 +345,9 @@ void FRemoteControlModule::StartupModule()
 
 #if WITH_EDITOR
 	FCoreDelegates::OnPostEngineInit.AddRaw(this, &FRemoteControlModule::HandleEnginePostInit);
+	FCoreUObjectDelegates::PreLoadMap.AddRaw(this, &FRemoteControlModule::HandleMapPreLoad);
 #endif
-
+	
 	// Register Property Factories
 	RegisterEntityFactory(FRemoteControlInstanceMaterial::StaticStruct()->GetFName(), FRemoteControlInstanceMaterialFactory::MakeInstance());
 }
@@ -354,6 +355,7 @@ void FRemoteControlModule::StartupModule()
 void FRemoteControlModule::ShutdownModule()
 {
 #if WITH_EDITOR
+	FCoreUObjectDelegates::PreLoadMap.RemoveAll(this);
 	UnregisterEditorDelegates();
 	FCoreDelegates::OnPostEngineInit.RemoveAll(this);
 #endif
@@ -1292,6 +1294,12 @@ void FRemoteControlModule::HandleEnginePostInit()
 	RegisterEditorDelegates();
 }
 
+void FRemoteControlModule::HandleMapPreLoad(const FString& MapName)
+{
+	constexpr bool bForceFinalizeChange = true;
+	TestOrFinalizeOngoingChange(bForceFinalizeChange);
+}
+	
 void FRemoteControlModule::RegisterEditorDelegates()
 {
 	if (GEditor)
@@ -1307,9 +1315,7 @@ void FRemoteControlModule::UnregisterEditorDelegates()
 		GEditor->GetTimerManager()->ClearTimer(OngoingChangeTimer);
 	}
 }
-#endif
 
-#if WITH_EDITOR
 FRemoteControlModule::FOngoingChange::FOngoingChange(FRCObjectReference InReference)
 {
 	Reference.Set<FRCObjectReference>(MoveTemp(InReference));
diff --git a/Engine/Plugins/VirtualProduction/RemoteControl/Source/RemoteControl/Private/RemoteControlModule.h b/Engine/Plugins/VirtualProduction/RemoteControl/Source/RemoteControl/Private/RemoteControlModule.h
--- a/Engine/Plugins/VirtualProduction/RemoteControl/Source/RemoteControl/Private/RemoteControlModule.h
+++ b/Engine/Plugins/VirtualProduction/RemoteControl/Source/RemoteControl/Private/RemoteControlModule.h
@@ -64,6 +64,9 @@ private:
 	/** Finalize an ongoing change, triggering post edit change on the tracked object. */
 	void TestOrFinalizeOngoingChange(bool bForceEndChange = false);
 
+	// End ongoing change on map preload.
+	void HandleMapPreLoad(const FString& MapName);
+	
 	/** Callback to handle registering delegates once the engine has finished its startup. */
 	void HandleEnginePostInit();
 
