Bring fix from partner stream (Fix some bindings getting invalidated when changing sublevels)

diff --git a/Engine/Plugins/VirtualProduction/RemoteControl/Source/RemoteControl/Private/RemoteControlBinding.cpp b/Engine/Plugins/VirtualProduction/RemoteControl/Source/RemoteControl/Private/RemoteControlBinding.cpp
--- a/Engine/Plugins/VirtualProduction/RemoteControl/Source/RemoteControl/Private/RemoteControlBinding.cpp
+++ b/Engine/Plugins/VirtualProduction/RemoteControl/Source/RemoteControl/Private/RemoteControlBinding.cpp
@@ -83,12 +83,26 @@ bool URemoteControlLevelIndependantBinding::IsBound(const TSoftObjectPtr<UObject
 	return BoundObject == Object;
 }
 
+bool URemoteControlLevelIndependantBinding::PruneDeletedObjects()
+{
+	if (!BoundObject.IsValid())
+	{
+		Modify();
+		BoundObject.Reset();
+		return true;
+	}
+
+	return false;
+}
+
 void URemoteControlLevelDependantBinding::SetBoundObject(const TSoftObjectPtr<UObject>& InObject)
 {
 	if (ensure(InObject))
 	{
 		UObject* EditorObject = FindObjectInCounterpartWorld(InObject.Get(), ECounterpartWorldTarget::Editor);
-		BoundObjectMap.FindOrAdd(EditorObject->GetTypedOuter<ULevel>()) = EditorObject;
+		ULevel* OuterLevel = EditorObject->GetTypedOuter<ULevel>();
+		BoundObjectMap.FindOrAdd(OuterLevel) = EditorObject;
+		SubLevelSelectionMap.FindOrAdd(OuterLevel->GetWorld()) = OuterLevel;
 		
 		Name = EditorObject->GetName();
 	}
@@ -128,6 +142,11 @@ void URemoteControlLevelDependantBinding::UnbindObject(const TSoftObjectPtr<UObj
 	{
 		if (It.Value() == InBoundObject)
 		{
+			if (InBoundObject)
+			{
+				SubLevelSelectionMap.Remove(InBoundObject->GetWorld());
+			}
+
 			It.RemoveCurrent();
 		}
 	}
@@ -136,7 +155,13 @@ void URemoteControlLevelDependantBinding::UnbindObject(const TSoftObjectPtr<UObj
 UObject* URemoteControlLevelDependantBinding::Resolve() const
 {
 	// Find the object in PIE if possible
-	UObject* Object = FindObjectFromCurrentWorld().Get();
+	UObject* Object = ResolveForCurrentWorld().Get();
+
+	if (Object)
+	{
+		LevelWithLastSuccessfulResolve = Object->GetTypedOuter<ULevel>();
+	}
+
 	return FindObjectInCounterpartWorld(Object, ECounterpartWorldTarget::PIE);
 }
 
@@ -158,21 +183,46 @@ bool URemoteControlLevelDependantBinding::IsBound(const TSoftObjectPtr<UObject>&
 	return false;
 }
 
-TSoftObjectPtr<UObject> URemoteControlLevelDependantBinding::FindObjectFromCurrentWorld() const
+bool URemoteControlLevelDependantBinding::PruneDeletedObjects()
+{
+	if (!ResolveForCurrentWorld())
+	{
+		if (UWorld* World = GetCurrentWorld())
+		{
+			if (TSoftObjectPtr<ULevel> LastLevelForBinding = SubLevelSelectionMap.FindRef(World))
+			{
+				if (!BoundObjectMap.FindRef(LastLevelForBinding))
+				{
+					Modify();
+					BoundObjectMap.Remove(LastLevelForBinding);
+					SubLevelSelectionMap.Remove(World);
+					return true;
+				}
+			}
+		}
+	}
+
+	return false;
+}
+
+TSoftObjectPtr<UObject> URemoteControlLevelDependantBinding::ResolveForCurrentWorld() const
 {
-	constexpr bool bForResolving = true;
 	if (UWorld* World = GetCurrentWorld())
 	{
+		// Try finding the object using the sub level selection map first.
+		if (TSoftObjectPtr<ULevel> LastBindingLevel = SubLevelSelectionMap.FindRef(World))
+		{
+			return BoundObjectMap.FindRef(LastBindingLevel);
+		}
+		// Resort to old method where we use the first level we find in the bound object map,
+		// and add an entry in the sub level selection map.
 		for (auto LevelIt = World->GetLevelIterator(); LevelIt; ++LevelIt)
 		{
 			TSoftObjectPtr<ULevel> WeakLevel = *LevelIt;
-			if (const TSoftObjectPtr<UObject>* ObjectPtr = BoundObjectMap.Find(WeakLevel))
+			if (TSoftObjectPtr<UObject> ObjectPtr = BoundObjectMap.FindRef(WeakLevel))
 			{
-				if (ObjectPtr->IsValid())
-				{
-					LevelWithLastSuccessfulResolve = WeakLevel;
-				}
-				return *ObjectPtr;
+				SubLevelSelectionMap.FindOrAdd(World) = WeakLevel;
+				return ObjectPtr;
 			}
 		}
 	}
diff --git a/Engine/Plugins/VirtualProduction/RemoteControl/Source/RemoteControl/Private/RemoteControlPreset.cpp b/Engine/Plugins/VirtualProduction/RemoteControl/Source/RemoteControl/Private/RemoteControlPreset.cpp
--- a/Engine/Plugins/VirtualProduction/RemoteControl/Source/RemoteControl/Private/RemoteControlPreset.cpp
+++ b/Engine/Plugins/VirtualProduction/RemoteControl/Source/RemoteControl/Private/RemoteControlPreset.cpp
@@ -1877,31 +1877,8 @@ void URemoteControlPreset::OnActorDeleted(AActor* Actor)
 			UObject* ResolvedObject = (*It)->Resolve();
 			if (ResolvedObject && (Actor == ResolvedObject || Actor == ResolvedObject->GetTypedOuter<AActor>()))
 			{
-				Modify();
-				(*It)->Modify();
-				(*It)->UnbindObject(ResolvedObject);
-				ModifiedBindings.Add(*It);
-
-				if (!(*It)->IsValid())
-				{
-					It.RemoveCurrent();
-				}
-			}
-		}
-	}
-
-	for (TSharedPtr<FRemoteControlEntity> Entity : Registry->GetExposedEntities<FRemoteControlEntity>())
-	{
-		if (Entity)
-		{
-			for (auto It = Entity->Bindings.CreateIterator(); It; ++It)
-			{
-				if (ModifiedBindings.Contains(It->Get()))
-				{
-					PerFrameUpdatedEntities.Add(Entity->GetId());
-					It.RemoveCurrent();
-					break;
-				}
+				// Defer binding clean up to next frame in case the actor deletion is actually an actor being moved to a different sub level.
+				PerFrameBindingsToClean.Add(*It);
 			}
 		}
 	}
@@ -2020,6 +1997,47 @@ void URemoteControlPreset::OnPackageReloaded(EPackageReloadPhase Phase, FPackage
 		}
 	}
 }
+
+void URemoteControlPreset::CleanUpBindings()
+{
+	TSet<URemoteControlBinding*> BindingsToDelete;
+
+
+	for (URemoteControlBinding* Binding : PerFrameBindingsToClean)
+	{
+		if (Binding)
+		{
+			if (Binding->PruneDeletedObjects())
+			{
+				BindingsToDelete.Add(Binding);
+			}
+		}
+	}
+
+	for (TSharedPtr<FRemoteControlEntity> Entity : Registry->GetExposedEntities<FRemoteControlEntity>())
+	{
+		if (Entity)
+		{
+			for (auto It = Entity->Bindings.CreateIterator(); It; ++It)
+			{
+				// Update bindings that were "touched" regardless of if they were pruned, so that the UI can re-resolve the binding
+				// if one of the object was moved to a sublevel.
+				if (PerFrameBindingsToClean.Contains(It->Get()))
+				{
+					PerFrameUpdatedEntities.Add(Entity->GetId());
+				}
+
+				if (BindingsToDelete.Contains(It->Get()))
+				{
+					It.RemoveCurrent();
+					break;
+				}
+			}
+		}
+	}
+
+	PerFrameBindingsToClean.Reset();
+}
 #endif
 
 void URemoteControlPreset::OnBeginFrame()
@@ -2035,6 +2053,10 @@ void URemoteControlPreset::OnBeginFrame()
 			Entry.Value.CheckForChange();
 		}
 	}
+
+#if WITH_EDITOR
+	CleanUpBindings();
+#endif
 }
 
 void URemoteControlPreset::OnEndFrame()
diff --git a/Engine/Plugins/VirtualProduction/RemoteControl/Source/RemoteControl/Public/RemoteControlBinding.h b/Engine/Plugins/VirtualProduction/RemoteControl/Source/RemoteControl/Public/RemoteControlBinding.h
--- a/Engine/Plugins/VirtualProduction/RemoteControl/Source/RemoteControl/Public/RemoteControlBinding.h
+++ b/Engine/Plugins/VirtualProduction/RemoteControl/Source/RemoteControl/Public/RemoteControlBinding.h
@@ -44,6 +44,12 @@ public:
 	 */
 	virtual bool IsBound(const TSoftObjectPtr<UObject>& Object) const PURE_VIRTUAL(URemoteControlBinding::IsBound, return false;);
 
+	/**
+	 * Remove objects that were deleted from the binding.
+	 * @return true if at least one object was removed from this binding's objects.
+	 */
+	virtual bool PruneDeletedObjects() PURE_VIRTUAL(URemoteControlBinding::PruneDeletedObjects, return false;);
+
 public:
 	/**
 	 * The name of this binding. Defaults to the bound object's name.
@@ -63,6 +69,7 @@ public:
 	virtual UObject* Resolve() const override;
 	virtual bool IsValid() const override;
 	virtual bool IsBound(const TSoftObjectPtr<UObject>& Object) const override;
+	virtual bool PruneDeletedObjects() override;
 
 private:
 	/**
@@ -84,6 +91,7 @@ public:
 	virtual UObject* Resolve() const override;
 	virtual bool IsValid() const override;
 	virtual bool IsBound(const TSoftObjectPtr<UObject>& Object) const override;
+	virtual bool PruneDeletedObjects() override;
 	//~ Begin URemoteControlBinding Interface
 
 	/**
@@ -99,7 +107,7 @@ public:
 	void InitializeForNewLevel();
 
 private:
-	TSoftObjectPtr<UObject> FindObjectFromCurrentWorld() const;
+	TSoftObjectPtr<UObject> ResolveForCurrentWorld() const;
 
 private:
 	/**
@@ -108,6 +116,14 @@ private:
 	UPROPERTY()
 	TMap<TSoftObjectPtr<ULevel>, TSoftObjectPtr<UObject>> BoundObjectMap;
 
+	/**
+	 * Keeps track of which sublevel was last used when binding in a particular world.
+	 * Used in the case where a binding points to objects that end up in the same world but in different sublevels, 
+	 * this ensures that we know which object was last 
+	 */
+	UPROPERTY()
+	mutable TMap<TSoftObjectPtr<UWorld>, TSoftObjectPtr<ULevel>> SubLevelSelectionMap;
+
 	/**
 	 * Caches the last level that had a successful resolve.
 	 * Used to decide which level to use when reinitializing this binding in a new level.
diff --git a/Engine/Plugins/VirtualProduction/RemoteControl/Source/RemoteControl/Public/RemoteControlPreset.h b/Engine/Plugins/VirtualProduction/RemoteControl/Source/RemoteControl/Public/RemoteControlPreset.h
--- a/Engine/Plugins/VirtualProduction/RemoteControl/Source/RemoteControl/Public/RemoteControlPreset.h
+++ b/Engine/Plugins/VirtualProduction/RemoteControl/Source/RemoteControl/Public/RemoteControlPreset.h
@@ -842,6 +842,9 @@ private:
 
 	/** Handles a package reloaded, used to detect a multi-user session being joined in order to update entities. */
 	void OnPackageReloaded(EPackageReloadPhase Phase, FPackageReloadedEvent* Event);
+
+	/** Remove deleted actors from bindings */
+	void CleanUpBindings();
 #endif
 
 	//~ Frame events handlers.
@@ -997,6 +1000,9 @@ PRAGMA_ENABLE_DEPRECATION_WARNINGS
 #if WITH_EDITOR
 	/** List of blueprints for which we have registered events. */
 	TSet<TWeakObjectPtr<UBlueprint>> BlueprintsWithRegisteredDelegates;
+
+	/** List of bindings for which we need to remove stale object pointers. */
+	TSet<URemoteControlBinding*> PerFrameBindingsToClean;
 #endif
 
 PRAGMA_DISABLE_DEPRECATION_WARNINGS
