Fix crash when changing map with ccr property exposed

diff --git a/Engine/Plugins/VirtualProduction/RemoteControl/Source/RemoteControlUI/Private/UI/SRemoteControlPanel.cpp b/Engine/Plugins/VirtualProduction/RemoteControl/Source/RemoteControlUI/Private/UI/SRemoteControlPanel.cpp
--- a/Engine/Plugins/VirtualProduction/RemoteControl/Source/RemoteControlUI/Private/UI/SRemoteControlPanel.cpp
+++ b/Engine/Plugins/VirtualProduction/RemoteControl/Source/RemoteControlUI/Private/UI/SRemoteControlPanel.cpp
@@ -740,6 +740,7 @@ void SRemoteControlPanel::OnMapChange(uint32)
 
 	// Clear the widget cache on map change to make sure we don't keep widgets around pointing to potentially stale objects.
 	WidgetRegistry->Clear();
+	Refresh();
 }
 
 void SRemoteControlPanel::RegisterEvents()
