Fix freeze when applying level snapshot with RC panel open

diff --git a/Engine/Plugins/VirtualProduction/RemoteControl/Source/RemoteControlUI/Private/UI/SRemoteControlPanel.cpp b/Engine/Plugins/VirtualProduction/RemoteControl/Source/RemoteControlUI/Private/UI/SRemoteControlPanel.cpp
--- a/Engine/Plugins/VirtualProduction/RemoteControl/Source/RemoteControlUI/Private/UI/SRemoteControlPanel.cpp
+++ b/Engine/Plugins/VirtualProduction/RemoteControl/Source/RemoteControlUI/Private/UI/SRemoteControlPanel.cpp
@@ -1008,13 +1008,14 @@ FReply SRemoteControlPanel::OnClickRebindAllButton()
 
 void SRemoteControlPanel::UpdateActorFunctionPicker()
 {
-	if (GEditor && ActorFunctionPicker)
+	if (GEditor && ActorFunctionPicker && !NextTickTimerHandle.IsValid())
 	{
-		GEditor->GetTimerManager()->SetTimerForNextTick(FTimerDelegate::CreateLambda([WeakPanelPtr = TWeakPtr<SRemoteControlPanel>(StaticCastSharedRef<SRemoteControlPanel>(AsShared()))]()
+		NextTickTimerHandle = GEditor->GetTimerManager()->SetTimerForNextTick(FTimerDelegate::CreateLambda([WeakPanelPtr = TWeakPtr<SRemoteControlPanel>(StaticCastSharedRef<SRemoteControlPanel>(AsShared()))]()
 		{
 			if (TSharedPtr<SRemoteControlPanel> PanelPtr = WeakPanelPtr.Pin())
 			{
 				PanelPtr->ActorFunctionPicker->Refresh();
+				PanelPtr->NextTickTimerHandle.Invalidate();
 			}
 		}));
 	}
diff --git a/Engine/Plugins/VirtualProduction/RemoteControl/Source/RemoteControlUI/Private/UI/SRemoteControlPanel.h b/Engine/Plugins/VirtualProduction/RemoteControl/Source/RemoteControlUI/Private/UI/SRemoteControlPanel.h
--- a/Engine/Plugins/VirtualProduction/RemoteControl/Source/RemoteControlUI/Private/UI/SRemoteControlPanel.h
+++ b/Engine/Plugins/VirtualProduction/RemoteControl/Source/RemoteControlUI/Private/UI/SRemoteControlPanel.h
@@ -223,7 +223,8 @@ private:
 	TSharedPtr<STextBlock> PresetNameTextBlock;
 	/** Holds a cache of widgets. */
 	TSharedPtr<FRCPanelWidgetRegistry> WidgetRegistry;
-	
+	/** Holds the handle to a timer set for next tick. Used to not schedule more than once event per frame */
+	FTimerHandle NextTickTimerHandle;
 	/** The toolkit that hosts this panel. */
 	TWeakPtr<IToolkitHost> ToolkitHost;
 };
