- Adding remote control preset hooks to transactions to detect modifications of exposed properties in blueprint calls. Otherwise, webapp isn't notified of these changes.

-  Updating exposed properties changed during transaction to consider rootcomponent

diff --git a/Engine/Plugins/VirtualProduction/RemoteControl/Source/RemoteControl/Private/RemoteControlPreset.cpp b/Engine/Plugins/VirtualProduction/RemoteControl/Source/RemoteControl/Private/RemoteControlPreset.cpp
--- a/Engine/Plugins/VirtualProduction/RemoteControl/Source/RemoteControl/Private/RemoteControlPreset.cpp
+++ b/Engine/Plugins/VirtualProduction/RemoteControl/Source/RemoteControl/Private/RemoteControlPreset.cpp
@@ -11,6 +11,7 @@
 #include "HAL/IConsoleManager.h"
 #include "IRemoteControlModule.h"
 #include "Misc/CoreDelegates.h"
+#include "Misc/ITransaction.h"
 #include "Misc/Optional.h"
 #include "RemoteControlExposeRegistry.h"
 #include "RemoteControlFieldPath.h"
@@ -1800,12 +1801,115 @@ void URemoteControlPreset::OnPreObjectPropertyChanged(UObject* Object, const cla
 	}
 }
 
+void URemoteControlPreset::OnObjectTransacted(UObject* InObject, const FTransactionObjectEvent& InTransactionEvent)
+{
+	TRACE_CPUPROFILER_EVENT_SCOPE(URemoteControlPreset::OnObjectTransacted);
+
+	if (InTransactionEvent.GetEventType() != ETransactionObjectEventType::Finalized || InObject == nullptr)
+	{
+		return;
+	}
+
+	//When modifying a blueprint component, the root component gets trashed entirely and transaction event 
+	//doesn't have information about actual properties being changed under it. We just have RootComponent
+	//marked as being changed on a given actor. When that happens, we consider all properties binded to a component 
+	//that belongs to that actor to be changed.
+	const bool bHasRootComponentChanged = InTransactionEvent.GetChangedProperties().Contains("RootComponent");
+
+	//Verify if have bindings to the specified object or a binding to a component owned by the modified object
+	bool bIsObjectBound = false;
+	for (const URemoteControlBinding* Binding : Bindings)
+	{
+		if (Binding->IsBound(InObject))
+		{
+			bIsObjectBound = true;
+			break;
+		}
+		else if (bHasRootComponentChanged)
+		{
+			if (UActorComponent* Component = Cast<UActorComponent>(Binding->Resolve()))
+			{
+				if (Component->GetOwner() == InObject)
+				{
+					bIsObjectBound = true;
+					break;
+				}
+			}
+		}
+	}
+
+	if (bIsObjectBound == false)
+	{
+		return;
+	}
+
+	//Go through all properties and verify if it was modified
+	for (const TSharedPtr<FRemoteControlProperty>& Property : Registry->GetExposedEntities<FRemoteControlProperty>())
+	{
+		if (PerFrameModifiedProperties.Contains(Property->GetId()))
+		{
+			continue;
+		}
+
+		TArray<UObject*> BoundObjects = Property->GetBoundObjects();
+		if (BoundObjects.Num() == 0 || Property->FieldPathInfo.Segments.Num() < 1)
+		{
+			continue;
+		}
+		
+		for (UObject* BoundObject : BoundObjects)
+		{
+			bool bPropertyModified = false;
+			if (BoundObject == InObject || BoundObject->GetOuter() == InObject)
+			{
+				//Before going into all transacted properties, verify if root component
+				//has changed and verify if its owner matches the input object
+				if (bHasRootComponentChanged)
+				{
+					if (UActorComponent* Component = Cast<UActorComponent>(BoundObject))
+					{
+						if (Component->GetOwner() == InObject)
+						{
+							PerFrameModifiedProperties.Add(Property->GetId());
+							bPropertyModified = true;
+						}
+					}
+				}
+
+				//If root component checkup didn't match, go through all transacted properties and look if it's the root of the exposed property
+				if (bPropertyModified == false)
+				{
+					for (const FName ChangedProperty : InTransactionEvent.GetChangedProperties())
+					{
+						//Changed properties will be a path (dot separated) to the modified property
+						//All tests only returned the root structure (for nested props) so for now, only verify the first segment if it matches
+						if (Property->FieldPathInfo.Segments[0].Name == ChangedProperty)
+						{
+							PerFrameModifiedProperties.Add(Property->GetId());
+							bPropertyModified = true;
+							break;		
+						}
+					}
+				}
+			}
+
+			//If property was modified, no need to continue looking at other bindings
+			if (bPropertyModified)
+			{
+				UE_LOG(LogRemoteControl, VeryVerbose, TEXT("(%s) Exposed property '%s'::'%s' change detected during transaction of %s"), *GetName(), *BoundObject->GetName(), *Property->FieldName.ToString(), *InObject->GetName());
+				break;
+			}
+		}
+	}
+}
+
 void URemoteControlPreset::RegisterDelegates()
 {
 	UnregisterDelegates();
 
 #if WITH_EDITOR
 	FCoreUObjectDelegates::OnObjectPropertyChanged.AddUObject(this, &URemoteControlPreset::OnObjectPropertyChanged);
+	FCoreUObjectDelegates::OnObjectTransacted.AddUObject(this, &URemoteControlPreset::OnObjectTransacted);
 
 		
 	FCoreUObjectDelegates::OnPreObjectPropertyChanged.AddUObject(this, &URemoteControlPreset::OnPreObjectPropertyChanged);
@@ -1858,6 +1962,7 @@ void URemoteControlPreset::UnregisterDelegates()
 		GEngine->OnLevelActorDeleted().RemoveAll(this);
 	}
 
+	FCoreUObjectDelegates::OnObjectTransacted.RemoveAll(this);
 	FCoreUObjectDelegates::OnPreObjectPropertyChanged.RemoveAll(this);
 	FCoreUObjectDelegates::OnObjectPropertyChanged.RemoveAll(this);
 #endif
diff --git a/Engine/Plugins/VirtualProduction/RemoteControl/Source/RemoteControl/Public/RemoteControlPreset.h b/Engine/Plugins/VirtualProduction/RemoteControl/Source/RemoteControl/Public/RemoteControlPreset.h
--- a/Engine/Plugins/VirtualProduction/RemoteControl/Source/RemoteControl/Public/RemoteControlPreset.h
+++ b/Engine/Plugins/VirtualProduction/RemoteControl/Source/RemoteControl/Public/RemoteControlPreset.h
@@ -20,6 +20,7 @@ struct FPropertyChangedEvent;
 struct FRCFieldPathInfo;
 struct FRemoteControlActor;
 struct FRemoteControlPresetLayout;
+class FTransactionObjectEvent;
 class FRemoteControlPresetRebindingManager;
 class UBlueprint;
 class URemoteControlExposeRegistry;
@@ -831,6 +832,7 @@ private:
 	//~ Keep track of any property change to notify if one of the exposed property has changed
 	void OnObjectPropertyChanged(UObject* Object, FPropertyChangedEvent& Event);
 	void OnPreObjectPropertyChanged(UObject* Object, const class FEditPropertyChain& PropertyChain);
+	void OnObjectTransacted(UObject* InObject, const FTransactionObjectEvent& InTransactionEvent);
 
 #if WITH_EDITOR	
 	//~ Handle events that can incur bindings to be modified.
