1) Remote API Panel Shortcut

Users need the ability to toggle the Editor mode on and off with a shortcut in the Remote Control UI.  Also if we can make the edit button look like the toggle in the lower right corner of the Stage Monitor it would look a lot nicer.

Note : Position of the Edit Mode toggle is subject to change.

2) "Use Less CPU" setting warning when using RemoteControl

Display a warning the RC Preset panel when the Editor Setting "Use Less CPU when in Background" is enabled, and prompt the user to disable it.

diff --git a/Engine/Plugins/VirtualProduction/RemoteControl/Source/RemoteControlUI/Private/Commands/RemoteControlCommands.cpp b/Engine/Plugins/VirtualProduction/RemoteControl/Source/RemoteControlUI/Private/Commands/RemoteControlCommands.cpp
new file mode 100644
--- /dev/null
+++ b/Engine/Plugins/VirtualProduction/RemoteControl/Source/RemoteControlUI/Private/Commands/RemoteControlCommands.cpp
@@ -0,0 +1,27 @@
+// Copyright Epic Games, Inc. All Rights Reserved.
+
+#include "Commands/RemoteControlCommands.h"
+
+#include "RemoteControlUI/Private/UI/RemoteControlPanelStyle.h"
+
+#define LOCTEXT_NAMESPACE "RemoteControlCommands"
+
+FRemoteControlCommands::FRemoteControlCommands()
+	: TCommands<FRemoteControlCommands>
+	(
+		TEXT("RemoteControl"),
+		LOCTEXT("RemoteControl", "Remote Control API"),
+		NAME_None,
+		FRemoteControlPanelStyle::GetStyleSetName()
+	)
+{
+}
+
+void FRemoteControlCommands::RegisterCommands()
+{
+	// Toggle Edit Mode
+
+	UI_COMMAND(ToggleEditMode, "Toggle Editing", "Toggles edit mode in remote control panel which enables user to modify the exposed entities.", EUserInterfaceActionType::Check, FInputChord(EModifierKey::Control, EKeys::E));
+}
+
+#undef LOCTEXT_NAMESPACE
diff --git a/Engine/Plugins/VirtualProduction/RemoteControl/Source/RemoteControlUI/Private/Commands/RemoteControlCommands.h b/Engine/Plugins/VirtualProduction/RemoteControl/Source/RemoteControlUI/Private/Commands/RemoteControlCommands.h
new file mode 100644
--- /dev/null
+++ b/Engine/Plugins/VirtualProduction/RemoteControl/Source/RemoteControlUI/Private/Commands/RemoteControlCommands.h
@@ -0,0 +1,28 @@
+// Copyright Epic Games, Inc. All Rights Reserved.
+
+#pragma once
+
+#include "CoreMinimal.h"
+
+#include "Framework/Commands/Commands.h"
+
+/**
+ * Defines commands for the Remote Control API which enables most functionality of the Remote Control Panel.
+ */
+class FRemoteControlCommands : public TCommands<FRemoteControlCommands>
+{
+public:
+
+	FRemoteControlCommands();
+
+	//~ BEGIN : TCommands<> Implementation(s)
+
+	virtual void RegisterCommands() override;
+
+	//~ END : TCommands<> Implementation(s)
+
+	/**
+	 * Holds the information about UI Command that toggles edit mode in RC Panel.
+	 */
+	TSharedPtr<FUICommandInfo> ToggleEditMode;
+};
diff --git a/Engine/Plugins/VirtualProduction/RemoteControl/Source/RemoteControlUI/Private/RemoteControlUIModule.cpp b/Engine/Plugins/VirtualProduction/RemoteControl/Source/RemoteControlUI/Private/RemoteControlUIModule.cpp
--- a/Engine/Plugins/VirtualProduction/RemoteControl/Source/RemoteControlUI/Private/RemoteControlUIModule.cpp
+++ b/Engine/Plugins/VirtualProduction/RemoteControl/Source/RemoteControlUI/Private/RemoteControlUIModule.cpp
@@ -4,11 +4,13 @@
 
 #include "AssetToolsModule.h"
 #include "AssetTools/RemoteControlPresetActions.h"
+#include "Commands/RemoteControlCommands.h"
 #include "EditorStyleSet.h"
 #include "Engine/StaticMesh.h"
 #include "Engine/SkeletalMesh.h"
 #include "Framework/MultiBox/MultiBoxBuilder.h"
 #include "HAL/PlatformApplicationMisc.h"
+#include "Interfaces/IMainFrameModule.h"
 #include "Kismet2/ComponentEditorUtils.h"
 #include "MaterialEditor/DEditorParameterValue.h"
 #include "PropertyHandle.h"
@@ -93,6 +95,7 @@ namespace RemoteControlUIModule
 void FRemoteControlUIModule::StartupModule()
 {
 	FRemoteControlPanelStyle::Initialize();
+	BindRemoteControlCommands();
 	RegisterAssetTools();
 	RegisterDetailRowExtension();
 	RegisterContextMenuExtender();
@@ -110,6 +113,7 @@ void FRemoteControlUIModule::ShutdownModule()
 	UnregisterContextMenuExtender();
 	UnregisterDetailRowExtension();
 	UnregisterAssetTools();
+	UnbindRemoteControlCommands();
 	FRemoteControlPanelStyle::Shutdown();
 }
 
@@ -229,6 +233,29 @@ void FRemoteControlUIModule::UnregisterEvents()
 	FEditorDelegates::PostUndoRedo.RemoveAll(this);
 }
 
+void FRemoteControlUIModule::ToggleEditMode()
+{
+	if (TSharedPtr<SRemoteControlPanel> Panel = WeakActivePanel.Pin()) // Check whether the panel is active.
+	{
+		Panel->SetEditMode(Panel->IsInEditMode() ? false : true);
+	}
+}
+
+bool FRemoteControlUIModule::CanToggleEditMode() const
+{
+	return WeakActivePanel.IsValid();
+}
+
+bool FRemoteControlUIModule::IsInEditMode() const
+{
+	if (TSharedPtr<SRemoteControlPanel> Panel = WeakActivePanel.Pin()) // Check whether the panel is active.
+	{
+		return Panel->IsInEditMode();
+	}
+
+	return false;
+}
+
 URemoteControlPreset* FRemoteControlUIModule::GetActivePreset() const
 {
 	if (const TSharedPtr<SRemoteControlPanel> Panel = WeakActivePanel.Pin())
@@ -256,6 +283,39 @@ void FRemoteControlUIModule::UnregisterAssetTools()
 	RemoteControlPresetActions.Reset();
 }
 
+void FRemoteControlUIModule::BindRemoteControlCommands()
+{
+	FRemoteControlCommands::Register();
+
+	const FRemoteControlCommands& Commands = FRemoteControlCommands::Get();
+
+	IMainFrameModule& MainFrame = FModuleManager::Get().LoadModuleChecked<IMainFrameModule>("MainFrame");
+
+	FUICommandList& ActionList = *MainFrame.GetMainFrameCommandBindings();
+
+	// Toggle Edit Mode
+
+	ActionList.MapAction(Commands.ToggleEditMode,
+		FExecuteAction::CreateRaw(this, &FRemoteControlUIModule::ToggleEditMode),
+		FCanExecuteAction::CreateRaw(this, &FRemoteControlUIModule::CanToggleEditMode),
+		FIsActionChecked::CreateRaw(this, &FRemoteControlUIModule::IsInEditMode),
+		FIsActionButtonVisible::CreateRaw(this, &FRemoteControlUIModule::CanToggleEditMode)
+	);
+}
+
+void FRemoteControlUIModule::UnbindRemoteControlCommands()
+{
+	const FRemoteControlCommands& Commands = FRemoteControlCommands::Get();
+
+	IMainFrameModule& MainFrame = FModuleManager::Get().LoadModuleChecked<IMainFrameModule>("MainFrame");
+
+	FUICommandList& ActionList = *MainFrame.GetMainFrameCommandBindings();
+
+	ActionList.UnmapAction(Commands.ToggleEditMode);
+
+	FRemoteControlCommands::Unregister();
+}
+
 void FRemoteControlUIModule::HandleCreatePropertyRowExtension(const FOnGenerateGlobalRowExtensionArgs& InArgs, TArray<FPropertyRowExtensionButton>& OutExtensions)
 {
 	// Expose/Unexpose button.
diff --git a/Engine/Plugins/VirtualProduction/RemoteControl/Source/RemoteControlUI/Private/RemoteControlUIModule.h b/Engine/Plugins/VirtualProduction/RemoteControl/Source/RemoteControlUI/Private/RemoteControlUIModule.h
--- a/Engine/Plugins/VirtualProduction/RemoteControl/Source/RemoteControlUI/Private/RemoteControlUIModule.h
+++ b/Engine/Plugins/VirtualProduction/RemoteControl/Source/RemoteControlUI/Private/RemoteControlUIModule.h
@@ -86,6 +86,10 @@ private:
 	void RegisterAssetTools();
 	void UnregisterAssetTools();
 
+	//~ Remote Control Commands
+	void BindRemoteControlCommands();
+	void UnbindRemoteControlCommands();
+
 	//~ Context menu extenders
 	void RegisterContextMenuExtender();
 	void UnregisterContextMenuExtender();
@@ -98,6 +102,11 @@ private:
 	void RegisterEvents();
 	void UnregisterEvents();
 
+	//~ Toggle Edit Mode RC Command.
+	void ToggleEditMode();
+	bool CanToggleEditMode() const;	
+	bool IsInEditMode() const;
+
 	/** Handle creating the property row extensions.  */
 	void HandleCreatePropertyRowExtension(const FOnGenerateGlobalRowExtensionArgs& InArgs, TArray<FPropertyRowExtensionButton>& OutExtensions);
 
diff --git a/Engine/Plugins/VirtualProduction/RemoteControl/Source/RemoteControlUI/Private/UI/RemoteControlPanelStyle.cpp b/Engine/Plugins/VirtualProduction/RemoteControl/Source/RemoteControlUI/Private/UI/RemoteControlPanelStyle.cpp
--- a/Engine/Plugins/VirtualProduction/RemoteControl/Source/RemoteControlUI/Private/UI/RemoteControlPanelStyle.cpp
+++ b/Engine/Plugins/VirtualProduction/RemoteControl/Source/RemoteControlUI/Private/UI/RemoteControlPanelStyle.cpp
@@ -28,6 +28,7 @@ void FRemoteControlPanelStyle::Initialize()
 	const FVector2D Icon8x8(8.0f, 8.0f);
 	const FVector2D Icon16x16(16.0f, 16.0f);
 	const FVector2D Icon20x20(20.0f, 20.0f);
+	const FVector2D Icon28x14(28.0f, 14.0f);
 	const FVector2D Icon64x64(64.0f, 64.0f);
 
 	StyleSet = MakeShared<FSlateStyleSet>(GetStyleSetName());
@@ -102,6 +103,26 @@ void FRemoteControlPanelStyle::Initialize()
 	
 	StyleSet->Set("RemoteControlPanel.GroupRow", GroupRowStyle);
 	
+	// Checkbox for live/edit mode
+	
+	{
+		StyleSet->Set("Switch.ToggleOff", new IMAGE_PLUGIN_BRUSH("Icons/Switch_OFF", Icon28x14));
+		
+		StyleSet->Set("Switch.ToggleOn", new IMAGE_PLUGIN_BRUSH("Icons/Switch_ON", Icon28x14));
+
+		FCheckBoxStyle SwitchStyle = FCheckBoxStyle()
+			.SetForegroundColor(FLinearColor::White)
+			.SetUncheckedImage(*StyleSet->GetBrush("Switch.ToggleOff"))
+			.SetUncheckedHoveredImage(*StyleSet->GetBrush("Switch.ToggleOff"))
+			.SetUncheckedPressedImage(*StyleSet->GetBrush("Switch.ToggleOff"))
+			.SetCheckedImage(*StyleSet->GetBrush("Switch.ToggleOn"))
+			.SetCheckedHoveredImage(*StyleSet->GetBrush("Switch.ToggleOn"))
+			.SetCheckedPressedImage(*StyleSet->GetBrush("Switch.ToggleOn"))
+			.SetPadding(FMargin(1, 1, 1, 1));
+
+		StyleSet->Set("RemoteControlPanel.Switch", SwitchStyle);
+	}
+
 	FSlateStyleRegistry::RegisterSlateStyle(*StyleSet.Get());
 }
 
diff --git a/Engine/Plugins/VirtualProduction/RemoteControl/Source/RemoteControlUI/Private/UI/SRemoteControlPanel.cpp b/Engine/Plugins/VirtualProduction/RemoteControl/Source/RemoteControlUI/Private/UI/SRemoteControlPanel.cpp
--- a/Engine/Plugins/VirtualProduction/RemoteControl/Source/RemoteControlUI/Private/UI/SRemoteControlPanel.cpp
+++ b/Engine/Plugins/VirtualProduction/RemoteControl/Source/RemoteControlUI/Private/UI/SRemoteControlPanel.cpp
@@ -42,6 +42,7 @@
 #include "SRCPanelExposedField.h"
 #include "SRCPanelTreeNode.h"
 #include "Subsystems/Subsystem.h"
+#include "SWarningOrErrorBox.h"
 #include "Templates/SharedPointer.h"
 #include "Templates/SubclassOf.h"
 #include "Templates/UnrealTypeTraits.h"
@@ -139,12 +140,6 @@ void SRemoteControlPanel::Construct(const FArguments& InArgs, URemoteControlPres
 		[
 			// Top tool bar
 			SNew(SHorizontalBox)
-			+ SHorizontalBox::Slot()
-			.AutoWidth()
-			.VAlign(VAlign_Center)
-			[
-				CreateCPUThrottleButton()
-			]
 
 			+ SHorizontalBox::Slot()
 			.Padding(FMargin(5.0f, 0.0f))
@@ -202,23 +197,28 @@ void SRemoteControlPanel::Construct(const FArguments& InArgs, URemoteControlPres
 						.Image(FEditorStyle::Get().GetBrush("LevelEditor.Tabs.Details"))
 					]
 				]
-
+				// Edit Mode
 				+ SHorizontalBox::Slot()
-				.VAlign(VAlign_Center)
-				.Padding(4.0f, 0)
 				.AutoWidth()
+				.Padding(2.f, 0.f)
+				.VAlign(VAlign_Center)
 				[
 					SNew(STextBlock)
-					.Text(LOCTEXT("EditModeLabel", "Edit Mode: "))
+					.Text(LOCTEXT("EditModeLabel", "Edit Mode : "))
 				]
 				+ SHorizontalBox::Slot()
-				.VAlign(VAlign_Center)
 				.AutoWidth()
+				.Padding(2.f, 0.f)
+				.VAlign(VAlign_Center)
 				[
 					SNew(SCheckBox)
+					.Style(FRemoteControlPanelStyle::Get(), "RemoteControlPanel.Switch")
+					.ToolTipText(LOCTEXT("EditModeTooltip", "Toggle Editing (Ctrl + E)"))
 					.IsChecked_Lambda([this]() { return this->bIsInEditMode ? ECheckBoxState::Checked : ECheckBoxState::Unchecked; })
+					.Padding(2.f)
 					.OnCheckStateChanged(this, &SRemoteControlPanel::OnEditModeCheckboxToggle)
 				]
+				// Enable Log
 				+ SHorizontalBox::Slot()
 				.VAlign(VAlign_Center)
 				.Padding(4.0f, 0)
@@ -337,6 +337,13 @@ void SRemoteControlPanel::Construct(const FArguments& InArgs, URemoteControlPres
 				SNew(SRCLogger)
 			]
 		]
+		// Use less CPU Warning
+		+SVerticalBox::Slot()
+		.AutoHeight()
+		.Padding(FMargin(5.f, 8.f, 5.f, 5.f))
+		[
+			CreateCPUThrottleWarning()
+		]
 	];
 
 	for (const TSharedRef<SWidget>& Widget : ExtensionWidgets)
@@ -484,23 +491,22 @@ FReply SRemoteControlPanel::OnClickDisableUseLessCPU() const
 	return FReply::Handled();
 }
 
-TSharedRef<SWidget> SRemoteControlPanel::CreateCPUThrottleButton() const
+TSharedRef<SWidget> SRemoteControlPanel::CreateCPUThrottleWarning() const
 {
 	FProperty* PerformanceThrottlingProperty = FindFieldChecked<FProperty>(UEditorPerformanceSettings::StaticClass(), GET_MEMBER_NAME_CHECKED(UEditorPerformanceSettings, bThrottleCPUWhenNotForeground));
 	FFormatNamedArguments Arguments;
 	Arguments.Add(TEXT("PropertyName"), PerformanceThrottlingProperty->GetDisplayNameText());
 	FText PerformanceWarningText = FText::Format(LOCTEXT("RemoteControlPerformanceWarning", "Warning: The editor setting '{PropertyName}' is currently enabled\nThis will stop editor windows from updating in realtime while the editor is not in focus"), Arguments);
 
-	return SNew(SButton)
-		.ButtonStyle(FAppStyle::Get(), "FlatButton")
-		.Visibility_Lambda([]() {return GetDefault<UEditorPerformanceSettings>()->bThrottleCPUWhenNotForeground ? EVisibility::Visible : EVisibility::Collapsed; } )
-		.OnClicked_Raw(this, &SRemoteControlPanel::OnClickDisableUseLessCPU)
+	return SNew(SWarningOrErrorBox)
+		.Visibility_Lambda([]() { return GetDefault<UEditorPerformanceSettings>()->bThrottleCPUWhenNotForeground ? EVisibility::Visible : EVisibility::Collapsed; })
+		.MessageStyle(EMessageStyle::Warning)
+		.Message(PerformanceWarningText)
 		[
-			SNew(STextBlock)
-			.ToolTipText(MoveTemp(PerformanceWarningText))
-			.TextStyle(FRemoteControlPanelStyle::Get(), "RemoteControlPanel.Button.TextStyle")
-			.Font(FAppStyle::Get().GetFontStyle("FontAwesome.10"))
-			.Text(FEditorFontGlyphs::Exclamation_Triangle)
+			SNew(SButton)
+			.OnClicked(this, &SRemoteControlPanel::OnClickDisableUseLessCPU)
+			.TextStyle(FAppStyle::Get(), "DialogButtonText")
+			.Text(LOCTEXT("RemoteControlPerformanceWarningDisable", "Disable"))
 		];
 }
 
diff --git a/Engine/Plugins/VirtualProduction/RemoteControl/Source/RemoteControlUI/Private/UI/SRemoteControlPanel.h b/Engine/Plugins/VirtualProduction/RemoteControl/Source/RemoteControlUI/Private/UI/SRemoteControlPanel.h
--- a/Engine/Plugins/VirtualProduction/RemoteControl/Source/RemoteControlUI/Private/UI/SRemoteControlPanel.h
+++ b/Engine/Plugins/VirtualProduction/RemoteControl/Source/RemoteControlUI/Private/UI/SRemoteControlPanel.h
@@ -89,7 +89,12 @@ public:
 	 * Set the edit mode of the panel.
 	 * @param bEditMode The desired mode.
 	 */
-	void SetEditMode(bool bEditMode) { bIsInEditMode = bEditMode; }
+	void SetEditMode(bool bEditMode)
+	{
+		bIsInEditMode = bEditMode;
+	
+		OnEditModeChange.ExecuteIfBound(SharedThis(this), bIsInEditMode);
+	}
 
 	/**
 	 * Get the exposed entity list.
@@ -100,6 +105,7 @@ public:
 	void Refresh();
 
 private:
+
 	/** Register editor events needed to handle reloading objects and blueprint libraries. */
 	void RegisterEvents();
 	/** Unregister editor events */
@@ -136,7 +142,7 @@ private:
 	FReply OnClickDisableUseLessCPU() const;
 
 	/** Creates a widget that warns the user when CPU throttling is enabled.  */
-	TSharedRef<SWidget> CreateCPUThrottleButton() const;
+	TSharedRef<SWidget> CreateCPUThrottleWarning() const;
 
 	/** Create expose button, allowing to expose blueprints and actor functions. */
 	TSharedRef<SWidget> CreateExposeButton();
diff --git a/Engine/Plugins/VirtualProduction/RemoteControl/Source/RemoteControlUI/RemoteControlUI.Build.cs b/Engine/Plugins/VirtualProduction/RemoteControl/Source/RemoteControlUI/RemoteControlUI.Build.cs
--- a/Engine/Plugins/VirtualProduction/RemoteControl/Source/RemoteControlUI/RemoteControlUI.Build.cs
+++ b/Engine/Plugins/VirtualProduction/RemoteControl/Source/RemoteControlUI/RemoteControlUI.Build.cs
@@ -24,6 +24,7 @@ public class RemoteControlUI : ModuleRules
 				"Engine",
 				"HotReload",
 				"InputCore",
+				"MainFrame",
 				"MessageLog",
 				"Projects",
 				"PropertyEditor",
@@ -33,6 +34,7 @@ public class RemoteControlUI : ModuleRules
 				"SceneOutliner",
 				"Slate",
 				"SlateCore",
+				"ToolWidgets",
 				"UnrealEd"
 			}
 		);
