Remote Control Preset
Add option to rebind all properties for actor
Add option to rebind exposed property to different component

diff --git a/Engine/Plugins/VirtualProduction/RemoteControl/Source/RemoteControl/Private/RemoteControlEntity.cpp b/Engine/Plugins/VirtualProduction/RemoteControl/Source/RemoteControl/Private/RemoteControlEntity.cpp
--- a/Engine/Plugins/VirtualProduction/RemoteControl/Source/RemoteControl/Private/RemoteControlEntity.cpp
+++ b/Engine/Plugins/VirtualProduction/RemoteControl/Source/RemoteControl/Private/RemoteControlEntity.cpp
@@ -46,22 +46,18 @@ void FRemoteControlEntity::BindObject(UObject* InObjectToBind)
 		return;
 	}
 
-	URemoteControlBinding* Binding = nullptr;
+	URemoteControlBinding* Binding = Owner->FindOrAddBinding(InObjectToBind);
 
 	if (Bindings.Num() == 0)
 	{
-		Binding = Owner->FindOrAddBinding(InObjectToBind);
 		Bindings.Emplace(Binding);
 	}
-
-	Binding = Bindings[0].Get();
-	
-	if (Binding)
+	else
 	{
-		Binding->Modify();
-		Binding->SetBoundObject(InObjectToBind);
-		OnEntityModifiedDelegate.ExecuteIfBound(Id);
+		Bindings[0] = Binding;
 	}
+
+	OnEntityModifiedDelegate.ExecuteIfBound(Id);
 }
 
 bool FRemoteControlEntity::IsBound() const
diff --git a/Engine/Plugins/VirtualProduction/RemoteControl/Source/RemoteControl/Private/RemoteControlPreset.cpp b/Engine/Plugins/VirtualProduction/RemoteControl/Source/RemoteControl/Private/RemoteControlPreset.cpp
--- a/Engine/Plugins/VirtualProduction/RemoteControl/Source/RemoteControl/Private/RemoteControlPreset.cpp
+++ b/Engine/Plugins/VirtualProduction/RemoteControl/Source/RemoteControl/Private/RemoteControlPreset.cpp
@@ -1268,6 +1268,14 @@ void URemoteControlPreset::RebindUnboundEntities()
 	Algo::Transform(Registry->GetExposedEntities(), PerFrameUpdatedEntities, [](const TSharedPtr<FRemoteControlEntity>& Entity) { return Entity->GetId(); });
 }
 
+void URemoteControlPreset::RebindAllEntitiesUnderSameActor(const FGuid& EntityId, AActor* NewActor)
+{
+	if (TSharedPtr<FRemoteControlEntity> Entity = Registry->GetExposedEntity(EntityId))
+	{
+		RebindingManager->RebindAllEntitiesUnderSameActor(this, Entity, NewActor);
+	}
+}
+
 void URemoteControlPreset::NotifyExposedPropertyChanged(FName PropertyLabel)
 {
 	if (TSharedPtr<FRemoteControlProperty> ExposedProperty = GetExposedEntity<FRemoteControlProperty>(GetExposedEntityId(PropertyLabel)).Pin())
diff --git a/Engine/Plugins/VirtualProduction/RemoteControl/Source/RemoteControl/Private/RemoteControlPresetRebindingManager.cpp b/Engine/Plugins/VirtualProduction/RemoteControl/Source/RemoteControl/Private/RemoteControlPresetRebindingManager.cpp
--- a/Engine/Plugins/VirtualProduction/RemoteControl/Source/RemoteControl/Private/RemoteControlPresetRebindingManager.cpp
+++ b/Engine/Plugins/VirtualProduction/RemoteControl/Source/RemoteControl/Private/RemoteControlPresetRebindingManager.cpp
@@ -232,6 +232,114 @@ void FRemoteControlPresetRebindingManager::Rebind(URemoteControlPreset* Preset)
 	}
 }
 
+void FRemoteControlPresetRebindingManager::RebindAllEntitiesUnderSameActor(URemoteControlPreset* Preset, const TSharedPtr<FRemoteControlEntity>& Entity, AActor* NewActor)
+{
+	check(Preset && Entity && NewActor);
+
+	TMap<FName, TSet<URemoteControlBinding*>> ActorsToBindings;
+	ActorsToBindings.Reserve(Preset->Bindings.Num());
+
+	auto GetBindingLastLevel = [](URemoteControlBinding* Binding) -> FSoftObjectPath
+	{
+		FSoftObjectPath LevelPath;
+		if (Binding && Binding->IsA<URemoteControlLevelDependantBinding>())
+		{
+			LevelPath = CastChecked<URemoteControlLevelDependantBinding>(Binding)->LevelWithLastSuccessfulResolve.ToSoftObjectPath();
+		}
+		return LevelPath;
+	};
+
+	auto GetBindingActorName = [](URemoteControlBinding* Binding) -> FName
+	{
+		FName ActorName;
+		if (Binding && Binding->IsA<URemoteControlLevelDependantBinding>())
+		{
+			URemoteControlLevelDependantBinding* LDBinding = CastChecked<URemoteControlLevelDependantBinding>(Binding);
+			FString Path = LDBinding->BoundObjectMap.FindRef(LDBinding->LevelWithLastSuccessfulResolve).ToString();
+			static const TCHAR* PersistentLevel = TEXT("PersistentLevel.");
+			const int32 PersistentLevelStringLength = 16;
+
+			int32 Position = Path.Find(PersistentLevel);
+
+			if (Position != INDEX_NONE)
+			{
+				int32 DotPosition = Path.Find(TEXT("."), ESearchCase::IgnoreCase, ESearchDir::FromStart, Position + PersistentLevelStringLength);
+				if (DotPosition != INDEX_NONE)
+				{
+					ActorName = *Path.Mid(Position + PersistentLevelStringLength, DotPosition - (Position + PersistentLevelStringLength));
+				}
+			}
+		}
+
+		return ActorName;
+	};
+
+
+	FName InitialBindingActorName;
+	FSoftObjectPath InitialBindingLastLevel;
+
+	const TArray<TWeakObjectPtr<URemoteControlBinding>>& EntityBindings = Entity->GetBindings();
+	if (EntityBindings.Num())
+	{
+		InitialBindingLastLevel = GetBindingLastLevel(EntityBindings[0].Get());
+		InitialBindingActorName = GetBindingActorName(EntityBindings[0].Get());
+	}
+
+	TMap<URemoteControlBinding*, FString> BindingToComponentName;
+
+	// Find bindings with the same actor as the original binding
+	for (URemoteControlBinding* Binding : Preset->Bindings)
+	{
+		FSoftObjectPath BindingItLastLevel = GetBindingLastLevel(Binding);
+		FName BindingItActorName = GetBindingActorName(Binding);
+
+		// To add a binding to the list of bindings to rebind, make sure it had the same level and actor as the initial binding
+		if (BindingItActorName == InitialBindingActorName && BindingItLastLevel == InitialBindingLastLevel)
+		{
+			BindingToComponentName.Add(Binding, BindingItActorName.ToString());
+		}
+	}
+
+	// Gather entities to rebind
+	TArray<TSharedPtr<FRemoteControlEntity>> EntitiesToRebind;
+	for (const TWeakPtr<FRemoteControlEntity>& WeakEntity : Preset->GetExposedEntities<FRemoteControlEntity>())
+	{
+		TSharedPtr<FRemoteControlEntity> RCEntity = WeakEntity.Pin();
+		if (RCEntity && RCEntity->GetBindings().Num())
+		{
+			if (BindingToComponentName.Contains(RCEntity->GetBindings()[0].Get()))
+			{
+				EntitiesToRebind.Add(RCEntity);
+			}
+		}
+	}
+
+	// Find the new binding target, then rebind.
+	for (const TSharedPtr<FRemoteControlEntity>& EntityToRebind : EntitiesToRebind)
+	{
+		URemoteControlBinding* EntityBinding = EntityToRebind->GetBindings()[0].Get();
+		FString ActorName = BindingToComponentName.FindChecked(EntityBinding);
+
+		UObject* RebindTarget = NewActor;
+
+		// If the binding is targetting a component rather than an actor
+		if (ActorName != EntityBinding->Name)
+		{
+			if (UObject* ResolvedComponent = StaticFindObject(UObject::StaticClass(), NewActor, *EntityBinding->Name, false))
+			{
+				RebindTarget = ResolvedComponent;
+			}
+		}
+
+		EntityToRebind->BindObject(RebindTarget);
+
+		if (Preset->GetExposedEntityType(EntityToRebind->GetId()) == FRemoteControlProperty::StaticStruct())
+		{
+			StaticCastSharedPtr<FRemoteControlProperty>(EntityToRebind)->EnableEditCondition();
+		}
+	}
+}
+
 TMap<UClass*, TArray<TSharedPtr<FRemoteControlEntity>>> FRemoteControlPresetRebindingManager::GroupByEntitySupportedOwnerClass(TConstArrayView<TSharedPtr<FRemoteControlEntity>> Entities) const
 {
 	TMap<UClass*, TArray<TSharedPtr<FRemoteControlEntity>>> GroupedEntities;
diff --git a/Engine/Plugins/VirtualProduction/RemoteControl/Source/RemoteControl/Private/RemoteControlPresetRebindingManager.h b/Engine/Plugins/VirtualProduction/RemoteControl/Source/RemoteControl/Private/RemoteControlPresetRebindingManager.h
--- a/Engine/Plugins/VirtualProduction/RemoteControl/Source/RemoteControl/Private/RemoteControlPresetRebindingManager.h
+++ b/Engine/Plugins/VirtualProduction/RemoteControl/Source/RemoteControl/Private/RemoteControlPresetRebindingManager.h
@@ -61,6 +61,11 @@ public:
 	 */
 	void Rebind(URemoteControlPreset* Preset);
 
+	/**
+	 * Given a RC Entity, rebind all entities with the same owner to a new actor
+	 */
+	void RebindAllEntitiesUnderSameActor(URemoteControlPreset* Preset, const TSharedPtr<FRemoteControlEntity>& Entity, AActor* NewActor);
+
 private:
 	/** Group RC Entities by the entity's supported owner class. */
 	TMap<UClass*, TArray<TSharedPtr<FRemoteControlEntity>>> GroupByEntitySupportedOwnerClass(TConstArrayView<TSharedPtr<FRemoteControlEntity>> Entities) const;
diff --git a/Engine/Plugins/VirtualProduction/RemoteControl/Source/RemoteControl/Public/RemoteControlBinding.h b/Engine/Plugins/VirtualProduction/RemoteControl/Source/RemoteControl/Public/RemoteControlBinding.h
--- a/Engine/Plugins/VirtualProduction/RemoteControl/Source/RemoteControl/Public/RemoteControlBinding.h
+++ b/Engine/Plugins/VirtualProduction/RemoteControl/Source/RemoteControl/Public/RemoteControlBinding.h
@@ -132,4 +132,6 @@ private:
 	 */
 	UPROPERTY()
 	mutable TSoftObjectPtr<ULevel> LevelWithLastSuccessfulResolve;
+
+	friend class FRemoteControlPresetRebindingManager;
 };
diff --git a/Engine/Plugins/VirtualProduction/RemoteControl/Source/RemoteControl/Public/RemoteControlPreset.h b/Engine/Plugins/VirtualProduction/RemoteControl/Source/RemoteControl/Public/RemoteControlPreset.h
--- a/Engine/Plugins/VirtualProduction/RemoteControl/Source/RemoteControl/Public/RemoteControlPreset.h
+++ b/Engine/Plugins/VirtualProduction/RemoteControl/Source/RemoteControl/Public/RemoteControlPreset.h
@@ -735,6 +735,11 @@ public:
 	 * Attempt to rebind all currently unbound properties.
 	 */
 	void RebindUnboundEntities();
+	
+	/**
+	 * Given a RC Entity, rebind all entities with the same owner to a new actor.
+	 */
+	void RebindAllEntitiesUnderSameActor(const FGuid& EntityId, AActor* NewActor);
 
 	DECLARE_MULTICAST_DELEGATE_TwoParams(FOnPresetEntityEvent, URemoteControlPreset* /*Preset*/, const FGuid& /*EntityId*/);
 	FOnPresetEntityEvent& OnEntityExposed() { return OnEntityExposedDelegate; }
diff --git a/Engine/Plugins/VirtualProduction/RemoteControl/Source/RemoteControlUI/Private/UI/SRCPanelExposedEntity.cpp b/Engine/Plugins/VirtualProduction/RemoteControl/Source/RemoteControlUI/Private/UI/SRCPanelExposedEntity.cpp
--- a/Engine/Plugins/VirtualProduction/RemoteControl/Source/RemoteControlUI/Private/UI/SRCPanelExposedEntity.cpp
+++ b/Engine/Plugins/VirtualProduction/RemoteControl/Source/RemoteControlUI/Private/UI/SRCPanelExposedEntity.cpp
@@ -15,6 +15,7 @@
 #include "SceneOutlinerFilters.h"
 #include "SceneOutlinerModule.h"
 #include "ScopedTransaction.h"
+#include "Styling/SlateIconFinder.h"
 #include "Modules/ModuleManager.h"
 #include "Widgets/Input/SButton.h"
 #include "Widgets/Input/SComboButton.h"
@@ -57,6 +58,24 @@ TSharedPtr<SWidget> SRCPanelExposedEntity::GetContextMenu()
 			constexpr bool bNoIndent = true;
 			SubMenuBuilder.AddWidget(CreateRebindMenuContent(), FText::GetEmpty(), bNoIndent);
 		}));
+
+	MenuBuilder.AddSubMenu(
+		LOCTEXT("EntityRebindComponentSubmenuLabel", "Rebind Component"),
+		LOCTEXT("EntityRebindComponentSubmenuToolTip", "Pick a component to rebind this exposed entity."),
+		FNewMenuDelegate::CreateLambda([this](FMenuBuilder& SubMenuBuilder)
+			{
+				CreateRebindComponentMenuContent(SubMenuBuilder);
+			}));
+
+	MenuBuilder.AddSubMenu(
+		LOCTEXT("EntityRebindAllUnderActorSubmenuLabel", "Rebind all properties for this actor"),
+		LOCTEXT("EntityRebindAllUnderActorSubmenuToolTip", "Pick an actor to rebind."),
+		FNewMenuDelegate::CreateLambda([this](FMenuBuilder& SubMenuBuilder)
+			{
+				constexpr bool bNoIndent = true;
+				SubMenuBuilder.AddWidget(CreateRebindAllPropertiesForActorMenuContent(), FText::GetEmpty(), bNoIndent);
+			}));
+
 	return MenuBuilder.MakeWidget();
 }
 
@@ -75,6 +94,61 @@ void SRCPanelExposedEntity::Initialize(const FGuid& InEntityId, URemoteControlPr
 	}
 }
 
+void SRCPanelExposedEntity::CreateRebindComponentMenuContent(FMenuBuilder& SubMenuBuilder)
+{
+	TInlineComponentArray<UActorComponent*> ComponentArray;
+
+	if (TSharedPtr<FRemoteControlEntity> Entity = GetEntity())
+	{
+		TArray<UObject*> BoundObjects = Entity->GetBoundObjects();
+		if (BoundObjects.Num() && BoundObjects[0])
+		{
+			if (AActor* OwnerActor = BoundObjects[0]->GetTypedOuter<AActor>())
+			{
+				OwnerActor->GetComponents(Entity->GetSupportedBindingClass(), ComponentArray);
+			}
+		}
+
+		for (UActorComponent* Component : ComponentArray)
+		{
+			SubMenuBuilder.AddMenuEntry(
+				FText::FromString(Component->GetName()),
+				FText::GetEmpty(),
+				FSlateIconFinder::FindIconForClass(Component->GetClass(), TEXT("SCS.Component")),
+				FUIAction(
+					FExecuteAction::CreateLambda([Entity, Component]
+						{
+							if (Entity)
+							{
+								Entity->BindObject(Component);
+							}
+						}),
+					FCanExecuteAction())
+			);
+		}
+	}
+}
+
+TSharedRef<SWidget> SRCPanelExposedEntity::CreateRebindAllPropertiesForActorMenuContent()
+{
+	FSceneOutlinerModule& SceneOutlinerModule = FModuleManager::Get().LoadModuleChecked<FSceneOutlinerModule>("SceneOutliner");
+	FSceneOutlinerInitializationOptions Options;
+	Options.bFocusSearchBoxWhenOpened = true;
+	Options.Filters = MakeShared<FSceneOutlinerFilters>();
+
+	if (TSharedPtr<FRemoteControlEntity> Entity = GetEntity())
+	{
+		Options.Filters->AddFilterPredicate<FActorTreeItem>(FActorTreeItem::FFilterPredicate::CreateRaw(this, &SRCPanelExposedEntity::IsActorSelectable));
+	}
+
+	return SNew(SBox)
+		.MaxDesiredHeight(400.0f)
+		.WidthOverride(300.0f)
+		[
+			SceneOutlinerModule.CreateActorPicker(Options, FOnActorPicked::CreateRaw(this, &SRCPanelExposedEntity::OnActorSelectedForRebindAllProperties))
+		];
+}
+
 TSharedRef<SWidget> SRCPanelExposedEntity::CreateInvalidWidget()
 {
 	return SNew(SHorizontalBox)
@@ -238,6 +312,17 @@ TSharedRef<SWidget> SRCPanelExposedEntity::CreateEntityWidget(TSharedPtr<SWidget
 	return Widget;
 }
 
+void SRCPanelExposedEntity::OnActorSelectedForRebindAllProperties(AActor* InActor) const
+{
+	if (TSharedPtr<FRemoteControlEntity> Entity = GetEntity())
+	{
+		Preset->RebindAllEntitiesUnderSameActor(Entity->GetId(), InActor);
+	}
+
+	FSlateApplication::Get().DismissAllMenus();
+}
+
+
 void SRCPanelExposedEntity::HandleUnexposeEntity()
 {
 	if (URemoteControlPreset* RCPreset = Preset.Get())
diff --git a/Engine/Plugins/VirtualProduction/RemoteControl/Source/RemoteControlUI/Private/UI/SRCPanelExposedEntity.h b/Engine/Plugins/VirtualProduction/RemoteControl/Source/RemoteControlUI/Private/UI/SRCPanelExposedEntity.h
--- a/Engine/Plugins/VirtualProduction/RemoteControl/Source/RemoteControlUI/Private/UI/SRCPanelExposedEntity.h
+++ b/Engine/Plugins/VirtualProduction/RemoteControl/Source/RemoteControlUI/Private/UI/SRCPanelExposedEntity.h
@@ -8,6 +8,7 @@
 #include "Widgets/Layout/SBorder.h"
 
 class AActor;
+class FMenuBuilder;
 struct FRemoteControlEntity;
 class SInlineEditableTextBlock;
 class SWidget;
@@ -51,10 +52,20 @@ protected:
 private:
 	/** Handles changing the object this entity is bound to upon selecting an actor in the rebinding dropdown. */
 	void OnActorSelected(AActor* InActor) const;
+	
 	/** Get the widget's border. */
 	const FSlateBrush* GetBorderImage() const;
 	/** Create the content of the rebind button.  */
 	TSharedRef<SWidget> CreateRebindMenuContent();
+
+	/** Create the content of the rebind component button. */
+	void CreateRebindComponentMenuContent(FMenuBuilder& SubMenuBuilder);
+
+	/** Create the content for the menu used to rebind all properties for the actor that owns this entity. */
+	TSharedRef<SWidget> CreateRebindAllPropertiesForActorMenuContent();
+
+	/** Handle selecting an actor for a rebind for all properties under an actor. */
+	void OnActorSelectedForRebindAllProperties(AActor* InActor) const;
 	/** Verifies that the entity's label doesn't already exist. */
 	bool OnVerifyItemLabelChanged(const FText& InLabel, FText& OutErrorMessage);
 	/** Handles committing a entity label. */
diff --git a/Engine/Plugins/VirtualProduction/RemoteControl/Source/RemoteControlUI/Private/UI/SRemoteControlPanel.h b/Engine/Plugins/VirtualProduction/RemoteControl/Source/RemoteControlUI/Private/UI/SRemoteControlPanel.h
--- a/Engine/Plugins/VirtualProduction/RemoteControl/Source/RemoteControlUI/Private/UI/SRemoteControlPanel.h
+++ b/Engine/Plugins/VirtualProduction/RemoteControl/Source/RemoteControlUI/Private/UI/SRemoteControlPanel.h
@@ -3,6 +3,7 @@
 
 #include "Engine/EngineTypes.h"
 #include "EditorUndoClient.h"
+#include "Engine/EngineTypes.h"
 #include "IDetailTreeNode.h"
 #include "Styling/SlateTypes.h"
 #include "RemoteControlFieldPath.h"
