RC Protocol toggle filter on add

diff --git a/Engine/Plugins/VirtualProduction/RemoteControl/Source/RemoteControlProtocolWidgets/Private/Widgets/SRCProtocolBindingList.cpp b/Engine/Plugins/VirtualProduction/RemoteControl/Source/RemoteControlProtocolWidgets/Private/Widgets/SRCProtocolBindingList.cpp
--- a/Engine/Plugins/VirtualProduction/RemoteControl/Source/RemoteControlProtocolWidgets/Private/Widgets/SRCProtocolBindingList.cpp
+++ b/Engine/Plugins/VirtualProduction/RemoteControl/Source/RemoteControlProtocolWidgets/Private/Widgets/SRCProtocolBindingList.cpp
@@ -122,7 +122,12 @@ void SRCProtocolBindingList::Construct(const FArguments& InArgs, TSharedRef<FPro
 	                {
 	                	if(ProtocolList.IsValid() && ProtocolList->GetSelectedProtocolName().IsValid())
 	                	{
-	                		ViewModel->AddBinding(*ProtocolList->GetSelectedProtocolName());
+	                		const FName& ProtocolName = *ProtocolList->GetSelectedProtocolName();
+	                		ViewModel->AddBinding(ProtocolName);
+	                		if (!IsProtocolShown(ProtocolName))
+	                		{
+	                			ToggleShowProtocol(ProtocolName);
+	                		}
 	                	}
                         return FReply::Handled();
 	                })
