RCP - Spotlight parameters cannot be controlled by RCP

diff --git a/Engine/Plugins/VirtualProduction/RemoteControl/Source/RemoteControlCommon/Public/RCPropertyUtilities.h b/Engine/Plugins/VirtualProduction/RemoteControl/Source/RemoteControlCommon/Public/RCPropertyUtilities.h
--- a/Engine/Plugins/VirtualProduction/RemoteControl/Source/RemoteControlCommon/Public/RCPropertyUtilities.h
+++ b/Engine/Plugins/VirtualProduction/RemoteControl/Source/RemoteControlCommon/Public/RCPropertyUtilities.h
@@ -474,6 +474,10 @@ namespace RemoteControlPropertyUtilities
 		{
 			SetterFunction = OwnerClass->FindFunctionByName(*SetterName);
 		}
+		else
+		{
+			return nullptr;
+		}
 #endif
 
 		FString PropertyName = Property->GetName();
