RCP bool fix

diff --git a/Engine/Plugins/VirtualProduction/RemoteControl/Source/RemoteControl/Private/RemoteControlProtocolBinding.cpp b/Engine/Plugins/VirtualProduction/RemoteControl/Source/RemoteControl/Private/RemoteControlProtocolBinding.cpp
--- a/Engine/Plugins/VirtualProduction/RemoteControl/Source/RemoteControl/Private/RemoteControlProtocolBinding.cpp
+++ b/Engine/Plugins/VirtualProduction/RemoteControl/Source/RemoteControl/Private/RemoteControlProtocolBinding.cpp
@@ -207,7 +207,8 @@ namespace EntityInterpolation
 		if (FBoolProperty* BoolProperty = CastField<FBoolProperty>(InProperty))
 		{
 			bool bBoolValue = false;
-			bSuccess = InterpolateValue(BoolProperty, OuterProperty, InRangeMappingBuffers, InProtocolValue, bBoolValue, InArrayIndex);
+			bSuccess = InterpolateValue(BoolProperty, OuterProperty, InRangeMappingBuffers, InProtocolValue, bBoolValue, InArrayIndex);			
+			bBoolValue = StaticCast<uint8>(bBoolValue) > 0; // Ensure 0 or 1, can be different if property was packed.
 			WritePropertyValue(InCborWriter, InProperty, bBoolValue, !bIsInArray);
 		}
 		else if (FNumericProperty* NumericProperty = CastField<FNumericProperty>(InProperty))
