RC Preset protocol range color picker updates

diff --git a/Engine/Plugins/VirtualProduction/RemoteControl/Source/RemoteControlCommon/Public/RCTypeTraits.h b/Engine/Plugins/VirtualProduction/RemoteControl/Source/RemoteControlCommon/Public/RCTypeTraits.h
--- a/Engine/Plugins/VirtualProduction/RemoteControl/Source/RemoteControlCommon/Public/RCTypeTraits.h
+++ b/Engine/Plugins/VirtualProduction/RemoteControl/Source/RemoteControlCommon/Public/RCTypeTraits.h
@@ -689,9 +689,10 @@ struct TRemoteControlTypeTraits<FColor>
 		
 	static Type DefaultMappingValueMax()
 	{
-		return Type(TRemoteControlTypeTraits<uint8>::DefaultMappingValueMax(),
-			TRemoteControlTypeTraits<uint8>::DefaultMappingValueMax(),
-			TRemoteControlTypeTraits<uint8>::DefaultMappingValueMax());
+		// @note: the use of Range instead of Mapping is intentional, by default Mapping is 0-1, but in this case we want 0-255
+		return Type(TRemoteControlTypeTraits<uint8>::DefaultRangeValueMax(),
+			TRemoteControlTypeTraits<uint8>::DefaultRangeValueMax(),
+			TRemoteControlTypeTraits<uint8>::DefaultRangeValueMax());
 	}
 };
 
diff --git a/Engine/Plugins/VirtualProduction/RemoteControl/Source/RemoteControlProtocolWidgets/Private/Widgets/SPropertyView.cpp b/Engine/Plugins/VirtualProduction/RemoteControl/Source/RemoteControlProtocolWidgets/Private/Widgets/SPropertyView.cpp
--- a/Engine/Plugins/VirtualProduction/RemoteControl/Source/RemoteControlProtocolWidgets/Private/Widgets/SPropertyView.cpp
+++ b/Engine/Plugins/VirtualProduction/RemoteControl/Source/RemoteControlProtocolWidgets/Private/Widgets/SPropertyView.cpp
@@ -73,7 +73,6 @@ int32 SPropertyView::DesiredWidth = 400.f;
 void SPropertyView::Construct(const FArguments& InArgs)
 {
 	bHasCustomPrepass = true;
-	bRefreshObjectToDisplay = false;
 	Object = TStrongObjectPtr<UObject>(InArgs._Object);
 	RootPropertyName = InArgs._RootPropertyName;
 	NameVisibility = InArgs._NameVisibility;
@@ -209,31 +208,9 @@ TSharedPtr<IPropertyHandle> SPropertyView::GetPropertyHandle() const
 	return Property;
 }
 
-bool SPropertyView::CustomPrepass(float InLayoutScaleMultiplier)
-{
-	if (bRefreshObjectToDisplay)
-	{
-		if (Object.IsValid())
-		{
-			Generator->SetObjects({Object.Get()});
-		}
-		else if (Struct.IsValid())
-		{
-			Generator->SetStructure(Struct);
-		}
-		ConstructInternal();
-		bRefreshObjectToDisplay = false;
-	}
-
-	return true;
-}
-
 void SPropertyView::Refresh()
 {
-	// forces CustomPrepass to be called, recreating widgets ie. for array items. Without this array items won't add/remove.
-	// @note: that this breaks current property handle references
 	MarkPrepassAsDirty();
-	bRefreshObjectToDisplay = true;
 }
 
 void SPropertyView::AddWidgets(const TArray<TSharedRef<IDetailTreeNode>>& InDetailTree, int32& InIndex, float InLeftPadding)
diff --git a/Engine/Plugins/VirtualProduction/RemoteControl/Source/RemoteControlProtocolWidgets/Private/Widgets/SPropertyView.h b/Engine/Plugins/VirtualProduction/RemoteControl/Source/RemoteControlProtocolWidgets/Private/Widgets/SPropertyView.h
--- a/Engine/Plugins/VirtualProduction/RemoteControl/Source/RemoteControlProtocolWidgets/Private/Widgets/SPropertyView.h
+++ b/Engine/Plugins/VirtualProduction/RemoteControl/Source/RemoteControlProtocolWidgets/Private/Widgets/SPropertyView.h
@@ -122,9 +122,6 @@ public:
 	FOnFinishedChangingProperties& OnFinishedChangingProperties() const { return Generator->OnFinishedChangingProperties(); }
 
 protected:
-	// ueent_hotfix Hack for 4.24 allow to refresh the ui in between two frame without any flickering
-	virtual bool CustomPrepass(float InLayoutScaleMultiplier) override;
-	
 	// Begin SWidget overrides.
 	// #ueent_todo: This is temporary until we find a better solution to the splitter issue
 	//				See SConstrainedBox's trick in cpp file
@@ -223,9 +220,6 @@ private:
 	/** Points to the currently used column size data. Can be provided via argument as well. */
 	TSharedPtr<RemoteControlProtocolWidgetUtils::FPropertyViewColumnSizeData> ColumnSizeData;
 
-	/** If there is a new object to display on the next tick */
-	uint8 bRefreshObjectToDisplay : 1;
-
 	/** Spacing between rows. 0 by default. */
 	float Spacing = 0;
 
