- Fixing crash in remote control when map is changed and webapp connected. Assumption was made that bindings were valid.

diff --git a/Engine/Plugins/VirtualProduction/RemoteControl/Source/WebRemoteControl/Private/WebRemoteControl.cpp b/Engine/Plugins/VirtualProduction/RemoteControl/Source/WebRemoteControl/Private/WebRemoteControl.cpp
--- a/Engine/Plugins/VirtualProduction/RemoteControl/Source/WebRemoteControl/Private/WebRemoteControl.cpp
+++ b/Engine/Plugins/VirtualProduction/RemoteControl/Source/WebRemoteControl/Private/WebRemoteControl.cpp
@@ -1097,6 +1097,14 @@ bool FWebRemoteControlModule::HandlePresetGetPropertyRoute(const FHttpServerRequ
 		return true;
 	}
 
+	if (!RemoteControlProperty->IsBound())
+	{
+		Response->Code = EHttpServerResponseCodes::NotFound;
+		WebRemoteControlUtils::CreateUTF8ErrorMessage(TEXT("Exposed entity was found but not bound to any object."), Response->Body);
+		OnComplete(MoveTemp(Response));
+		return true;
+	}
+
 	TArray<uint8> WorkingBuffer;
 	FMemoryWriter Writer(WorkingBuffer);
 
diff --git a/Engine/Plugins/VirtualProduction/RemoteControl/Source/WebRemoteControl/Private/WebSocketMessageHandler.cpp b/Engine/Plugins/VirtualProduction/RemoteControl/Source/WebRemoteControl/Private/WebSocketMessageHandler.cpp
--- a/Engine/Plugins/VirtualProduction/RemoteControl/Source/WebRemoteControl/Private/WebSocketMessageHandler.cpp
+++ b/Engine/Plugins/VirtualProduction/RemoteControl/Source/WebRemoteControl/Private/WebSocketMessageHandler.cpp
@@ -804,10 +804,13 @@ bool FWebSocketMessageHandler::WritePropertyChangeEventPayload(URemoteControlPre
 		FRCObjectReference ObjectRef;
 		if (TSharedPtr<FRemoteControlProperty> RCProperty = InPreset->GetExposedEntity<FRemoteControlProperty>(RCPropertyId).Pin())
 		{
-			if (IRemoteControlModule::Get().ResolveObjectProperty(ERCAccess::READ_ACCESS, RCProperty->GetBoundObjects()[0], RCProperty->FieldPathInfo.ToString(), ObjectRef))
+			if (RCProperty->IsBound())
 			{
-				bHasProperty = true;
-				PropValuesOnScope.Add(WebSocketMessageHandlerStructUtils::CreatePropertyValueOnScope(RCProperty, ObjectRef));
+				if (IRemoteControlModule::Get().ResolveObjectProperty(ERCAccess::READ_ACCESS, RCProperty->GetBoundObjects()[0], RCProperty->FieldPathInfo.ToString(), ObjectRef))
+				{
+					bHasProperty = true;
+					PropValuesOnScope.Add(WebSocketMessageHandlerStructUtils::CreatePropertyValueOnScope(RCProperty, ObjectRef));
+				}
 			}
 		}
 	}
