- Fixing remote control calls depending on objects having to be loaded
- Adding better blueprint native class filtering support for webapi asset search

diff --git a/Engine/Plugins/VirtualProduction/RemoteControl/Source/WebRemoteControl/Private/RemoteControlModels.h b/Engine/Plugins/VirtualProduction/RemoteControl/Source/WebRemoteControl/Private/RemoteControlModels.h
--- a/Engine/Plugins/VirtualProduction/RemoteControl/Source/WebRemoteControl/Private/RemoteControlModels.h
+++ b/Engine/Plugins/VirtualProduction/RemoteControl/Source/WebRemoteControl/Private/RemoteControlModels.h
@@ -669,6 +669,10 @@ struct FRCAssetFilter
 	UPROPERTY()
 	TSet<FName> RecursiveClassesExclusionSet;
 
+	/** Only if EnableBlueprintNativeClassFiltering is true, resulting asset will be filtered for dependants of classes in this list. */
+	UPROPERTY()
+	TArray<FName> NativeParentClasses;
+
 	/** If true, subclasses of ClassNames will also be included and RecursiveClassesExclusionSet will be excluded. */
 	UPROPERTY()
 	bool RecursiveClasses = false;
@@ -676,5 +680,8 @@ struct FRCAssetFilter
 	/** If true, PackagePath components will be recursive */
 	UPROPERTY()
 	bool RecursivePaths = false;
+
+	/** When dealing with blueprint classes, you might want to filter for a base class which can't be picked by asset registry if you derive from a blueprint class */
+	bool EnableBlueprintNativeClassFiltering = false;
 };
 
diff --git a/Engine/Plugins/VirtualProduction/RemoteControl/Source/WebRemoteControl/Private/Serialization/RCJsonStructDeserializerBackend.cpp b/Engine/Plugins/VirtualProduction/RemoteControl/Source/WebRemoteControl/Private/Serialization/RCJsonStructDeserializerBackend.cpp
--- a/Engine/Plugins/VirtualProduction/RemoteControl/Source/WebRemoteControl/Private/Serialization/RCJsonStructDeserializerBackend.cpp
+++ b/Engine/Plugins/VirtualProduction/RemoteControl/Source/WebRemoteControl/Private/Serialization/RCJsonStructDeserializerBackend.cpp
@@ -115,22 +115,66 @@ namespace
 		}
 		return false;
 	}
+
+	template<typename PropertyType>
+	static bool HandleObjectProperty(PropertyType* ObjectProperty, FProperty* Outer, void* Data, int32 ArrayIndex, const FString& StringValue)
+	{
+		constexpr UObject* ObjectOuter = nullptr;
+		constexpr TCHAR* Filename = nullptr;
+		UObject* LoadedObject = StaticLoadObject(ObjectProperty->PropertyClass, ObjectOuter, *StringValue, Filename, LOAD_NoWarn);
+		if (LoadedObject == nullptr)
+		{
+			UE_LOG(LogRemoteControl, Warning, TEXT("Deserialization error: Could not load object %s for property %s."), *StringValue, *ObjectProperty->GetName());
+		}
+
+		if (void* Ptr = GetPropertyValuePtr(ObjectProperty, Outer, Data, ArrayIndex))
+		{
+			ObjectProperty->SetPropertyValue(Ptr, typename PropertyType::TCppType(LoadedObject));
+			return true;
+		}
+
+		return false;
+	}
 }
 
 bool FRCJsonStructDeserializerBackend::ReadProperty(FProperty* Property, FProperty* Outer, void* Data, int32 ArrayIndex)
 {
+	//Whether we handled the property read or we need to user parent class version
+	TOptional<bool> HandledResult;
+
 	if (GetLastNotation() == EJsonNotation::String)
 	{
 		const FString& StringValue = GetReader()->GetValueAsString();
 
 		if (Property->IsA<FByteProperty>() || Property->IsA<FEnumProperty>())
 		{
-			const bool bReadResult = ReadEnum(Property, Outer, Data, ArrayIndex, StringValue);
-			if (!bReadResult)
+			HandledResult = ReadEnum(Property, Outer, Data, ArrayIndex, StringValue);
+			if (HandledResult.GetValue() == false)
 			{
 				UE_LOG(LogRemoteControl, Error, TEXT("Deserialization error: %s does not contain value %s"), *Property->GetName(), *StringValue);
 			}
-			return bReadResult;
+		}
+		else if (FObjectProperty* ObjectProperty = CastField<FObjectProperty>(Property))
+		{
+			if (IsInGameThread())
+			{
+				HandledResult = HandleObjectProperty(ObjectProperty, Outer, Data, ArrayIndex, StringValue);
+			}
+			else
+			{
+				UE_LOG(LogRemoteControl, Verbose, TEXT("ObjectProperty: Deserializing RemoteControl payload from another thread. Can't preload object %s in property %s."), *StringValue, *Property->GetName());
+			}
+		}
+		else if (FWeakObjectProperty* WeakObjectProperty = CastField<FWeakObjectProperty>(Property))
+		{
+			if (IsInGameThread())
+			{
+				HandledResult = HandleObjectProperty(WeakObjectProperty, Outer, Data, ArrayIndex, StringValue);
+			}
+			else
+			{
+				UE_LOG(LogRemoteControl, Verbose, TEXT("WeakObjectProperty: Deserializing RemoteControl payload from another thread. Can't preload object %s in property %s."), *StringValue, *Property->GetName());
+			}
 		}
 	}
 	else if (GetLastNotation() == EJsonNotation::Number)
@@ -143,16 +187,22 @@ bool FRCJsonStructDeserializerBackend::ReadProperty(FProperty* Property, FProper
 				if (DoubleValue > std::numeric_limits<float>::max())
 				{
 					FloatProperty->SetPropertyValue(Ptr, std::numeric_limits<float>::max());
-					return true;
+					HandledResult = true;
 				}
 			}
 			else
 			{
-				return false;
+				HandledResult = false;
 			}
 		}
 	}
 
-
-	return FJsonStructDeserializerBackend::ReadProperty(Property, Outer, Data, ArrayIndex);
+	if (HandledResult.IsSet() == false)
+	{
+		return FJsonStructDeserializerBackend::ReadProperty(Property, Outer, Data, ArrayIndex);
+	}
+	else
+	{
+		return HandledResult.GetValue();
+	}
 }
diff --git a/Engine/Plugins/VirtualProduction/RemoteControl/Source/WebRemoteControl/Private/WebRemoteControl.cpp b/Engine/Plugins/VirtualProduction/RemoteControl/Source/WebRemoteControl/Private/WebRemoteControl.cpp
--- a/Engine/Plugins/VirtualProduction/RemoteControl/Source/WebRemoteControl/Private/WebRemoteControl.cpp
+++ b/Engine/Plugins/VirtualProduction/RemoteControl/Source/WebRemoteControl/Private/WebRemoteControl.cpp
@@ -51,6 +51,7 @@
 #include "AssetRegistry/AssetRegistryModule.h"
 
 // Miscelleanous
+#include "Blueprint/BlueprintSupport.h"
 #include "Misc/App.h"
 #include "UObject/UnrealType.h"
 #include "Templates/UnrealTemplate.h"
@@ -202,6 +203,40 @@ namespace WebRemoteControl
 		// By default, web remote control is disabled in -game and packaged game.
 		return bIsEditor || FParse::Param(FCommandLine::Get(), TEXT("RCWebControlEnable"));
 	}
+
+	void DoNativeClassBlueprintFilter(const FSearchAssetRequest& SearchAssetRequest, TArray<FAssetData>& OutFilteredAssets)
+	{
+		if (SearchAssetRequest.Filter.EnableBlueprintNativeClassFiltering && SearchAssetRequest.Filter.NativeParentClasses.Num() > 0)
+		{
+			TSet<UClass*> NativeClasses;
+			NativeClasses.Reserve(SearchAssetRequest.Filter.NativeParentClasses.Num());
+			for (FName NativeClass : SearchAssetRequest.Filter.NativeParentClasses)
+			{
+				if (UClass* Class = FindObject<UClass>(ANY_PACKAGE, *NativeClass.ToString()))
+				{
+					NativeClasses.Add(Class);
+				}
+			}
+
+			for (auto It = OutFilteredAssets.CreateIterator(); It; ++It)
+			{
+				const FAssetData& AssetData = *It;
+				const FString NativeParentClassPath = AssetData.GetTagValueRef<FString>(FBlueprintTags::NativeParentClassPath);
+				const FSoftClassPath ClassPath(NativeParentClassPath);
+				UClass* NativeParentClass = ClassPath.ResolveClass();
+				if (NativeParentClass)
+				{
+					for (UClass* ClassFilter : NativeClasses)
+					{
+						if (!NativeParentClass->IsChildOf(ClassFilter))
+						{
+							It.RemoveCurrent();
+						}
+					}
+				}
+			}
+		}
+	}
 }
 
 void FWebRemoteControlModule::StartupModule()
@@ -1396,10 +1431,15 @@ bool FWebRemoteControlModule::HandleSearchAssetRoute(const FHttpServerRequest& R
 
 	FAssetRegistryModule& AssetRegistryModule = FModuleManager::Get().LoadModuleChecked<FAssetRegistryModule>(TEXT("AssetRegistry"));
 	IAssetRegistry& AssetRegistry = AssetRegistryModule.Get();
+
 	FARFilter Filter = SearchAssetRequest.Filter.ToARFilter();
 
 	TArray<FAssetData> Assets;
 	AssetRegistryModule.Get().GetAssets(Filter, Assets);
+
+	//Do advanced blueprint filtering if required
+	WebRemoteControl::DoNativeClassBlueprintFilter(SearchAssetRequest, Assets);
+
 	TArrayView<FAssetData> AssetsView{Assets};
 	int32 ArrayEnd = FMath::Min(SearchAssetRequest.Limit, Assets.Num());
 
diff --git a/Engine/Plugins/VirtualProduction/RemoteControl/Source/WebRemoteControl/Private/WebRemoteControlUtils.cpp b/Engine/Plugins/VirtualProduction/RemoteControl/Source/WebRemoteControl/Private/WebRemoteControlUtils.cpp
--- a/Engine/Plugins/VirtualProduction/RemoteControl/Source/WebRemoteControl/Private/WebRemoteControlUtils.cpp
+++ b/Engine/Plugins/VirtualProduction/RemoteControl/Source/WebRemoteControl/Private/WebRemoteControlUtils.cpp
@@ -148,7 +148,7 @@ namespace RemotePayloadSerializer
 				Reader.Seek(ParametersDelimiters.BlockStart);
 				Reader.SetLimitSize(ParametersDelimiters.BlockEnd);
 
-				FJsonStructDeserializerBackend Backend(Reader);
+				FRCJsonStructDeserializerBackend Backend(Reader);
 				if (!FStructDeserializer::Deserialize((void*)OutCall.ParamStruct.GetStructMemory(), *const_cast<UStruct*>(OutCall.ParamStruct.GetStruct()), Backend, FStructDeserializerPolicies()))
 				{
 					ErrorText = TEXT("Parameters object improperly formatted.");
