1) Remote Control - Right-click Rebind is available outside of selected row

Rebind is available. It's a bit confusing to have this context menu not confined to the associated property row.

2) Remote API Panel Shortcut

Being outside Edit Mode does not lock the individual exposed properties as it did in 4.27. Currently, it only seems to affect group headings. For example, properties can still be removed, renamed, and sorted while outside of edit mode in 5.0.

diff --git a/Engine/Plugins/VirtualProduction/RemoteControl/Source/RemoteControlUI/Private/UI/SRCPanelExposedEntitiesList.cpp b/Engine/Plugins/VirtualProduction/RemoteControl/Source/RemoteControlUI/Private/UI/SRCPanelExposedEntitiesList.cpp
--- a/Engine/Plugins/VirtualProduction/RemoteControl/Source/RemoteControlUI/Private/UI/SRCPanelExposedEntitiesList.cpp
+++ b/Engine/Plugins/VirtualProduction/RemoteControl/Source/RemoteControlUI/Private/UI/SRCPanelExposedEntitiesList.cpp
@@ -44,7 +44,7 @@ void SRCPanelExposedEntitiesList::Construct(const FArguments& InArgs, URemoteCon
 			.OnGetChildren(this, &SRCPanelExposedEntitiesList::OnGetNodeChildren)
 			.OnSelectionChanged(this, &SRCPanelExposedEntitiesList::OnSelectionChanged)
 			.OnContextMenuOpening(this, &SRCPanelExposedEntitiesList::OnContextMenuOpening)
-			.ClearSelectionOnClick(false)
+			.ClearSelectionOnClick(true)
 	];
 
 	RegisterEvents();
@@ -164,7 +164,8 @@ void SRCPanelExposedEntitiesList::GenerateListWidgets()
 			Args.Preset = Preset.Get();
 			Args.WidgetRegistry = WidgetRegistry;
 			Args.ColumnSizeData = ColumnSizeData;
-			
+			Args.bIsInEditMode = bIsInEditMode;
+
 			FieldWidgetMap.Add(Entity->GetId(), FRemoteControlUIModule::Get().GenerateEntityWidget(Args));
 		}
 	}
diff --git a/Engine/Plugins/VirtualProduction/RemoteControl/Source/RemoteControlUI/Private/UI/SRCPanelExposedField.cpp b/Engine/Plugins/VirtualProduction/RemoteControl/Source/RemoteControlUI/Private/UI/SRCPanelExposedField.cpp
--- a/Engine/Plugins/VirtualProduction/RemoteControl/Source/RemoteControlUI/Private/UI/SRCPanelExposedField.cpp
+++ b/Engine/Plugins/VirtualProduction/RemoteControl/Source/RemoteControlUI/Private/UI/SRCPanelExposedField.cpp
@@ -59,7 +59,7 @@ namespace ExposedFieldUtils
 
 TSharedPtr<SRCPanelTreeNode> SRCPanelExposedField::MakeInstance(const FGenerateWidgetArgs& Args)
 {
-	return SNew(SRCPanelExposedField, StaticCastSharedPtr<FRemoteControlField>(Args.Entity), Args.ColumnSizeData, Args.WidgetRegistry).Preset(Args.Preset);
+	return SNew(SRCPanelExposedField, StaticCastSharedPtr<FRemoteControlField>(Args.Entity), Args.ColumnSizeData, Args.WidgetRegistry).Preset(Args.Preset).EditMode(Args.bIsInEditMode);
 }
 
 void SRCPanelExposedField::Construct(const FArguments& InArgs, TWeakPtr<FRemoteControlField> InField, FRCColumnSizeData InColumnSizeData, TWeakPtr<FRCPanelWidgetRegistry> InWidgetRegistry)
diff --git a/Engine/Plugins/VirtualProduction/RemoteControl/Source/RemoteControlUI/Private/UI/SRemoteControlPanel.h b/Engine/Plugins/VirtualProduction/RemoteControl/Source/RemoteControlUI/Private/UI/SRemoteControlPanel.h
--- a/Engine/Plugins/VirtualProduction/RemoteControl/Source/RemoteControlUI/Private/UI/SRemoteControlPanel.h
+++ b/Engine/Plugins/VirtualProduction/RemoteControl/Source/RemoteControlUI/Private/UI/SRemoteControlPanel.h
@@ -93,8 +93,6 @@ public:
 	void SetEditMode(bool bEditMode)
 	{
 		bIsInEditMode = bEditMode;
-	
-		OnEditModeChange.ExecuteIfBound(SharedThis(this), bIsInEditMode);
 	}
 
 	/**
