Fix: Remote Control - MIDI, DMX and OSC value range are not displayed

diff --git a/Engine/Plugins/VirtualProduction/RemoteControl/Source/RemoteControlProtocolWidgets/Private/Widgets/SPropertyView.cpp b/Engine/Plugins/VirtualProduction/RemoteControl/Source/RemoteControlProtocolWidgets/Private/Widgets/SPropertyView.cpp
--- a/Engine/Plugins/VirtualProduction/RemoteControl/Source/RemoteControlProtocolWidgets/Private/Widgets/SPropertyView.cpp
+++ b/Engine/Plugins/VirtualProduction/RemoteControl/Source/RemoteControlProtocolWidgets/Private/Widgets/SPropertyView.cpp
@@ -72,7 +72,6 @@ int32 SPropertyView::DesiredWidth = 400.f;
 
 void SPropertyView::Construct(const FArguments& InArgs)
 {
-	bHasCustomPrepass = true;
 	Object = TStrongObjectPtr<UObject>(InArgs._Object);
 	RootPropertyName = InArgs._RootPropertyName;
 	NameVisibility = InArgs._NameVisibility;
