BUG in RC and Material Parameter Collection. Ensure thrown when attempting to add parameter to MPC using Remote Control

diff --git a/Engine/Plugins/VirtualProduction/RemoteControl/Source/RemoteControlUI/Private/UI/SRemoteControlPanel.cpp b/Engine/Plugins/VirtualProduction/RemoteControl/Source/RemoteControlUI/Private/UI/SRemoteControlPanel.cpp
--- a/Engine/Plugins/VirtualProduction/RemoteControl/Source/RemoteControlUI/Private/UI/SRemoteControlPanel.cpp
+++ b/Engine/Plugins/VirtualProduction/RemoteControl/Source/RemoteControlUI/Private/UI/SRemoteControlPanel.cpp
@@ -951,7 +951,8 @@ void SRemoteControlPanel::UpdateEntityDetailsView(const TSharedPtr<SRCPanelTreeN
 
 	if (SelectedNode)
 	{
-		if (SelectedNode->GetRCType() != SRCPanelTreeNode::Group)
+		if (SelectedNode->GetRCType() != SRCPanelTreeNode::Group &&
+			SelectedNode->GetRCType() != SRCPanelTreeNode::FieldChild) // Field Child does not contain entity ID, that is why it should not be processed
 		{
 			TSharedPtr<FRemoteControlEntity> Entity = Preset->GetExposedEntity<FRemoteControlEntity>(SelectedNode->GetRCId()).Pin();
 			SelectedEntityPtr = RemoteControlPanelUtils::GetEntityOnScope(Entity, Preset->GetExposedEntityType(SelectedNode->GetRCId()));
