- Updating remote control row generator parameter to also show hidden properties to fix ndisplay configuration data not showing up
- Fixing crash when creating rebinding menu

diff --git a/Engine/Plugins/VirtualProduction/RemoteControl/Source/RemoteControlUI/Private/UI/RCPanelWidgetRegistry.cpp b/Engine/Plugins/VirtualProduction/RemoteControl/Source/RemoteControlUI/Private/UI/RCPanelWidgetRegistry.cpp
--- a/Engine/Plugins/VirtualProduction/RemoteControl/Source/RemoteControlUI/Private/UI/RCPanelWidgetRegistry.cpp
+++ b/Engine/Plugins/VirtualProduction/RemoteControl/Source/RemoteControlUI/Private/UI/RCPanelWidgetRegistry.cpp
@@ -109,7 +109,9 @@ TSharedPtr<IDetailTreeNode> FRCPanelWidgetRegistry::GetObjectTreeNode(UObject* I
 		// Since we must keep many PRG objects alive in order to access the handle data, validating the nodes each tick is very taxing.
 		// We can override the validation with a lambda since the validation function in PRG is not necessary for our implementation
 		auto ValidationLambda = ([](const FRootPropertyNodeList& PropertyNodeList) { return true; });
-		Generator = FModuleManager::GetModuleChecked<FPropertyEditorModule>("PropertyEditor").CreatePropertyRowGenerator(FPropertyRowGeneratorArgs());
+		FPropertyRowGeneratorArgs Args;
+		Args.bShouldShowHiddenProperties = true;
+		Generator = FModuleManager::GetModuleChecked<FPropertyEditorModule>("PropertyEditor").CreatePropertyRowGenerator(Args);
 		Generator->SetCustomValidatePropertyNodesFunction(FOnValidatePropertyRowGeneratorNodes::CreateLambda(MoveTemp(ValidationLambda)));
 		Generator->SetObjects({InObject});
 		ObjectToRowGenerator.Add(WeakObject, Generator);
diff --git a/Engine/Plugins/VirtualProduction/RemoteControl/Source/RemoteControlUI/Private/UI/SRCPanelExposedEntity.cpp b/Engine/Plugins/VirtualProduction/RemoteControl/Source/RemoteControlUI/Private/UI/SRCPanelExposedEntity.cpp
--- a/Engine/Plugins/VirtualProduction/RemoteControl/Source/RemoteControlUI/Private/UI/SRCPanelExposedEntity.cpp
+++ b/Engine/Plugins/VirtualProduction/RemoteControl/Source/RemoteControlUI/Private/UI/SRCPanelExposedEntity.cpp
@@ -2,25 +2,26 @@
 
 #include "SRCPanelExposedEntity.h"
 
+#include "ActorTreeItem.h"
+#include "Components/ActorComponent.h"
 #include "EditorFontGlyphs.h"
 #include "Framework/Application/SlateApplication.h"
 #include "Framework/MultiBox/MultiBoxBuilder.h"
-#include "ActorTreeItem.h"
 #include "GameFramework/Actor.h"
+#include "Modules/ModuleManager.h"
 #include "RemoteControlEntity.h"
-#include "RemoteControlPreset.h"
 #include "RemoteControlField.h"
 #include "RemoteControlPanelStyle.h"
-#include "SRCPanelDragHandle.h"
+#include "RemoteControlPreset.h"
 #include "SceneOutlinerFilters.h"
 #include "SceneOutlinerModule.h"
 #include "ScopedTransaction.h"
+#include "SRCPanelDragHandle.h"
 #include "Styling/SlateIconFinder.h"
-#include "Modules/ModuleManager.h"
 #include "Widgets/Input/SButton.h"
 #include "Widgets/Input/SComboButton.h"
-#include "Widgets/Layout/SBox.h"
 #include "Widgets/Layout/SBorder.h"
+#include "Widgets/Layout/SBox.h"
 #include "Widgets/SBoxPanel.h"
 #include "Widgets/Text/SInlineEditableTextBlock.h"
 
@@ -103,9 +104,12 @@ void SRCPanelExposedEntity::CreateRebindComponentMenuContent(FMenuBuilder& SubMe
 		TArray<UObject*> BoundObjects = Entity->GetBoundObjects();
 		if (BoundObjects.Num() && BoundObjects[0])
 		{
-			if (AActor* OwnerActor = BoundObjects[0]->GetTypedOuter<AActor>())
+			if (BoundObjects[0]->IsA<UActorComponent>())
 			{
-				OwnerActor->GetComponents(Entity->GetSupportedBindingClass(), ComponentArray);
+				if (AActor* OwnerActor = BoundObjects[0]->GetTypedOuter<AActor>())
+				{
+					OwnerActor->GetComponents(Entity->GetSupportedBindingClass(), ComponentArray);
+				}
 			}
 		}
 
