RC. Reset Material Asset in SRemoteControlPanel if material need recompile shades

diff --git a/Engine/Plugins/VirtualProduction/RemoteControl/Source/RemoteControlUI/Private/UI/SRemoteControlPanel.cpp b/Engine/Plugins/VirtualProduction/RemoteControl/Source/RemoteControlUI/Private/UI/SRemoteControlPanel.cpp
--- a/Engine/Plugins/VirtualProduction/RemoteControl/Source/RemoteControlUI/Private/UI/SRemoteControlPanel.cpp
+++ b/Engine/Plugins/VirtualProduction/RemoteControl/Source/RemoteControlUI/Private/UI/SRemoteControlPanel.cpp
@@ -1057,6 +1057,8 @@ FReply SRemoteControlPanel::OnClickSettingsButton()
 
 void SRemoteControlPanel::OnMaterialCompiled(UMaterialInterface* MaterialInterface)
 {
+	// Clear the widget cache on material compiled to make sure we have valid property nodes for IPropertyRowGenerator
+	WidgetRegistry->Clear();
 	Refresh();
 }
 
