BUG in RC and Material List. Slow details panel performance when lots of material slots are involved

diff --git a/Engine/Plugins/VirtualProduction/RemoteControl/Source/RemoteControl/Private/RemoteControlField.cpp b/Engine/Plugins/VirtualProduction/RemoteControl/Source/RemoteControl/Private/RemoteControlField.cpp
--- a/Engine/Plugins/VirtualProduction/RemoteControl/Source/RemoteControl/Private/RemoteControlField.cpp
+++ b/Engine/Plugins/VirtualProduction/RemoteControl/Source/RemoteControl/Private/RemoteControlField.cpp
@@ -311,6 +311,11 @@ bool FRemoteControlProperty::CheckIsBoundToPropertyPath(const FString& InPath) c
 	return FieldPathInfo.ToPathPropertyString() == InPath;
 }
 
+bool FRemoteControlProperty::CheckIsBoundToString(const FString& InPath) const
+{
+	return FieldPathInfo.ToString() == InPath;
+}
+
 bool FRemoteControlProperty::ContainsBoundObjects(TArray<UObject*> InObjects) const
 {
 	const TArray<UObject*> BoundObjects = GetBoundObjects();
diff --git a/Engine/Plugins/VirtualProduction/RemoteControl/Source/RemoteControl/Public/RemoteControlField.h b/Engine/Plugins/VirtualProduction/RemoteControl/Source/RemoteControl/Public/RemoteControlField.h
--- a/Engine/Plugins/VirtualProduction/RemoteControl/Source/RemoteControl/Public/RemoteControlField.h
+++ b/Engine/Plugins/VirtualProduction/RemoteControl/Source/RemoteControl/Public/RemoteControlField.h
@@ -135,6 +135,13 @@ public:
 	 */
 	virtual bool CheckIsBoundToPropertyPath(const FString& InPropertyPath) const;
 
+	/**
+	 * Check whether given property string is bound to the property
+	 * @param InPathString Given Property String, this is the path without duplicates
+	 * return true if the property bound to given path
+	 */
+	virtual bool CheckIsBoundToString(const FString& InPath) const;
+
 	/**
 	 * Check whether any of the given objects are bound to this property
 	 * @param InObjects Objects to check
diff --git a/Engine/Plugins/VirtualProduction/RemoteControl/Source/RemoteControlUI/Private/RemoteControlUIModule.cpp b/Engine/Plugins/VirtualProduction/RemoteControl/Source/RemoteControlUI/Private/RemoteControlUIModule.cpp
--- a/Engine/Plugins/VirtualProduction/RemoteControl/Source/RemoteControlUI/Private/RemoteControlUIModule.cpp
+++ b/Engine/Plugins/VirtualProduction/RemoteControl/Source/RemoteControlUI/Private/RemoteControlUIModule.cpp
@@ -5,6 +5,7 @@
 #include "AssetToolsModule.h"
 #include "AssetTools/RemoteControlPresetActions.h"
 #include "Commands/RemoteControlCommands.h"
+#include "Components/SkeletalMeshComponent.h"
 #include "EditorStyleSet.h"
 #include "Engine/StaticMesh.h"
 #include "Engine/SkeletalMesh.h"
@@ -13,6 +14,7 @@
 #include "Interfaces/IMainFrameModule.h"
 #include "Kismet2/ComponentEditorUtils.h"
 #include "MaterialEditor/DEditorParameterValue.h"
+#include "Materials/Material.h"
 #include "PropertyHandle.h"
 #include "RemoteControlActor.h"
 #include "RemoteControlField.h"
@@ -43,6 +45,90 @@ static const FName DetailsTabIdentifiers[] = {
 	"LevelEditorSelectionDetails4"
 };
 
+FRCExposesPropertyArgs::FRCExposesPropertyArgs()
+	: PropertyHandle(nullptr)
+	, OwnerObject(nullptr)
+	, PropertyPath(TEXT(""))
+	, Property(nullptr)
+	, Id(FGuid::NewGuid())
+{
+}
+
+FRCExposesPropertyArgs::FRCExposesPropertyArgs(const FOnGenerateGlobalRowExtensionArgs& InExtensionArgs)
+	: PropertyHandle(InExtensionArgs.PropertyHandle)
+	, OwnerObject(InExtensionArgs.OwnerObject)
+	, PropertyPath(InExtensionArgs.PropertyPath)
+	, Property(InExtensionArgs.Property)
+	, Id(FGuid::NewGuid())
+{
+}
+
+FRCExposesPropertyArgs::FRCExposesPropertyArgs(FOnGenerateGlobalRowExtensionArgs&& InExtensionArgs)
+	: PropertyHandle(InExtensionArgs.PropertyHandle)
+	, OwnerObject(InExtensionArgs.OwnerObject)
+	, PropertyPath(InExtensionArgs.PropertyPath)
+	, Property(InExtensionArgs.Property)
+	, Id(FGuid::NewGuid())
+{
+}
+
+FRCExposesPropertyArgs::FRCExposesPropertyArgs(TSharedPtr<IPropertyHandle>& InPropertyHandle)
+	: PropertyHandle(InPropertyHandle)
+	, Id(FGuid::NewGuid())
+{
+}
+
+FRCExposesPropertyArgs::FRCExposesPropertyArgs(UObject* InOwnerObject, const FString& InPropertyPath, FProperty* InProperty)
+	: OwnerObject(InOwnerObject)
+	, PropertyPath(InPropertyPath)
+	, Property(InProperty)
+	, Id(FGuid::NewGuid())
+{
+}
+
+FRCExposesPropertyArgs::EType FRCExposesPropertyArgs::GetType() const
+{
+	if (PropertyHandle.IsValid())
+	{
+		return EType::E_Handle;
+	}
+	if (OwnerObject && !PropertyPath.IsEmpty() && Property)
+	{
+		return EType::E_OwnerObject;
+	}
+
+	return EType::E_None;
+}
+
+bool FRCExposesPropertyArgs::IsValid() const
+{
+	const EType Type = GetType();
+	return Type == EType::E_Handle || Type == EType::E_OwnerObject;
+}
+
+
+FProperty* FRCExposesPropertyArgs::GetProperty() const
+{
+	const EType Type = GetType();
+	if (Type == EType::E_Handle)
+	{
+		PropertyHandle->GetProperty();
+	}
+	else
+	{
+		return Property;
+	}
+
+	return nullptr;
+}
+
+FProperty* FRCExposesPropertyArgs::GetPropertyChecked() const
+{
+	FProperty* PropertyChecked = GetProperty();
+	check(PropertyChecked);
+	return PropertyChecked;
+}
+
 namespace RemoteControlUIModule
 {
 	const TArray<FName>& GetCustomizedStructNames()
@@ -63,29 +149,52 @@ namespace RemoteControlUIModule
 		Toolkit->InitRemoteControlPresetEditor(Mode, EditWithinLevelEditor, Preset);
 	}
 
-	bool IsStaticOrSkeletalMaterial(TSharedPtr<IPropertyHandle> InParentPropertyHandle)
+	FObjectProperty* GetObjectProperty(const FRCExposesPropertyArgs& InPropertyArgs)
+	{
+		if (!ensureMsgf(InPropertyArgs.IsValid(), TEXT("Extension Property Args not valid")))
+		{
+			return nullptr;
+		}
+
+		FRCFieldPathInfo RCFieldPathInfo(InPropertyArgs.PropertyPath);
+		RCFieldPathInfo.Resolve(InPropertyArgs.OwnerObject);
+		FRCFieldResolvedData ResolvedData = RCFieldPathInfo.GetResolvedData();
+		return CastField<FObjectProperty>(InPropertyArgs.GetProperty());
+	}
+
+	bool IsStaticOrSkeletalMaterial(const FRCExposesPropertyArgs& InPropertyArgs)
 	{
-		if (!InParentPropertyHandle.IsValid() || !InParentPropertyHandle->IsValidHandle() || InParentPropertyHandle->GetNumOuterObjects() == 0)
+		if (!InPropertyArgs.IsValid())
 		{
 			return false;
 		}
 
-		if (FProperty* OwnerProperty = InParentPropertyHandle->GetProperty())
+		const FRCExposesPropertyArgs::EType ExtensionArgsType = InPropertyArgs.GetType();
+
+		// Only in case of Owner Object is applicable for MaterialInterface
+		if (ExtensionArgsType == FRCExposesPropertyArgs::EType::E_OwnerObject)
 		{
-			TArray<UObject*> OuterObjects;
+			FProperty* Property = InPropertyArgs.GetProperty();
 
-			InParentPropertyHandle->GetOuterObjects(OuterObjects);
+			if (!ensure(Property))
+			{
+				return false;
+			}
 
-			const bool bIsStaticMaterial = OwnerProperty->GetFName() == UStaticMesh::GetStaticMaterialsName() && OuterObjects[0]->GetClass()->IsChildOf(UStaticMesh::StaticClass());
-			
-			const bool bIsSkeletalMaterial = OwnerProperty->GetFName() == USkeletalMesh::GetMaterialsMemberName() && OuterObjects[0]->GetClass()->IsChildOf(USkeletalMesh::StaticClass());
+			FObjectProperty* ObjectProperty = GetObjectProperty(InPropertyArgs);
 
-			return bIsStaticMaterial || bIsSkeletalMaterial;
+			if (!ObjectProperty)
+			{
+				return false;
+			}
+
+			return Property->GetFName() == UMaterialInterface::StaticClass()->GetFName() && (InPropertyArgs.OwnerObject->GetClass()->IsChildOf(UStaticMesh::StaticClass()) || InPropertyArgs.OwnerObject->GetClass()->IsChildOf(USkeletalMesh::StaticClass()));
 		}
 
 		return false;
 	}
 
+
 	bool IsTransientObjectAllowListed(UObject* Object)
 	{
 		return Object && Object->IsA<UDEditorParameterValue>();
@@ -318,50 +427,51 @@ void FRemoteControlUIModule::UnbindRemoteControlCommands()
 
 void FRemoteControlUIModule::HandleCreatePropertyRowExtension(const FOnGenerateGlobalRowExtensionArgs& InArgs, TArray<FPropertyRowExtensionButton>& OutExtensions)
 {
-	// Expose/Unexpose button.
+	const FRCExposesPropertyArgs ExposesPropertyArgs(InArgs);
 
-	FPropertyRowExtensionButton& ExposeButton = OutExtensions.AddDefaulted_GetRef();
-	ExposeButton.Icon = TAttribute<FSlateIcon>::Create(
-		[this, PropertyHandle = InArgs.PropertyHandle]
+	if (ExposesPropertyArgs.IsValid())
+	{
+		// Expose/Unexpose button.
+		FPropertyRowExtensionButton& ExposeButton = OutExtensions.AddDefaulted_GetRef();
+		ExposeButton.Icon = TAttribute<FSlateIcon>::Create(
+			[this, ExposesPropertyArgs]
 		{
-			return OnGetExposedIcon(PropertyHandle);
+			return OnGetExposedIcon(ExposesPropertyArgs);
 		});
 
-	ExposeButton.Label = TAttribute<FText>::Create(TAttribute<FText>::FGetter::CreateRaw(this, &FRemoteControlUIModule::GetExposePropertyButtonText, InArgs.PropertyHandle));
-	ExposeButton.ToolTip = TAttribute<FText>::Create(TAttribute<FText>::FGetter::CreateRaw(this, &FRemoteControlUIModule::GetExposePropertyButtonTooltip, InArgs.PropertyHandle));
-	ExposeButton.UIAction = FUIAction(
-		FExecuteAction::CreateRaw(this, &FRemoteControlUIModule::OnToggleExposeProperty, InArgs.PropertyHandle),
-		FCanExecuteAction::CreateRaw(this, &FRemoteControlUIModule::CanToggleExposeProperty, InArgs.PropertyHandle),
-		FGetActionCheckState::CreateRaw(this, &FRemoteControlUIModule::GetPropertyExposedCheckState, InArgs.PropertyHandle),
-		FIsActionButtonVisible::CreateRaw(this, &FRemoteControlUIModule::CanToggleExposeProperty, InArgs.PropertyHandle)
-	);
-
-	// Override material(s) warning.
+		ExposeButton.Label = TAttribute<FText>::Create(TAttribute<FText>::FGetter::CreateRaw(this, &FRemoteControlUIModule::GetExposePropertyButtonText, ExposesPropertyArgs));
+		ExposeButton.ToolTip = TAttribute<FText>::Create(TAttribute<FText>::FGetter::CreateRaw(this, &FRemoteControlUIModule::GetExposePropertyButtonTooltip, ExposesPropertyArgs));
+		ExposeButton.UIAction = FUIAction(
+			FExecuteAction::CreateRaw(this, &FRemoteControlUIModule::OnToggleExposeProperty, ExposesPropertyArgs),
+			FCanExecuteAction::CreateRaw(this, &FRemoteControlUIModule::CanToggleExposeProperty, ExposesPropertyArgs),
+			FGetActionCheckState::CreateRaw(this, &FRemoteControlUIModule::GetPropertyExposedCheckState, ExposesPropertyArgs),
+			FIsActionButtonVisible::CreateRaw(this, &FRemoteControlUIModule::CanToggleExposeProperty, ExposesPropertyArgs)
+		);
 
-	FPropertyRowExtensionButton& OverrideMaterialButton = OutExtensions.AddDefaulted_GetRef();
+		// Override material(s) warning.
+		FPropertyRowExtensionButton& OverrideMaterialButton = OutExtensions.AddDefaulted_GetRef();
 
-	OverrideMaterialButton.Icon = TAttribute<FSlateIcon>::Create(
-		[this, PropertyHandle = InArgs.PropertyHandle]
+		OverrideMaterialButton.Icon = TAttribute<FSlateIcon>::Create(
+			[this, ExposesPropertyArgs]
 		{
-			return OnGetOverrideMaterialsIcon(PropertyHandle);
+			return OnGetOverrideMaterialsIcon(ExposesPropertyArgs);
 		}
 		);
 
-	OverrideMaterialButton.Label = LOCTEXT("OverrideMaterial", "Override Material");
-
-	OverrideMaterialButton.ToolTip = LOCTEXT("OverrideMaterialToolTip", "Click to override this material in order to expose this property to Remote Control.");
-
-	OverrideMaterialButton.UIAction = FUIAction(
-		FExecuteAction::CreateRaw(this, &FRemoteControlUIModule::TryOverridingMaterials, InArgs.PropertyHandle),
-		FCanExecuteAction::CreateRaw(this, &FRemoteControlUIModule::IsStaticOrSkeletalMaterialProperty, InArgs.PropertyHandle),
-		FGetActionCheckState(),
-		FIsActionButtonVisible::CreateRaw(this, &FRemoteControlUIModule::IsStaticOrSkeletalMaterialProperty, InArgs.PropertyHandle)
-	);
+		OverrideMaterialButton.Label = LOCTEXT("OverrideMaterial", "Override Material");
+		OverrideMaterialButton.ToolTip = LOCTEXT("OverrideMaterialToolTip", "Click to override this material in order to expose this property to Remote Control.");
+		OverrideMaterialButton.UIAction = FUIAction(
+			FExecuteAction::CreateRaw(this, &FRemoteControlUIModule::TryOverridingMaterials, ExposesPropertyArgs),
+			FCanExecuteAction::CreateRaw(this, &FRemoteControlUIModule::IsStaticOrSkeletalMaterialProperty, ExposesPropertyArgs),
+			FGetActionCheckState(),
+			FIsActionButtonVisible::CreateRaw(this, &FRemoteControlUIModule::IsStaticOrSkeletalMaterialProperty, ExposesPropertyArgs)
+		);
+	}
 
 	WeakDetailsTreeNode = InArgs.OwnerTreeNode;
 }
 
-FSlateIcon FRemoteControlUIModule::OnGetExposedIcon(TSharedPtr<IPropertyHandle> Handle) const
+FSlateIcon FRemoteControlUIModule::OnGetExposedIcon(const FRCExposesPropertyArgs& InPropertyArgs) const
 {
 	FName BrushName("NoBrush");
 
@@ -369,7 +479,7 @@ FSlateIcon FRemoteControlUIModule::OnGetExposedIcon(TSharedPtr<IPropertyHandle>
 	{
 		if (Panel->GetPreset())
 		{
-			EPropertyExposeStatus Status = GetPropertyExposeStatus(Handle);
+			EPropertyExposeStatus Status = GetPropertyExposeStatus(InPropertyArgs);
 			if (Status == EPropertyExposeStatus::Exposed)
 			{
 				BrushName = "Level.VisibleIcon16x";
@@ -384,23 +494,24 @@ FSlateIcon FRemoteControlUIModule::OnGetExposedIcon(TSharedPtr<IPropertyHandle>
 	return FSlateIcon(FAppStyle::Get().GetStyleSetName(), BrushName);
 }
 
-bool FRemoteControlUIModule::CanToggleExposeProperty(TSharedPtr<IPropertyHandle> Handle) const
+bool FRemoteControlUIModule::CanToggleExposeProperty(const FRCExposesPropertyArgs InPropertyArgs) const
 {
 	if (TSharedPtr<SRemoteControlPanel> Panel = WeakActivePanel.Pin())
 	{
-		return Handle && ShouldDisplayExposeIcon(Handle.ToSharedRef());
+		return InPropertyArgs.IsValid() && ShouldDisplayExposeIcon(InPropertyArgs);
 	}
 
 	return false;
 }
 
-ECheckBoxState FRemoteControlUIModule::GetPropertyExposedCheckState(TSharedPtr<IPropertyHandle> Handle) const
+
+ECheckBoxState FRemoteControlUIModule::GetPropertyExposedCheckState(const FRCExposesPropertyArgs InPropertyArgs) const
 {
 	if (TSharedPtr<SRemoteControlPanel> Panel = WeakActivePanel.Pin())
 	{
 		if (Panel->GetPreset())
 		{
-			EPropertyExposeStatus ExposeStatus = GetPropertyExposeStatus(Handle);
+			EPropertyExposeStatus ExposeStatus = GetPropertyExposeStatus(InPropertyArgs);
 			if (ExposeStatus == EPropertyExposeStatus::Exposed)
 			{
 				return ECheckBoxState::Checked;
@@ -411,37 +522,38 @@ ECheckBoxState FRemoteControlUIModule::GetPropertyExposedCheckState(TSharedPtr<I
 	return ECheckBoxState::Unchecked;
 }
 
-void FRemoteControlUIModule::OnToggleExposeProperty(TSharedPtr<IPropertyHandle> Handle)
+void FRemoteControlUIModule::OnToggleExposeProperty(const FRCExposesPropertyArgs InPropertyArgs)
 {
-	if (!ensureMsgf(Handle && Handle->IsValidHandle(), TEXT("Property could not be exposed because the handle was invalid.")))
+	if (!ensureMsgf(InPropertyArgs.IsValid(), TEXT("Property could not be exposed because the extension args was invalid.")))
 	{
 		return;
 	}
 
 	if (TSharedPtr<SRemoteControlPanel> Panel = WeakActivePanel.Pin())
 	{
-		Panel->ToggleProperty(Handle);
+		Panel->ToggleProperty(InPropertyArgs);
 	}
 }
 
-FRemoteControlUIModule::EPropertyExposeStatus FRemoteControlUIModule::GetPropertyExposeStatus(const TSharedPtr<IPropertyHandle>& Handle) const
+FRemoteControlUIModule::EPropertyExposeStatus FRemoteControlUIModule::GetPropertyExposeStatus(const FRCExposesPropertyArgs& InPropertyArgs) const
 {
-	if (Handle && Handle->IsValidHandle())
+	if (InPropertyArgs.IsValid())
 	{
 		if (TSharedPtr<SRemoteControlPanel> Panel = WeakActivePanel.Pin())
 		{
-			return Panel->IsExposed(Handle) ? EPropertyExposeStatus::Exposed : EPropertyExposeStatus::Unexposed;
+			return Panel->IsExposed(InPropertyArgs) ? EPropertyExposeStatus::Exposed : EPropertyExposeStatus::Unexposed;
 		}
 	}
 
 	return EPropertyExposeStatus::Unexposable;
 }
 
-FSlateIcon FRemoteControlUIModule::OnGetOverrideMaterialsIcon(TSharedPtr<IPropertyHandle> Handle) const
+
+FSlateIcon FRemoteControlUIModule::OnGetOverrideMaterialsIcon(const FRCExposesPropertyArgs& InPropertyArgs) const
 {
 	FName BrushName("NoBrush");
 
-	if (IsStaticOrSkeletalMaterialProperty(Handle))
+	if (IsStaticOrSkeletalMaterialProperty(InPropertyArgs))
 	{
 		BrushName = "Icons.Warning";
 	}
@@ -449,22 +561,21 @@ FSlateIcon FRemoteControlUIModule::OnGetOverrideMaterialsIcon(TSharedPtr<IProper
 	return FSlateIcon(FAppStyle::Get().GetStyleSetName(), BrushName);
 }
 
-bool FRemoteControlUIModule::IsStaticOrSkeletalMaterialProperty(TSharedPtr<IPropertyHandle> Handle) const
+
+bool FRemoteControlUIModule::IsStaticOrSkeletalMaterialProperty(const FRCExposesPropertyArgs InPropertyArgs) const
 {
 	if (TSharedPtr<SRemoteControlPanel> Panel = WeakActivePanel.Pin()) // Check whether the panel is active.
 	{
-		if (Panel->GetPreset() && Handle && Handle->IsValidHandle()) // Ensure that we have a valid preset and handle.
+		if (Panel->GetPreset() && InPropertyArgs.IsValid()) // Ensure that we have a valid preset and handle.
 		{
-			if (TSharedPtr<IPropertyHandle> ParentHandle = Handle->GetParentHandle())
-			{
-				return RemoteControlUIModule::IsStaticOrSkeletalMaterial(ParentHandle);
-			}
+			return RemoteControlUIModule::IsStaticOrSkeletalMaterial(InPropertyArgs);
 		}
 	}
 
 	return false;
 }
 
+
 void FRemoteControlUIModule::AddGetPathOption(FMenuBuilder& MenuBuilder, AActor* SelectedActor)
 {
 	auto CopyLambda = [SelectedActor]()
@@ -495,61 +606,122 @@ TSharedRef<FExtender> FRemoteControlUIModule::ExtendLevelViewportContextMenuForR
 	return Extender.ToSharedRef();
 }
 
-bool FRemoteControlUIModule::ShouldDisplayExposeIcon(const TSharedRef<IPropertyHandle>& PropertyHandle) const
+bool FRemoteControlUIModule::ShouldDisplayExposeIcon(const FRCExposesPropertyArgs& InPropertyArgs) const
 {
-	if (FProperty* Prop = PropertyHandle->GetProperty())
+	if (!InPropertyArgs.IsValid())
 	{
-		// Don't display an expose icon for RCEntities since they're only displayed in the Remote Control Panel.
-		if (Prop->GetOwnerStruct() && Prop->GetOwnerStruct()->IsChildOf(FRemoteControlEntity::StaticStruct()))
-		{
-			return false;
-		}
+		return false;
+	}
 
-		// Don't display an expose icon for Static or Skeletal Materials as they need to be overriden.
-		if (TSharedPtr<IPropertyHandle> ParentHandle = PropertyHandle->GetParentHandle())
-		{
-			if (RemoteControlUIModule::IsStaticOrSkeletalMaterial(ParentHandle))
-			{
-				return false;
-			}
-		}
+
+	if (RemoteControlUIModule::IsStaticOrSkeletalMaterial(InPropertyArgs))
+	{
+		return false;
+	}
+
+	const FRCExposesPropertyArgs::EType ExtensionArgsType = InPropertyArgs.GetType();
+
+	if (ExtensionArgsType == FRCExposesPropertyArgs::EType::E_Handle)
+	{
+		TSharedPtr<IPropertyHandle> PropertyHandle = InPropertyArgs.PropertyHandle;
 
 		if (PropertyHandle->GetNumOuterObjects() == 1)
 		{
 			TArray<UObject*> OuterObjects;
 			PropertyHandle->GetOuterObjects(OuterObjects);
 
-			if (OuterObjects[0])
+			if (!IsAllowedOwnerObjects(OuterObjects))
 			{
-				// Don't display an expose icon for default objects.
-				if (OuterObjects[0]->HasAnyFlags(RF_ClassDefaultObject | RF_ArchetypeObject))
-				{
-					return false;
-				}
-
-				// Don't display an expose icon for transient objects such as material editor parameters.
-				if (OuterObjects[0]->GetOutermost()->HasAnyFlags(RF_Transient) && !RemoteControlUIModule::IsTransientObjectAllowListed(OuterObjects[0]))
-				{
-					return false;
-				}
+				return false;
 			}
 		}
 	}
+	else if (ExtensionArgsType == FRCExposesPropertyArgs::EType::E_OwnerObject)
+	{
+		if (!IsAllowedOwnerObjects({ InPropertyArgs.OwnerObject }))
+		{
+			return false;
+		}
+	}
+
 
 	for (const TPair<FDelegateHandle, FOnDisplayExposeIcon>& DelegatePair : ExternalFilterDelegates)
 	{
 		if (DelegatePair.Value.IsBound())
 		{
-			if (!DelegatePair.Value.Execute(PropertyHandle))
+			if (!DelegatePair.Value.Execute(InPropertyArgs))
 			{
 				return false;
 			}
 		}
 	}
 
+
+	return true;
+}
+
+bool FRemoteControlUIModule::ShouldSkipOwnPanelProperty(const FRCExposesPropertyArgs& InPropertyArgs) const
+{
+	if (!InPropertyArgs.IsValid())
+	{
+		return false;
+	}
+
+	auto CheckOwnProperty = [](FProperty* InProp)
+	{
+		if (!InProp)
+		{
+			return false;
+		}
+
+
+		// Don't display an expose icon for RCEntities since they're only displayed in the Remote Control Panel.
+		if (InProp->GetOwnerStruct() && InProp->GetOwnerStruct()->IsChildOf(FRemoteControlEntity::StaticStruct()))
+		{
+			return true;
+		}
+
+		return false;
+	};
+
+	const FRCExposesPropertyArgs::EType ArgsType = InPropertyArgs.GetType();
+
+	if (ArgsType == FRCExposesPropertyArgs::EType::E_Handle)
+	{
+		if (TSharedPtr<IPropertyHandle> PropertyHandle = InPropertyArgs.PropertyHandle)
+		{
+			CheckOwnProperty(PropertyHandle->GetProperty());
+		}
+	}
+	else if (ArgsType == FRCExposesPropertyArgs::EType::E_OwnerObject)
+	{
+		return CheckOwnProperty(InPropertyArgs.Property);
+	}
+
+	return false;
+}
+
+bool FRemoteControlUIModule::IsAllowedOwnerObjects(TArray<UObject*> InOuterObjects) const
+{
+	if (InOuterObjects[0])
+	{
+		// Don't display an expose icon for default objects.
+		if (InOuterObjects[0]->HasAnyFlags(RF_ClassDefaultObject | RF_ArchetypeObject))
+		{
+			return false;
+		}
+
+		// Don't display an expose icon for transient objects such as material editor parameters.
+		if (InOuterObjects[0]->GetOutermost()->HasAnyFlags(RF_Transient) && !RemoteControlUIModule::IsTransientObjectAllowListed(InOuterObjects[0]))
+		{
+			return false;
+		}
+	}
+
 	return true;
 }
 
+
 void FRemoteControlUIModule::RegisterStructCustomizations()
 {
 	FPropertyEditorModule& PropertyEditorModule = FModuleManager::GetModuleChecked<FPropertyEditorModule>("PropertyEditor");
@@ -616,12 +788,12 @@ void FRemoteControlUIModule::RegisterWidgetFactories()
 	RegisterWidgetFactoryForType(FRemoteControlInstanceMaterial::StaticStruct(), FOnGenerateRCWidget::CreateStatic(&SRCPanelExposedField::MakeInstance));
 }
 
-FText FRemoteControlUIModule::GetExposePropertyButtonTooltip(TSharedPtr<IPropertyHandle> Handle) const
+FText FRemoteControlUIModule::GetExposePropertyButtonTooltip(const FRCExposesPropertyArgs InPropertyArgs) const
 {
 	if (URemoteControlPreset* Preset = GetActivePreset())
 	{
 		const FText PresetName = FText::FromString(Preset->GetName());
-		if (GetPropertyExposeStatus(Handle) == EPropertyExposeStatus::Exposed)
+		if (GetPropertyExposeStatus(InPropertyArgs) == EPropertyExposeStatus::Exposed)
 		{
 			return FText::Format(LOCTEXT("ExposePropertyToolTip", "Unexpose this property from RemoteControl Preset '{0}'."), PresetName);
 		}
@@ -634,9 +806,10 @@ FText FRemoteControlUIModule::GetExposePropertyButtonTooltip(TSharedPtr<IPropert
 	return LOCTEXT("InvalidExposePropertyTooltip", "Invalid Preset");
 }
 
-FText FRemoteControlUIModule::GetExposePropertyButtonText(TSharedPtr<IPropertyHandle> Handle) const
+
+FText FRemoteControlUIModule::GetExposePropertyButtonText(const FRCExposesPropertyArgs InPropertyArgs) const
 {
-	if (GetPropertyExposeStatus(Handle) == EPropertyExposeStatus::Exposed)
+	if (GetPropertyExposeStatus(InPropertyArgs) == EPropertyExposeStatus::Exposed)
 	{
 		return LOCTEXT("ExposePropertyToolTip", "Unexpose property");
 	}
@@ -646,157 +819,220 @@ FText FRemoteControlUIModule::GetExposePropertyButtonText(TSharedPtr<IPropertyHa
 	}
 }
 
-void FRemoteControlUIModule::TryOverridingMaterials(TSharedPtr<IPropertyHandle> ForThisProperty)
+void FRemoteControlUIModule::TryOverridingMaterials(const FRCExposesPropertyArgs InPropertyArgs)
 {
-	if (!ensureMsgf(ForThisProperty && ForThisProperty->IsValidHandle(), TEXT("Property could not be exposed because the handle was invalid.")))
+	if (!ensureMsgf(InPropertyArgs.IsValid(), TEXT("Property could not be exposed because the handle was invalid.")))
 	{
 		return;
 	}
 
-	UObject* OriginalMaterialAsObject;
+	const FRCExposesPropertyArgs::EType ExtensionArgsType = InPropertyArgs.GetType();
 
-	if (ForThisProperty->GetValue(OriginalMaterialAsObject) == FPropertyAccess::Success)
+	if (ExtensionArgsType == FRCExposesPropertyArgs::EType::E_Handle)
 	{
-		if (UMaterialInterface* OriginalMaterial = Cast<UMaterialInterface>(OriginalMaterialAsObject))
-		{
-			FName MaterialSlotNameToBeOverriden;
+		UObject* OriginalMaterialAsObject = nullptr;
+		TSharedPtr<IPropertyHandle> PropertyHandle = InPropertyArgs.PropertyHandle;
 
-			// NOTE : Obtain the source object information from the property.
-
-			if (ForThisProperty->GetNumOuterObjects() == 1)
+		if (PropertyHandle->GetValue(OriginalMaterialAsObject) == FPropertyAccess::Success)
+		{
+			if (UMaterialInterface* OriginalMaterial = Cast<UMaterialInterface>(OriginalMaterialAsObject))
 			{
-				TArray<UObject*> OuterObjects;
+				FName MaterialSlotNameToBeOverriden;
 
-				ForThisProperty->GetOuterObjects(OuterObjects);
+				// NOTE : Obtain the source object information from the property.
 
-				if (UStaticMesh* StaticMesh = Cast<UStaticMesh>(OuterObjects[0]))
+				if (PropertyHandle->GetNumOuterObjects() == 1)
 				{
-					// NOTE : As 'StaticMaterials' must be protected for async build, always use the accessors even internally.
+					TArray<UObject*> OuterObjects;
 
-					for (FStaticMaterial StaticMaterial : StaticMesh->GetStaticMaterials())
+					PropertyHandle->GetOuterObjects(OuterObjects);
+
+					check(OuterObjects.Num());
+
+					if (UStaticMesh* StaticMesh = Cast<UStaticMesh>(OuterObjects[0]))
 					{
-						if (StaticMaterial.MaterialInterface == OriginalMaterial)
+						// NOTE : As 'StaticMaterials' must be protected for async build, always use the accessors even internally.
+						for (const FStaticMaterial StaticMaterial : StaticMesh->GetStaticMaterials())
 						{
-							MaterialSlotNameToBeOverriden = StaticMaterial.MaterialSlotName;
+							if (StaticMaterial.MaterialInterface == OriginalMaterial)
+							{
+								MaterialSlotNameToBeOverriden = StaticMaterial.MaterialSlotName;
 
-							break;
+								break;
+							}
 						}
 					}
-				}
-				else if (USkeletalMesh* SkeletalMesh = Cast<USkeletalMesh>(OuterObjects[0]))
-				{
-					// NOTE : As 'Materials' must be protected for async build, always use the accessors even internally.
-
-					for (FSkeletalMaterial SkeletalMaterial : SkeletalMesh->GetMaterials())
+					else if (USkeletalMesh* SkeletalMesh = Cast<USkeletalMesh>(OuterObjects[0]))
 					{
-						if (SkeletalMaterial.MaterialInterface == OriginalMaterial)
+						// NOTE : As 'Materials' must be protected for async build, always use the accessors even internally.
+						for (const FSkeletalMaterial& SkeletalMaterial : SkeletalMesh->GetMaterials())
 						{
-							MaterialSlotNameToBeOverriden = SkeletalMaterial.MaterialSlotName;
+							if (SkeletalMaterial.MaterialInterface == OriginalMaterial)
+							{
+								MaterialSlotNameToBeOverriden = SkeletalMaterial.MaterialSlotName;
+
+								break;
+							}
+						}
+					}
+
+					if (!ensureMsgf(!MaterialSlotNameToBeOverriden.IsNone(), TEXT("Material Property could not be exposed because the property does not contain any valid slot names.")))
+					{
+						return;
+					}
 
-							break;
+					if (UMeshComponent* MeshComponentToBeModified = GetSelectedMeshComponentToBeModified(OuterObjects[0], OriginalMaterial))
+					{
+						FScopedTransaction Transaction(LOCTEXT("OverrideMaterial", "Override Material"));
+						const int32 TargetMaterialIndex = MeshComponentToBeModified->GetMaterialIndex(MaterialSlotNameToBeOverriden);
+
+						if (FComponentEditorUtils::AttemptApplyMaterialToComponent(MeshComponentToBeModified, OriginalMaterial, TargetMaterialIndex))
+						{
+							RefreshPanels();
 						}
 					}
 				}
 			}
+		}
+	}
+	else if (ExtensionArgsType == FRCExposesPropertyArgs::EType::E_OwnerObject)
+	{
+		FRCFieldPathInfo RCFieldPathInfo(InPropertyArgs.PropertyPath);
+		RCFieldPathInfo.Resolve(InPropertyArgs.OwnerObject);
+		FRCFieldResolvedData ResolvedData = RCFieldPathInfo.GetResolvedData();
+		FObjectProperty* ObjectProperty = CastField<FObjectProperty>(InPropertyArgs.Property);
 
-			if (!ensureMsgf(!MaterialSlotNameToBeOverriden.IsNone(), TEXT("Material Property could not be exposed because the property does not contain any valid slot names.")))
-			{
-				return;
-			}
-		
-			// NOTE : Obtain the actor and/or object information from the details panel.
+		if (!ObjectProperty)
+		{
+			return;
+		}
 
-			TArray<TWeakObjectPtr<AActor>> SelectedActors;
+		// Material can't be null
+		UMaterialInterface* OriginalMaterial = Cast<UMaterialInterface>(ObjectProperty->GetObjectPropertyValue(ResolvedData.ContainerAddress));
+		OriginalMaterial = OriginalMaterial ? OriginalMaterial : UMaterial::GetDefaultMaterial(MD_Surface);
+		if (OriginalMaterial)
+		{
+			int32 MaterialIndex = INDEX_NONE;
 
-			TArray<TWeakObjectPtr<UObject>> SelectedObjects;
+			int32 StartPathArrayIndex = -1;
+			int32 EndPathArrayIndex = -1;
+			InPropertyArgs.PropertyPath.FindChar('[', StartPathArrayIndex);
+			InPropertyArgs.PropertyPath.FindChar(']', EndPathArrayIndex);
 
-			if (TSharedPtr<IDetailTreeNode> OwnerTreeNode = WeakDetailsTreeNode.Pin())
+			if (StartPathArrayIndex != INDEX_NONE && EndPathArrayIndex != INDEX_NONE)
 			{
-				if (IDetailsView* DetailsView = OwnerTreeNode->GetNodeDetailsView())
-				{
-					SelectedActors = DetailsView->GetSelectedActors();
-
-					SelectedObjects = DetailsView->GetSelectedObjects();
-				}
+				const FString MaterialSlotNameToBeOverriden = InPropertyArgs.PropertyPath.Mid(StartPathArrayIndex + 1, EndPathArrayIndex - 1 - StartPathArrayIndex);
+				LexFromString(MaterialIndex, *MaterialSlotNameToBeOverriden);
 			}
-			else if (SharedDetailsPanel.IsValid()) // Fallback to global detail panel reference if Detail Tree Node is invalid.
+
+			if (!ensureMsgf(MaterialIndex > -1, TEXT("Material Property could not be exposed because the property does not contain any valid slot index.")))
 			{
-				SelectedActors = SharedDetailsPanel->GetSelectedActors();
-			
-				SelectedObjects = SharedDetailsPanel->GetSelectedObjects();
+				return;
 			}
 
-			TWeakObjectPtr<UMeshComponent> MeshComponentToBeModified;
-
-			if (SelectedActors.Num()) // If user selected actor then get the component from it.
+			if (UMeshComponent* MeshComponentToBeModified = GetSelectedMeshComponentToBeModified(InPropertyArgs.OwnerObject, OriginalMaterial))
 			{
-				// NOTE : Allow single selection only.
+				FScopedTransaction Transaction(LOCTEXT("OverrideMaterial", "Override Material"));
 
-				if (!ensureMsgf(SelectedActors.Num() == 1, TEXT("Property could not be exposed as multiple actor(s) are selected.")))
+				if (FComponentEditorUtils::AttemptApplyMaterialToComponent(MeshComponentToBeModified, OriginalMaterial, MaterialIndex))
 				{
-					return;
+					RefreshPanels();
 				}
+			}
+		}
+	}
+}
 
-				for (TWeakObjectPtr<AActor> SelectedActor : SelectedActors)
-				{
-					TInlineComponentArray<UMeshComponent*> MeshComponents(SelectedActor.Get());
+UMeshComponent* FRemoteControlUIModule::GetSelectedMeshComponentToBeModified(UObject* InOwnerObject, UMaterialInterface* InOriginalMaterial)
+{
+	// NOTE : Obtain the actor and/or object information from the details panel.
 
-					// NOTE : First mesh component that has the material slot gets served (FCFS approach).
+	TArray<TWeakObjectPtr<AActor>> SelectedActors;
 
-					for (UMeshComponent* MeshComponent : MeshComponents)
-					{
-						TArray<FName> MaterialSlots = MeshComponent->GetMaterialSlotNames();
+	TArray<TWeakObjectPtr<UObject>> SelectedObjects;
 
-						if (MaterialSlots.Contains(MaterialSlotNameToBeOverriden))
-						{
-							MeshComponentToBeModified = MeshComponent;
+	if (TSharedPtr<IDetailTreeNode> OwnerTreeNode = WeakDetailsTreeNode.Pin())
+	{
+		if (IDetailsView* DetailsView = OwnerTreeNode->GetNodeDetailsView())
+		{
+			SelectedActors = DetailsView->GetSelectedActors();
 
-							break;
-						}
-					}
-				}
-			}
-			else if (SelectedObjects.Num()) // If user selected a component then proceed with it.
-			{
-				// NOTE : Allow single selection only.
+			SelectedObjects = DetailsView->GetSelectedObjects();
+		}
+	}
+	else if (SharedDetailsPanel.IsValid()) // Fallback to global detail panel reference if Detail Tree Node is invalid.
+	{
+		SelectedActors = SharedDetailsPanel->GetSelectedActors();
 
-				if (!ensureMsgf(SelectedObjects.Num() == 1, TEXT("Property could not be exposed as multiple component(s) are selected.")))
-				{
-					return;
-				}
+		SelectedObjects = SharedDetailsPanel->GetSelectedObjects();
+	}
 
-				for (TWeakObjectPtr<UObject> SelectedObject : SelectedObjects)
-				{
-					if (UMeshComponent* MeshComponent = Cast<UMeshComponent>(SelectedObject))
-					{
-						TArray<FName> MaterialSlots = MeshComponent->GetMaterialSlotNames();
+	UMeshComponent* MeshComponentToBeModified = nullptr;
 
-						if (MaterialSlots.Contains(MaterialSlotNameToBeOverriden))
-						{
-							MeshComponentToBeModified = MeshComponent;
+	auto GetComponentToBeModified = [InOwnerObject](UMeshComponent* InMeshComponent)
+	{
+		UMeshComponent* ModifinedComponent = nullptr;
 
-							break;
-						}
-					}
-				}
+		if (UStaticMeshComponent* StaticMeshComponent = Cast<UStaticMeshComponent>(InMeshComponent))
+		{
+			if (StaticMeshComponent->GetStaticMesh() == InOwnerObject)
+			{
+				ModifinedComponent = InMeshComponent;
 			}
-
-			if (MeshComponentToBeModified.IsValid())
+		}
+		else if (USkeletalMeshComponent* SkeletalMeshComponent = Cast<USkeletalMeshComponent>(InMeshComponent))
+		{
+			if (SkeletalMeshComponent->SkeletalMesh == InOwnerObject)
 			{
-				FScopedTransaction Transaction(LOCTEXT("OverrideMaterial", "Override Material"));
+				ModifinedComponent = InMeshComponent;
+			}
+		}
 
-				const int32 TargetMaterialIndex = MeshComponentToBeModified->GetMaterialIndex(MaterialSlotNameToBeOverriden);
+		return ModifinedComponent;
+	};
 
-				if (FComponentEditorUtils::AttemptApplyMaterialToComponent(MeshComponentToBeModified.Get(), OriginalMaterial, TargetMaterialIndex))
-				{
-					RefreshPanels();
-				}
+	if (SelectedActors.Num()) // If user selected actor then get the component from it.
+	{
+		// NOTE : Allow single selection only.
+
+		if (!ensureMsgf(SelectedActors.Num() == 1, TEXT("Property could not be exposed as multiple actor(s) are selected.")))
+		{
+			return nullptr;
+		}
+
+		for (TWeakObjectPtr<AActor> SelectedActor : SelectedActors)
+		{
+			TInlineComponentArray<UMeshComponent*> MeshComponents(SelectedActor.Get());
+
+			// NOTE : First mesh component that has the material slot gets served (FCFS approach).
+
+			for (UMeshComponent* MeshComponent : MeshComponents)
+			{
+				MeshComponentToBeModified = GetComponentToBeModified(MeshComponent);
 			}
 		}
+	}
+	else if (SelectedObjects.Num()) // If user selected a component then proceed with it.
+	{
+		// NOTE : Allow single selection only.
+		if (!ensureMsgf(SelectedObjects.Num() == 1, TEXT("Property could not be exposed as multiple component(s) are selected.")))
+		{
+			return nullptr;
+		}
 
+		for (TWeakObjectPtr<UObject> SelectedObject : SelectedObjects)
+		{
+			if (UMeshComponent* MeshComponent = Cast<UMeshComponent>(SelectedObject))
+			{
+				MeshComponentToBeModified = GetComponentToBeModified(MeshComponent);
+			}
+		}
 	}
+
+	return MeshComponentToBeModified;
 }
 
+
 void FRemoteControlUIModule::RefreshPanels()
 {
 	if (TSharedPtr<IDetailTreeNode> OwnerTreeNode = WeakDetailsTreeNode.Pin())
diff --git a/Engine/Plugins/VirtualProduction/RemoteControl/Source/RemoteControlUI/Private/RemoteControlUIModule.h b/Engine/Plugins/VirtualProduction/RemoteControl/Source/RemoteControlUI/Private/RemoteControlUIModule.h
--- a/Engine/Plugins/VirtualProduction/RemoteControl/Source/RemoteControlUI/Private/RemoteControlUIModule.h
+++ b/Engine/Plugins/VirtualProduction/RemoteControl/Source/RemoteControlUI/Private/RemoteControlUIModule.h
@@ -107,29 +107,29 @@ private:
 	bool CanToggleEditMode() const;	
 	bool IsInEditMode() const;
 
-	/** Handle creating the property row extensions.  */
+	/** Handle creating the row extensions.  */
 	void HandleCreatePropertyRowExtension(const FOnGenerateGlobalRowExtensionArgs& InArgs, TArray<FPropertyRowExtensionButton>& OutExtensions);
 
-	/** Handle getting the icon displayed in the property row extension. */
-	FSlateIcon OnGetExposedIcon(TSharedPtr<IPropertyHandle> Handle) const;
+	/** Handle getting the icon displayed in the row extension. */
+	FSlateIcon OnGetExposedIcon(const FRCExposesPropertyArgs& InArgs) const;
 
 	/** Handle getting the expose button visibility. */
-	bool CanToggleExposeProperty(TSharedPtr<IPropertyHandle> Handle) const;
+	bool CanToggleExposeProperty(const FRCExposesPropertyArgs InArgs) const;
 
 	/** Is the property currently exposed? */
-	ECheckBoxState GetPropertyExposedCheckState(TSharedPtr<IPropertyHandle> Handle) const;
+	ECheckBoxState GetPropertyExposedCheckState(const FRCExposesPropertyArgs InArgs) const;
 
 	/** Handle clicking the expose button. */
-	void OnToggleExposeProperty(TSharedPtr<IPropertyHandle> Handle);
+	void OnToggleExposeProperty(const FRCExposesPropertyArgs InArgs);
 
 	/** Returns whether a property is exposed, unexposed or unexposable. */
-	EPropertyExposeStatus GetPropertyExposeStatus(const TSharedPtr<IPropertyHandle>& Handle) const;
+	EPropertyExposeStatus GetPropertyExposeStatus(const FRCExposesPropertyArgs& InArgs) const;
 
-	/** Handle getting the icon displayed in the property row extension. */
-	FSlateIcon OnGetOverrideMaterialsIcon(TSharedPtr<IPropertyHandle> Handle) const;
+	/** Handle getting the icon displayed in the row extension. */
+	FSlateIcon OnGetOverrideMaterialsIcon(const FRCExposesPropertyArgs& InArgs) const;
 
 	/** Handle getting the override materials button visibility. */
-	bool IsStaticOrSkeletalMaterialProperty(TSharedPtr<IPropertyHandle> Handle) const;
+	bool IsStaticOrSkeletalMaterialProperty(const FRCExposesPropertyArgs InArgs) const;
 
 	/** Handle adding an option to get the object path in the actors' context menu. */
 	void AddGetPathOption(class FMenuBuilder& MenuBuilder, AActor* SelectedActor);
@@ -137,8 +137,13 @@ private:
 	/** Handle adding the menu extender for the actors. */
 	TSharedRef<FExtender> ExtendLevelViewportContextMenuForRemoteControl(const TSharedRef<FUICommandList> CommandList, TArray<AActor*> SelectedActors);
 
-	/** Returns whether a given property should have an exposed icon. */
-	bool ShouldDisplayExposeIcon(const TSharedRef<IPropertyHandle>& PropertyHandle) const;
+	/** Returns whether a given extension args should have an exposed icon. */
+	bool ShouldDisplayExposeIcon(const FRCExposesPropertyArgs& InArgs) const;
+
+	/** Return true if the extension args from remote control panel */
+	bool ShouldSkipOwnPanelProperty(const FRCExposesPropertyArgs& InArgs) const;
+
+	bool IsAllowedOwnerObjects(TArray<UObject*> InOuterObjects) const;
 
 	//~ Handle struct details customizations for common RC types.
 	void RegisterStructCustomizations();
@@ -152,13 +157,16 @@ private:
 	void RegisterWidgetFactories();
 
 	/** Returns expose button tooltip based on exposed state */
-	FText GetExposePropertyButtonTooltip(TSharedPtr<IPropertyHandle> Handle) const;
+	FText GetExposePropertyButtonTooltip(const FRCExposesPropertyArgs InArgs) const;
 
 	/** Returns expose button text based on exposed state */
-	FText GetExposePropertyButtonText(TSharedPtr<IPropertyHandle> Handle) const;
+	FText GetExposePropertyButtonText(const FRCExposesPropertyArgs InArgs) const;
 	
 	/** Attempts to replace the static or skeletal materials with their corressponding overrides. */
-	void TryOverridingMaterials(TSharedPtr<IPropertyHandle> ForThisProperty);
+	void TryOverridingMaterials(const FRCExposesPropertyArgs InArgs);
+
+	/** Return Selected Mesh Component by given OwnerObject and MaterialInterface */
+	UMeshComponent* GetSelectedMeshComponentToBeModified(UObject* InOwnerObject,  UMaterialInterface* InOriginalMaterial);
 
 	/** Attempts to refresh the details panel. */
 	void RefreshPanels();
diff --git a/Engine/Plugins/VirtualProduction/RemoteControl/Source/RemoteControlUI/Private/UI/SRemoteControlPanel.cpp b/Engine/Plugins/VirtualProduction/RemoteControl/Source/RemoteControlUI/Private/UI/SRemoteControlPanel.cpp
--- a/Engine/Plugins/VirtualProduction/RemoteControl/Source/RemoteControlUI/Private/UI/SRemoteControlPanel.cpp
+++ b/Engine/Plugins/VirtualProduction/RemoteControl/Source/RemoteControlUI/Private/UI/SRemoteControlPanel.cpp
@@ -31,7 +31,6 @@
 #include "RemoteControlLogger.h"
 #include "RemoteControlPanelStyle.h"
 #include "RemoteControlPreset.h"
-#include "RemoteControlUIModule.h"
 #include "RemoteControlSettings.h"
 #include "ScopedTransaction.h"
 #include "SClassViewer.h"
@@ -61,6 +60,7 @@
 
 #define LOCTEXT_NAMESPACE "RemoteControlPanel"
 
+
 namespace RemoteControlPanelUtils
 {
 	bool IsExposableActor(AActor* Actor)
@@ -111,7 +111,7 @@ void SRemoteControlPanel::Construct(const FArguments& InArgs, URemoteControlPres
 		{
 			UpdateEntityDetailsView(EntityList->GetSelection());
 			UpdateRebindButtonVisibility();
-			CachedExposedProperties.Reset();
+			CachedExposedPropertyArgs.Reset();
 		})
 		.EditMode_Lambda([this](){ return bIsInEditMode; });
 	
@@ -386,23 +386,66 @@ void SRemoteControlPanel::PostRedo(bool bSuccess)
 	Refresh();
 }
 
-bool SRemoteControlPanel::IsExposed(const TSharedPtr<IPropertyHandle>& PropertyHandle)
+bool SRemoteControlPanel::IsExposed(const FRCExposesPropertyArgs& InPropertyArgs)
 {
-	if (CachedExposedProperties.Contains(TWeakPtr<IPropertyHandle>{PropertyHandle}))
+	if (!ensure(InPropertyArgs.IsValid()))
 	{
-		return true;
+		return false;
+	}
+
+	const FRCExposesPropertyArgs::EType ExtensionArgsType = InPropertyArgs.GetType();
+
+	auto CheckCachedExposedArgs = [this, InPropertyArgs](const TArray<UObject*> InOwnerObjects, const FString& InPath, bool bIsCheckIsBoundByFullPath)
+	{
+		if (CachedExposedPropertyArgs.Contains(InPropertyArgs))
+		{
+			return true;
+		}
+
+		const bool bAllObjectsExposed = IsAllObjectsExposed(InOwnerObjects, InPath, bIsCheckIsBoundByFullPath);
+
+		if (bAllObjectsExposed)
+		{
+			CachedExposedPropertyArgs.Emplace(InPropertyArgs);
+		}
+
+		return bAllObjectsExposed;
+	};
+
+	if (ExtensionArgsType == FRCExposesPropertyArgs::EType::E_Handle)
+	{
+		TArray<UObject*> OuterObjects;
+		InPropertyArgs.PropertyHandle->GetOuterObjects(OuterObjects);
+		const FString Path = InPropertyArgs.PropertyHandle->GeneratePathToProperty();
+
+		constexpr bool bIsCheckIsBoundByFullPath = true;
+		return CheckCachedExposedArgs({ OuterObjects }, Path, bIsCheckIsBoundByFullPath);
+	}
+	else if (ExtensionArgsType == FRCExposesPropertyArgs::EType::E_OwnerObject)
+	{
+		constexpr bool bIsCheckIsBoundByFullPath = false;
+		return CheckCachedExposedArgs({ InPropertyArgs.OwnerObject }, InPropertyArgs.PropertyPath, bIsCheckIsBoundByFullPath);
 	}
-	
-	TArray<UObject*> OuterObjects;
-	PropertyHandle->GetOuterObjects(OuterObjects);
-	const FString Path = PropertyHandle->GeneratePathToProperty();
 
+	// It never should hit this point
+	ensure(false);
+
+	return false;
+}
+
+
+bool SRemoteControlPanel::IsAllObjectsExposed(TArray<UObject*> InOuterObjects, const FString& InPath, bool bUsingDuplicatesInPath)
+{
 	TArray<TSharedPtr<FRemoteControlProperty>, TInlineAllocator<1>> PotentialMatches;
 	for (const TWeakPtr<FRemoteControlProperty>& WeakProperty : Preset->GetExposedEntities<FRemoteControlProperty>())
 	{
 		if (TSharedPtr<FRemoteControlProperty> Property = WeakProperty.Pin())
 		{
-			if (Property->CheckIsBoundToPropertyPath(Path))
+			// If that was exposed by property path it should be checked by the full path with duplicated like propertypath.propertypath[0]
+			// If that was exposed by the owner object it should be without duplicated in the path, just propertypath[0]
+			const bool Isbound = bUsingDuplicatesInPath ? Property->CheckIsBoundToPropertyPath(InPath) : Property->CheckIsBoundToString(InPath);
+
+			if (Isbound)
 			{
 				PotentialMatches.Add(Property);
 			}
@@ -411,13 +454,13 @@ bool SRemoteControlPanel::IsExposed(const TSharedPtr<IPropertyHandle>& PropertyH
 
 	bool bAllObjectsExposed = true;
 
-	for (UObject* OuterObject : OuterObjects)
+	for (UObject* OuterObject : InOuterObjects)
 	{
 		bool bFoundPropForObject = false;
 
 		for (const TSharedPtr<FRemoteControlProperty>& Property : PotentialMatches)
 		{
-			if (Property->ContainsBoundObjects(OuterObjects))
+			if (Property->ContainsBoundObjects({ InOuterObjects } ))
 			{
 				bFoundPropForObject = true;
 				break;
@@ -427,46 +470,72 @@ bool SRemoteControlPanel::IsExposed(const TSharedPtr<IPropertyHandle>& PropertyH
 		bAllObjectsExposed &= bFoundPropForObject;
 	}
 
-	if (bAllObjectsExposed)
-	{
-		CachedExposedProperties.Emplace(PropertyHandle);
-	}
-	
 	return bAllObjectsExposed;
 }
 
-void SRemoteControlPanel::ToggleProperty(const TSharedPtr<IPropertyHandle>& PropertyHandle)
+void SRemoteControlPanel::ToggleProperty(const FRCExposesPropertyArgs& InPropertyArgs)
 {
-	TSet<UObject*> UniqueOuterObjects;
+	if (!ensure(InPropertyArgs.IsValid()))
 	{
-		// Make sure properties are only being exposed once per object.
-		TArray<UObject*> OuterObjects;
-		PropertyHandle->GetOuterObjects(OuterObjects);
-		UniqueOuterObjects.Append(MoveTemp(OuterObjects));
+		return;
 	}
 
-	if (IsExposed(PropertyHandle))
+	if (IsExposed(InPropertyArgs))
 	{
 		FScopedTransaction Transaction(LOCTEXT("UnexposeProperty", "Unexpose Property"));
 		Preset->Modify();
-		Unexpose(PropertyHandle);
+		Unexpose(InPropertyArgs);
 		return;
 	}
-	if (UniqueOuterObjects.Num())
+
+
+	auto PreExpose = [this]()
 	{
 		FScopedTransaction Transaction(LOCTEXT("ExposeProperty", "Expose Property"));
 		Preset->Modify();
-		
-		for (UObject* Object : UniqueOuterObjects)
+	};
+
+	auto PostExpose = [this, InPropertyArgs]()
+	{
+		Refresh();
+		CachedExposedPropertyArgs.Emplace(InPropertyArgs);
+	};
+
+	const FRCExposesPropertyArgs::EType ExtensionArgsType = InPropertyArgs.GetType();
+	if (ExtensionArgsType == FRCExposesPropertyArgs::EType::E_Handle)
+	{
+		TSet<UObject*> UniqueOuterObjects;
 		{
-			if (Object)
+			// Make sure properties are only being exposed once per object.
+			TArray<UObject*> OuterObjects;
+			InPropertyArgs.PropertyHandle->GetOuterObjects(OuterObjects);
+			UniqueOuterObjects.Append(MoveTemp(OuterObjects));
+		}
+
+		if (UniqueOuterObjects.Num())
+		{
+			PreExpose();
+
+			for (UObject* Object : UniqueOuterObjects)
 			{
-				constexpr bool bCleanDuplicates = true; //GeneratePathToProperty duplicates container name (Array.Array[1], Set.Set[1], etc...)
-				ExposeProperty(Object, FRCFieldPathInfo{PropertyHandle->GeneratePathToProperty(), bCleanDuplicates});
+				if (Object)
+				{
+					constexpr bool bCleanDuplicates = true; // GeneratePathToProperty duplicates container name (Array.Array[1], Set.Set[1], etc...)
+					ExposeProperty(Object, FRCFieldPathInfo{ InPropertyArgs.PropertyHandle->GeneratePathToProperty(), bCleanDuplicates });
+				}
 			}
+
+			PostExpose();
 		}
-		
-		CachedExposedProperties.Emplace(PropertyHandle);
+	}
+	else if (ExtensionArgsType == FRCExposesPropertyArgs::EType::E_OwnerObject)
+	{
+		PreExpose();
+
+		constexpr bool bCleanDuplicates = true; // GeneratePathToProperty duplicates container name (Array.Array[1], Set.Set[1], etc...)
+		ExposeProperty(InPropertyArgs.OwnerObject, FRCFieldPathInfo{ InPropertyArgs.PropertyPath, bCleanDuplicates });
+
+		PostExpose();
 	}
 }
 
@@ -806,35 +875,55 @@ void SRemoteControlPanel::Refresh()
 	EntityList->Refresh();
 }
 
-void SRemoteControlPanel::Unexpose(const TSharedPtr<IPropertyHandle>& Handle)
+
+void SRemoteControlPanel::Unexpose(const FRCExposesPropertyArgs& InPropertyArgs)
 {
-	TArray<UObject*> OuterObjects;
-	Handle->GetOuterObjects(OuterObjects);
-	const FString Path = Handle->GeneratePathToProperty();
+	if (!InPropertyArgs.IsValid())
+	{
+		return;
+	}
 
-	// Find an exposed property with the same path.
-	TArray<TSharedPtr<FRemoteControlProperty>, TInlineAllocator<1>> PotentialMatches;
-	for (const TWeakPtr<FRemoteControlProperty>& WeakProperty : Preset->GetExposedEntities<FRemoteControlProperty>())
+	auto CheckAndExpose = [&](TArray<UObject*> InOuterObjects, const FString& InPath)
 	{
-		if (TSharedPtr<FRemoteControlProperty> Property = WeakProperty.Pin())
+		// Find an exposed property with the same path.
+		TArray<TSharedPtr<FRemoteControlProperty>, TInlineAllocator<1>> PotentialMatches;
+		for (const TWeakPtr<FRemoteControlProperty>& WeakProperty : Preset->GetExposedEntities<FRemoteControlProperty>())
 		{
-			if (Property->CheckIsBoundToPropertyPath(Path))
+			if (TSharedPtr<FRemoteControlProperty> Property = WeakProperty.Pin())
 			{
-				PotentialMatches.Add(Property);
+				if (Property->CheckIsBoundToString(InPath))
+				{
+					PotentialMatches.Add(Property);
+				}
 			}
 		}
-	}
 
-	for (const TSharedPtr<FRemoteControlProperty>& Property : PotentialMatches)
-	{
-		if (Property->ContainsBoundObjects(OuterObjects))
+		for (const TSharedPtr<FRemoteControlProperty>& Property : PotentialMatches)
 		{
-			Preset->Unexpose(Property->GetId());
-			break;
+			if (Property->ContainsBoundObjects(InOuterObjects))
+			{
+				Preset->Unexpose(Property->GetId());
+				break;
+			}
 		}
+	};
+
+	const FRCExposesPropertyArgs::EType ExtensionArgsType = InPropertyArgs.GetType();
+
+	if (ExtensionArgsType == FRCExposesPropertyArgs::EType::E_Handle)
+	{
+		TArray<UObject*> OuterObjects;
+		InPropertyArgs.PropertyHandle->GetOuterObjects(OuterObjects);
+
+		CheckAndExpose(OuterObjects, InPropertyArgs.PropertyHandle->GeneratePathToProperty());
+	}
+	else if (ExtensionArgsType == FRCExposesPropertyArgs::EType::E_OwnerObject)
+	{
+		CheckAndExpose({ InPropertyArgs.OwnerObject }, InPropertyArgs.PropertyPath);
 	}
 }
 
+
 void SRemoteControlPanel::OnEditModeCheckboxToggle(ECheckBoxState State)
 {
 	bIsInEditMode = (State == ECheckBoxState::Checked) ? true : false;
@@ -1041,12 +1130,12 @@ void SRemoteControlPanel::OnAssetRenamed(const FAssetData& Asset, const FString&
 
 void SRemoteControlPanel::OnEntityExposed(URemoteControlPreset* InPreset, const FGuid& InEntityId)
 {
-	CachedExposedProperties.Empty();
+	CachedExposedPropertyArgs.Empty();
 }
 
 void SRemoteControlPanel::OnEntityUnexposed(URemoteControlPreset* InPreset, const FGuid& InEntityId)
 {
-	CachedExposedProperties.Empty();
+	CachedExposedPropertyArgs.Empty();
 }
 
 FReply SRemoteControlPanel::OnClickSettingsButton()
@@ -1062,4 +1151,6 @@ void SRemoteControlPanel::OnMaterialCompiled(UMaterialInterface* MaterialInterfa
 	Refresh();
 }
 
+
+
 #undef LOCTEXT_NAMESPACE /*RemoteControlPanel*/
diff --git a/Engine/Plugins/VirtualProduction/RemoteControl/Source/RemoteControlUI/Private/UI/SRemoteControlPanel.h b/Engine/Plugins/VirtualProduction/RemoteControl/Source/RemoteControlUI/Private/UI/SRemoteControlPanel.h
--- a/Engine/Plugins/VirtualProduction/RemoteControl/Source/RemoteControlUI/Private/UI/SRemoteControlPanel.h
+++ b/Engine/Plugins/VirtualProduction/RemoteControl/Source/RemoteControlUI/Private/UI/SRemoteControlPanel.h
@@ -5,6 +5,7 @@
 #include "EditorUndoClient.h"
 #include "Engine/EngineTypes.h"
 #include "IDetailTreeNode.h"
+#include "RemoteControlUIModule.h"
 #include "Styling/SlateTypes.h"
 #include "RemoteControlFieldPath.h"
 #include "UObject/StrongObjectPtr.h"
@@ -65,16 +66,24 @@ public:
 	const URemoteControlPreset* GetPreset() const { return Preset.Get(); }
 
 	/**
-	 * @param Handle The handle representing the property to check.
+	 * @param InArgs extension arguments
 	 * @return Whether a property is exposed or not.
 	 */
-	bool IsExposed(const TSharedPtr<IPropertyHandle>& PropertyHandle);
+	bool IsExposed(const FRCExposesPropertyArgs& InArgs);
+
+	/**
+	 * @param InOuterObjects outer objects to check
+	 * @param InPropertyPath full property path
+	 * @param bUsingDuplicatesInPath whether duplications like property.property[1] should be exists or just property[1]
+	 * @return Whether objects is exposed or not.
+	 */
+	bool IsAllObjectsExposed(TArray<UObject*> InOuterObjects, const FString& InPropertyPath, bool bUsingDuplicatesInPath);
 
 	/**
 	 * Exposes or unexposes a property.
-	 * @param Handle The handle of the property to toggle.
+	 * @param InArgs The extension arguments of the property to toggle.
 	 */
-	void ToggleProperty(const TSharedPtr<IPropertyHandle>& PropertyHandle);
+	void ToggleProperty(const FRCExposesPropertyArgs& InArgs);
 
 	/**
 	 * @return Whether or not the panel is in edit mode.
@@ -111,7 +120,7 @@ private:
 	void UnregisterEvents();
 
 	/** Unexpose a field from the preset. */
-	void Unexpose(const TSharedPtr<IPropertyHandle>& Handle);
+	void Unexpose(const FRCExposesPropertyArgs& InArgs);
 
 	/** Handle the using toggling the edit mode check box. */
 	void OnEditModeCheckboxToggle(ECheckBoxState State);
@@ -223,8 +232,8 @@ private:
 	TSharedPtr<SBox> EntityProtocolDetails;
 	/** Whether to show the rebind all button. */
 	bool bShowRebindButton = false;
-	/** Cache of exposed properties. */
-	TSet<TWeakPtr<IPropertyHandle>> CachedExposedProperties;
+	/** Cache of exposed property arguments. */
+	TSet<FRCExposesPropertyArgs> CachedExposedPropertyArgs;
 	/** Preset name widget. */
 	TSharedPtr<STextBlock> PresetNameTextBlock;
 	/** Holds a cache of widgets. */
diff --git a/Engine/Plugins/VirtualProduction/RemoteControl/Source/RemoteControlUI/Public/IRemoteControlUIModule.h b/Engine/Plugins/VirtualProduction/RemoteControl/Source/RemoteControlUI/Public/IRemoteControlUIModule.h
--- a/Engine/Plugins/VirtualProduction/RemoteControl/Source/RemoteControlUI/Public/IRemoteControlUIModule.h
+++ b/Engine/Plugins/VirtualProduction/RemoteControl/Source/RemoteControlUI/Public/IRemoteControlUIModule.h
@@ -10,6 +10,7 @@
 
 class FDelegateHandle;
 class FRCPanelWidgetRegistry;
+struct FOnGenerateGlobalRowExtensionArgs;
 struct FRemoteControlEntity;
 class IPropertyHandle;
 class IDetailCategoryBuilder;
@@ -35,6 +36,74 @@ struct FGenerateWidgetArgs
 	TWeakPtr<FRCPanelWidgetRegistry> WidgetRegistry;
 };
 
+/**
+ * Exposed property arguments
+ * The struct is wrapper struct to generate Remote Control Entity from IPropertyHandle or OwnerObject with Property Path and Property instance 
+ */
+struct FRCExposesPropertyArgs
+{
+	/**
+	 * Type of the Extension arguments
+	 */
+	enum class EType : uint8
+	{
+		E_Handle,
+		E_OwnerObject,
+		E_None
+	};
+
+	FRCExposesPropertyArgs();
+	FRCExposesPropertyArgs(const FOnGenerateGlobalRowExtensionArgs& InExtensionArgs);
+	FRCExposesPropertyArgs(FOnGenerateGlobalRowExtensionArgs&& InExtensionArgs);
+	FRCExposesPropertyArgs(TSharedPtr<IPropertyHandle>& InPropertyHandle);
+	FRCExposesPropertyArgs(UObject* InOwnerObject, const FString& InPropertyPath, FProperty* InProperty);
+
+
+	/** whether extensions arguments valid*/
+	bool IsValid() const;
+
+	/* Get Type of the Extension Arguments */
+	EType GetType() const;
+
+	/** Get property on the exposed arguments type */
+	FProperty* GetProperty() const;
+
+	/** Checked version of Get Property */
+	FProperty* GetPropertyChecked() const;
+
+	friend FORCEINLINE bool operator==(const FRCExposesPropertyArgs& LHS, const FRCExposesPropertyArgs& RHS)
+	{
+		return LHS.Id == RHS.Id;
+	}
+
+	friend FORCEINLINE bool operator!=(const FRCExposesPropertyArgs& LHS, const FRCExposesPropertyArgs& RHS)
+	{
+		return LHS.Id != RHS.Id;
+	}
+
+	friend FORCEINLINE uint32 GetTypeHash(const FRCExposesPropertyArgs& InArgs)
+	{
+		return GetTypeHash(InArgs.Id);
+	}
+
+public:
+	/** The detail row's property handle. */
+	TSharedPtr<IPropertyHandle> PropertyHandle;
+
+	/** Owner object for the property extension */
+	UObject* OwnerObject;
+
+	/** Path of the exposed property */
+	FString PropertyPath;
+
+	/** Exposed property */
+	FProperty* Property;
+
+private:
+	/** Unique generated ID of the struct */
+	FGuid Id;
+};
+
 DECLARE_DELEGATE_RetVal_OneParam(TSharedPtr<SRCPanelTreeNode>, FOnGenerateRCWidget, const FGenerateWidgetArgs& /*Args*/);
 
 DECLARE_MULTICAST_DELEGATE_OneParam(FOnGenerateExtensions, TArray<TSharedRef<class SWidget>>& /*OutExtensions*/);
@@ -42,7 +111,7 @@ DECLARE_MULTICAST_DELEGATE_OneParam(FOnGenerateExtensions, TArray<TSharedRef<cla
 /**
  * Filter queried in order to determine if a property should be displayed.
  */
-DECLARE_DELEGATE_RetVal_OneParam(bool, FOnDisplayExposeIcon, TSharedRef<IPropertyHandle> /*PropertyHandle*/);
+DECLARE_DELEGATE_RetVal_OneParam(bool, FOnDisplayExposeIcon, const FRCExposesPropertyArgs& /*InPropertyArgs*/);
 
 /**
  * Callback called to customize the display of a metadata entry for entities.
diff --git a/Engine/Source/Editor/PropertyEditor/Private/MaterialList.cpp b/Engine/Source/Editor/PropertyEditor/Private/MaterialList.cpp
--- a/Engine/Source/Editor/PropertyEditor/Private/MaterialList.cpp
+++ b/Engine/Source/Editor/PropertyEditor/Private/MaterialList.cpp
@@ -16,6 +16,7 @@
 #include "Widgets/Layout/SBox.h"
 #include "Widgets/Text/STextBlock.h"
 
+
 #define LOCTEXT_NAMESPACE "PropertyCustomizationHelpers"
 
 /** Definition of extra widgets delegate for bottom material value field. */
@@ -167,6 +168,7 @@ namespace ValueExtender
 
 TSharedRef<SWidget> FMaterialItemView::CreateValueContent(IDetailLayoutBuilder& InDetailBuilder, const TArray<FAssetData>& OwnerAssetDataArray, UActorComponent* InActorComponent)
 {	
+
 	return
 		SNew(SVerticalBox)
 		+SVerticalBox::Slot()
@@ -335,110 +337,34 @@ void FMaterialItemView::GoToAssetInContentBrowser( TWeakObjectPtr<UObject> Objec
 	}
 }
 
-PRAGMA_DISABLE_OPTIMIZATION
-TSharedPtr<IPropertyHandle> FMaterialItemView::GetMaterialHandle(IDetailLayoutBuilder& InDetailBuilder, UActorComponent* InCurrentComponent) const
+
+FOnGenerateGlobalRowExtensionArgs FMaterialItemView::GetGlobalRowExtensionArgs(IDetailLayoutBuilder& InDetailBuilder, UActorComponent* InCurrentComponent) const
 {
+	FOnGenerateGlobalRowExtensionArgs OnGenerateGlobalRowExtensionArgs;
+
 	USceneComponent* SceneComponent = Cast<USceneComponent>(InCurrentComponent);
 	if (!SceneComponent)
 	{
-		return nullptr;
+		return OnGenerateGlobalRowExtensionArgs;
 	}
 
 	TSharedPtr<IPropertyHandle> PropertyHandle;
 	UObject* OwnerObject = nullptr;
 	FString PropertyPath;
-	if (SceneComponent->GetMaterialPropertyPath(MaterialItem.SlotIndex, OwnerObject, PropertyPath))
+	FProperty* MaterialProperty = nullptr;
+	if (SceneComponent->GetMaterialPropertyPath(MaterialItem.SlotIndex, OwnerObject, PropertyPath, MaterialProperty))
 	{
-		TArray<FString> PathArray;
-		// Split full path to array
-		PropertyPath.ParseIntoArray(PathArray, TEXT("."));
-
-		for (int32 PathArrayIndex = 0; PathArrayIndex < PathArray.Num(); ++PathArrayIndex)
-		{
-			const FString& PathArrayItem = PathArray[PathArrayIndex];
-
-			int32 StartPathArrayIndex = -1;
-			int32 EndPathArrayIndex = -1;
-			PathArrayItem.FindChar('[', StartPathArrayIndex);
-			PathArrayItem.FindChar(']', EndPathArrayIndex);
-
-			// Get a property name without array 
-			FString PropertyName = (StartPathArrayIndex != -1) ? PathArrayItem.Mid(0, StartPathArrayIndex) : PathArrayItem;
-
-			// create a property handle if that is first path element
-			if (PathArrayIndex == 0)
-			{
-				UClass* OwnerObjectClass = OwnerObject->GetClass();
-
-				// Exception for UStaticMesh
-				// StaticMaterials property does not contain EditAnywhere flag which is required for PropertyHandle generation
-				if (OwnerObjectClass == UStaticMesh::StaticClass())
-				{					
-					FProperty* Property = UStaticMesh::StaticClass()->FindPropertyByName(UStaticMesh::GetStaticMaterialsName());
-					check(Property);
-
-					const EPropertyFlags OriginalPropertyFlags = Property->GetPropertyFlags();
-					Property->SetPropertyFlags(CPF_Edit);
-					PropertyHandle = InDetailBuilder.AddObjectPropertyData( { OwnerObject }, *PropertyName);
-					Property->SetPropertyFlags(OriginalPropertyFlags);
-				}
-				else
-				{
-					PropertyHandle = InDetailBuilder.AddObjectPropertyData( { OwnerObject }, *PropertyName);
-				}				
-			}
-
-			if (PropertyHandle.IsValid())
-			{
-				// Then check if property path contain array element like Materials[0]
-				if(const TSharedPtr<IPropertyHandleArray> HandleArray = PropertyHandle->AsArray())
-				{
-					if (StartPathArrayIndex != INDEX_NONE && EndPathArrayIndex != INDEX_NONE)
-					{
-						const FString ArrayIndexStr = PathArrayItem.Mid(StartPathArrayIndex + 1, EndPathArrayIndex - 1 - StartPathArrayIndex);
-
-						int32 ArrayIndex = INDEX_NONE;
-						LexFromString(ArrayIndex, *ArrayIndexStr);
-
-						uint32 HandleArrayElementsNum = 0;
-						HandleArray->GetNumElements(HandleArrayElementsNum);
-						if (ensure(ArrayIndex != INDEX_NONE) && ensure(ArrayIndex < static_cast<int32>(HandleArrayElementsNum)))
-						{
-							PropertyHandle = HandleArray->GetElement(ArrayIndex);
-						}
-					}
-				}
-				// Property Handle is inner struct element if we have more then 1 element in path array PathArray
-				// For example Materials[0].MaterialInterface 
-				if (PathArrayIndex != 0)
-				{
-					// Make sure the property name is valid
-					if (ensure(!PropertyName.IsEmpty()))
-					{
-						PropertyHandle = PropertyHandle->GetChildHandle(*PropertyName);
-					}
-					else
-					{
-						PropertyHandle = nullptr;
-					}
-				}
-			}
-		}
+		OnGenerateGlobalRowExtensionArgs.OwnerObject = OwnerObject;
+		OnGenerateGlobalRowExtensionArgs.PropertyPath = PropertyPath;
+		OnGenerateGlobalRowExtensionArgs.Property = MaterialProperty;
 	}
 
-	return PropertyHandle;
+	return OnGenerateGlobalRowExtensionArgs;
 }
-PRAGMA_ENABLE_OPTIMIZATION;
+
 
 TSharedPtr<SWidget> FMaterialItemView::GetGlobalRowExtensionWidget(IDetailLayoutBuilder& InDetailBuilder, UActorComponent* InCurrentComponent) const
 {
-	const TSharedPtr<IPropertyHandle> MaterialPropertyHandle = GetMaterialHandle(InDetailBuilder, InCurrentComponent);
-
-	if (!MaterialPropertyHandle)
-	{
-		return SNullWidget::NullWidget;
-	}
-		
 	FSlimHorizontalToolBarBuilder ToolbarBuilder(TSharedPtr<FUICommandList>(), FMultiBoxCustomization::None);
 	FPropertyEditorModule& PropertyEditorModule = FModuleManager::Get().GetModuleChecked<FPropertyEditorModule>("PropertyEditor");
 
@@ -446,8 +372,7 @@ TSharedPtr<SWidget> FMaterialItemView::GetGlobalRowExtensionWidget(IDetailLayout
 	TArray<FPropertyRowExtensionButton> ExtensionButtons;
 
 	// Add custom property extensions
-	FOnGenerateGlobalRowExtensionArgs RowExtensionArgs;
-	RowExtensionArgs.PropertyHandle = MaterialPropertyHandle;
+	FOnGenerateGlobalRowExtensionArgs RowExtensionArgs = GetGlobalRowExtensionArgs(InDetailBuilder, InCurrentComponent);
 	PropertyEditorModule.GetGlobalRowExtensionDelegate().Broadcast(RowExtensionArgs, ExtensionButtons);
 
 	// Build extension toolbar 
diff --git a/Engine/Source/Editor/PropertyEditor/Public/MaterialList.h b/Engine/Source/Editor/PropertyEditor/Public/MaterialList.h
--- a/Engine/Source/Editor/PropertyEditor/Public/MaterialList.h
+++ b/Engine/Source/Editor/PropertyEditor/Public/MaterialList.h
@@ -228,13 +228,14 @@ private:
 	void GoToAssetInContentBrowser( TWeakObjectPtr<UObject> Object );
 
 	/**
-	 * Get Material handle for given component
+	 * Get Property Row Extenssion Arguments.
+	 * The extension arguments is needed to add button or custom widget to property row.
 	 *
 	 * @param InDetailBuilder The builder for laying custom details.
 	 * @param InCurrentComponent The current component using the current material.
-	 * @return Property Handle pointer
+	 * @return Generate Global Row Extension Args
 	 */
-	TSharedPtr<IPropertyHandle> GetMaterialHandle(IDetailLayoutBuilder& InDetailBuilder, UActorComponent* InCurrentComponent) const;
+	struct FOnGenerateGlobalRowExtensionArgs GetGlobalRowExtensionArgs(IDetailLayoutBuilder& InDetailBuilder, UActorComponent* InCurrentComponent) const;
 
 	/**
 	 * Get row extension widget
diff --git a/Engine/Source/Editor/PropertyEditor/Public/PropertyEditorDelegates.h b/Engine/Source/Editor/PropertyEditor/Public/PropertyEditorDelegates.h
--- a/Engine/Source/Editor/PropertyEditor/Public/PropertyEditorDelegates.h
+++ b/Engine/Source/Editor/PropertyEditor/Public/PropertyEditorDelegates.h
@@ -90,12 +90,25 @@ DECLARE_DELEGATE_RetVal(bool, FIsPropertyEditingEnabled );
  */
 DECLARE_MULTICAST_DELEGATE_OneParam(FOnFinishedChangingProperties, const FPropertyChangedEvent&);
 
+/**
+ * A property row extension arguments is displayed at the end of a property row, either inline or as a button.
+ */
 struct FOnGenerateGlobalRowExtensionArgs
-{
+{	
 	/** The detail row's property handle. */
 	TSharedPtr<IPropertyHandle> PropertyHandle;
+
 	/** The detail row's owner tree node. */
 	TWeakPtr<IDetailTreeNode> OwnerTreeNode;
+
+	/** Owner object for the property extension */
+	UObject* OwnerObject = nullptr;
+
+	/** Path of the exposed property */
+	FString PropertyPath = TEXT("");
+
+	/** Exposed property */
+	FProperty* Property = nullptr;
 };
 
 /** 
diff --git a/Engine/Source/Runtime/Engine/Classes/Components/DecalComponent.h b/Engine/Source/Runtime/Engine/Classes/Components/DecalComponent.h
--- a/Engine/Source/Runtime/Engine/Classes/Components/DecalComponent.h
+++ b/Engine/Source/Runtime/Engine/Classes/Components/DecalComponent.h
@@ -176,7 +176,7 @@ public:
 
 	//~ Begin USceneComponent Interface
 #if WITH_EDITOR
-	virtual bool GetMaterialPropertyPath(int32 ElementIndex, UObject*& OutOwner, FString& OutPropertyPath) override;
+	virtual bool GetMaterialPropertyPath(int32 ElementIndex, UObject*& OutOwner, FString& OutPropertyPath, FProperty*& OutProperty) override;
 #endif // WITH_EDITOR
 	virtual FBoxSphereBounds CalcBounds(const FTransform& LocalToWorld) const override;
 	//~ End USceneComponent Interface
diff --git a/Engine/Source/Runtime/Engine/Classes/Components/LightComponent.h b/Engine/Source/Runtime/Engine/Classes/Components/LightComponent.h
--- a/Engine/Source/Runtime/Engine/Classes/Components/LightComponent.h
+++ b/Engine/Source/Runtime/Engine/Classes/Components/LightComponent.h
@@ -443,7 +443,7 @@ protected:
 
 	//~ Begin USceneComponent Interface
 #if WITH_EDITOR
-	virtual bool GetMaterialPropertyPath(int32 ElementIndex, UObject*& OutOwner, FString& OutPropertyPath) override;
+	virtual bool GetMaterialPropertyPath(int32 ElementIndex, UObject*& OutOwner, FString& OutPropertyPath, FProperty*& OutProperty) override;
 #endif // WITH_EDITOR
 	//~ End USceneComponent Interface
 
diff --git a/Engine/Source/Runtime/Engine/Classes/Components/MaterialBillboardComponent.h b/Engine/Source/Runtime/Engine/Classes/Components/MaterialBillboardComponent.h
--- a/Engine/Source/Runtime/Engine/Classes/Components/MaterialBillboardComponent.h
+++ b/Engine/Source/Runtime/Engine/Classes/Components/MaterialBillboardComponent.h
@@ -91,7 +91,7 @@ class ENGINE_API UMaterialBillboardComponent : public UPrimitiveComponent
 
 	//~ Begin USceneComponent Interface
 #if WITH_EDITOR
-	virtual bool GetMaterialPropertyPath(int32 ElementIndex, UObject*& OutOwner, FString& OutPropertyPath) override;
+	virtual bool GetMaterialPropertyPath(int32 ElementIndex, UObject*& OutOwner, FString& OutPropertyPath, FProperty*& OutProperty) override;
 #endif // WITH_EDITOR
 	//~ End USceneComponent Interface
 };
diff --git a/Engine/Source/Runtime/Engine/Classes/Components/SceneComponent.h b/Engine/Source/Runtime/Engine/Classes/Components/SceneComponent.h
--- a/Engine/Source/Runtime/Engine/Classes/Components/SceneComponent.h
+++ b/Engine/Source/Runtime/Engine/Classes/Components/SceneComponent.h
@@ -844,9 +844,10 @@ public:
 	 * @param ElementIndex		- The element to access the material of.
 	 * @param OutOwner			- Property UObject owner.
 	 * @param OutPropertyPath	- Full material property path.
+	 * @param OutProperty		- Material Property.
 	 * @return true if that was successfully resolved and component has material
 	 */
-	virtual bool GetMaterialPropertyPath(int32 ElementIndex, UObject*& OutOwner, FString& OutPropertyPath);
+	virtual bool GetMaterialPropertyPath(int32 ElementIndex, UObject*& OutOwner, FString& OutPropertyPath, FProperty*& OutProperty);
 #endif // WITH_EDITOR
 
 protected:
diff --git a/Engine/Source/Runtime/Engine/Classes/Components/SkinnedMeshComponent.h b/Engine/Source/Runtime/Engine/Classes/Components/SkinnedMeshComponent.h
--- a/Engine/Source/Runtime/Engine/Classes/Components/SkinnedMeshComponent.h
+++ b/Engine/Source/Runtime/Engine/Classes/Components/SkinnedMeshComponent.h
@@ -872,7 +872,7 @@ protected:
 public:
 	//~ Begin USceneComponent Interface
 #if WITH_EDITOR
-	virtual bool GetMaterialPropertyPath(int32 ElementIndex, UObject*& OutOwner, FString& OutPropertyPath) override;
+	virtual bool GetMaterialPropertyPath(int32 ElementIndex, UObject*& OutOwner, FString& OutPropertyPath, FProperty*& OutProperty) override;
 #endif // WITH_EDITOR
 	virtual FBoxSphereBounds CalcBounds(const FTransform& LocalToWorld) const override;
 	virtual FTransform GetSocketTransform(FName InSocketName, ERelativeTransformSpace TransformSpace = RTS_World) const override;
diff --git a/Engine/Source/Runtime/Engine/Classes/Components/StaticMeshComponent.h b/Engine/Source/Runtime/Engine/Classes/Components/StaticMeshComponent.h
--- a/Engine/Source/Runtime/Engine/Classes/Components/StaticMeshComponent.h
+++ b/Engine/Source/Runtime/Engine/Classes/Components/StaticMeshComponent.h
@@ -470,7 +470,7 @@ public:
 
 	//~ Begin USceneComponent Interface
 #if WITH_EDITOR
-	virtual bool GetMaterialPropertyPath(int32 ElementIndex, UObject*& OutOwner, FString& OutPropertyPath) override;
+	virtual bool GetMaterialPropertyPath(int32 ElementIndex, UObject*& OutOwner, FString& OutPropertyPath, FProperty*& OutProperty) override;
 #endif // WITH_EDITOR
 	virtual FBoxSphereBounds CalcBounds(const FTransform& LocalToWorld) const override;
 	virtual bool HasAnySockets() const override;
diff --git a/Engine/Source/Runtime/Engine/Classes/Components/TextRenderComponent.h b/Engine/Source/Runtime/Engine/Classes/Components/TextRenderComponent.h
--- a/Engine/Source/Runtime/Engine/Classes/Components/TextRenderComponent.h
+++ b/Engine/Source/Runtime/Engine/Classes/Components/TextRenderComponent.h
@@ -168,7 +168,7 @@ class ENGINE_API UTextRenderComponent : public UPrimitiveComponent
 
 	//~ Begin USceneComponent Interface.
 #if WITH_EDITOR
-	virtual bool GetMaterialPropertyPath(int32 ElementIndex, UObject*& OutOwner, FString& OutPropertyPath) override;
+	virtual bool GetMaterialPropertyPath(int32 ElementIndex, UObject*& OutOwner, FString& OutPropertyPath, FProperty*& OutProperty) override;
 #endif // WITH_EDITOR
 	virtual FBoxSphereBounds CalcBounds(const FTransform& LocalToWorld) const override;
 	// Note: this is required because GetRenderMatrix is overridden, and therefore the updated world-space bounds must be based on that (not the component transform)
diff --git a/Engine/Source/Runtime/Engine/Classes/Particles/ParticleSystemComponent.h b/Engine/Source/Runtime/Engine/Classes/Particles/ParticleSystemComponent.h
--- a/Engine/Source/Runtime/Engine/Classes/Particles/ParticleSystemComponent.h
+++ b/Engine/Source/Runtime/Engine/Classes/Particles/ParticleSystemComponent.h
@@ -1412,7 +1412,7 @@ protected:
 
 public:
 #if WITH_EDITOR
-	virtual bool GetMaterialPropertyPath(int32 ElementIndex, UObject*& OutOwner, FString& OutPropertyPath) override;
+	virtual bool GetMaterialPropertyPath(int32 ElementIndex, UObject*& OutOwner, FString& OutPropertyPath, FProperty*& OutProperty) override;
 #endif // WITH_EDITOR
 	virtual void Activate(bool bReset=false) override;
 	virtual void Deactivate() override;
diff --git a/Engine/Source/Runtime/Engine/Private/Components/DecalComponent.cpp b/Engine/Source/Runtime/Engine/Private/Components/DecalComponent.cpp
--- a/Engine/Source/Runtime/Engine/Private/Components/DecalComponent.cpp
+++ b/Engine/Source/Runtime/Engine/Private/Components/DecalComponent.cpp
@@ -133,12 +133,13 @@ bool UDecalComponent::IsPostLoadThreadSafe() const
 }
 
 #if WITH_EDITOR
-bool UDecalComponent::GetMaterialPropertyPath(int32 ElementIndex, UObject*& OutOwner, FString& OutPropertyPath)
+bool UDecalComponent::GetMaterialPropertyPath(int32 ElementIndex, UObject*& OutOwner, FString& OutPropertyPath, FProperty*& OutProperty)
 {
 	if(ElementIndex == 0)
 	{
 		OutOwner = this;
 		OutPropertyPath = GET_MEMBER_NAME_STRING_CHECKED(UDecalComponent, DecalMaterial);
+		OutProperty = UDecalComponent::StaticClass()->FindPropertyByName(GET_MEMBER_NAME_CHECKED(UDecalComponent, DecalMaterial));
 		return true;
 	}
 
diff --git a/Engine/Source/Runtime/Engine/Private/Components/LightComponent.cpp b/Engine/Source/Runtime/Engine/Private/Components/LightComponent.cpp
--- a/Engine/Source/Runtime/Engine/Private/Components/LightComponent.cpp
+++ b/Engine/Source/Runtime/Engine/Private/Components/LightComponent.cpp
@@ -906,12 +906,13 @@ void ULightComponent::DestroyRenderState_Concurrent()
 }
 
 #if WITH_EDITOR
-bool ULightComponent::GetMaterialPropertyPath(int32 ElementIndex, UObject*& OutOwner, FString& OutPropertyPath)
+bool ULightComponent::GetMaterialPropertyPath(int32 ElementIndex, UObject*& OutOwner, FString& OutPropertyPath, FProperty*& OutProperty)
 {
 	if (ElementIndex == 0)
 	{
 		OutOwner = this;
 		OutPropertyPath = GET_MEMBER_NAME_STRING_CHECKED(ULightComponent, LightFunctionMaterial);
+		OutProperty = ULightComponent::StaticClass()->FindPropertyByName(GET_MEMBER_NAME_CHECKED(ULightComponent, LightFunctionMaterial));
 		return true;
 	}
 
diff --git a/Engine/Source/Runtime/Engine/Private/Components/MaterialBillboardComponent.cpp b/Engine/Source/Runtime/Engine/Private/Components/MaterialBillboardComponent.cpp
--- a/Engine/Source/Runtime/Engine/Private/Components/MaterialBillboardComponent.cpp
+++ b/Engine/Source/Runtime/Engine/Private/Components/MaterialBillboardComponent.cpp
@@ -319,13 +319,14 @@ FBoxSphereBounds UMaterialBillboardComponent::CalcBounds(const FTransform& Local
 }
 
 #if WITH_EDITOR
-bool UMaterialBillboardComponent::GetMaterialPropertyPath(int32 ElementIndex, UObject*& OutOwner, FString& OutPropertyPath)
+bool UMaterialBillboardComponent::GetMaterialPropertyPath(int32 ElementIndex, UObject*& OutOwner, FString& OutPropertyPath, FProperty*& OutProperty)
 {
 	if (Elements.IsValidIndex(ElementIndex))
 	{
 		OutOwner = this;
 		OutPropertyPath = FString::Printf(TEXT("%s[%d].%s"), GET_MEMBER_NAME_STRING_CHECKED(UMaterialBillboardComponent, Elements), ElementIndex, GET_MEMBER_NAME_STRING_CHECKED(FMaterialSpriteElement, Material));
-	
+		OutProperty = FMaterialSpriteElement::StaticStruct()->FindPropertyByName(GET_MEMBER_NAME_CHECKED(FMaterialSpriteElement, Material));
+
 		return true;
 	}
 
diff --git a/Engine/Source/Runtime/Engine/Private/Components/SceneComponent.cpp b/Engine/Source/Runtime/Engine/Private/Components/SceneComponent.cpp
--- a/Engine/Source/Runtime/Engine/Private/Components/SceneComponent.cpp
+++ b/Engine/Source/Runtime/Engine/Private/Components/SceneComponent.cpp
@@ -3043,7 +3043,7 @@ bool USceneComponent::IsVisible() const
 }
 
 #if WITH_EDITOR
-bool USceneComponent::GetMaterialPropertyPath(int32 ElementIndex, UObject*& OutOwner, FString& OutPropertyPath)
+bool USceneComponent::GetMaterialPropertyPath(int32 ElementIndex, UObject*& OutOwner, FString& OutPropertyPath, FProperty*& OutProperty)
 {
 	// Should be overriden in inherited classes
 	return false;
diff --git a/Engine/Source/Runtime/Engine/Private/Components/SkinnedMeshComponent.cpp b/Engine/Source/Runtime/Engine/Private/Components/SkinnedMeshComponent.cpp
--- a/Engine/Source/Runtime/Engine/Private/Components/SkinnedMeshComponent.cpp
+++ b/Engine/Source/Runtime/Engine/Private/Components/SkinnedMeshComponent.cpp
@@ -2429,18 +2429,23 @@ bool USkinnedMeshComponent::UpdateOverlapsImpl(const TOverlapArrayView* PendingO
 }
 
 #if WITH_EDITOR
-bool USkinnedMeshComponent::GetMaterialPropertyPath(int32 ElementIndex, UObject*& OutOwner, FString& OutPropertyPath)
+bool USkinnedMeshComponent::GetMaterialPropertyPath(int32 ElementIndex, UObject*& OutOwner, FString& OutPropertyPath, FProperty*& OutProperty)
 {
 	if(OverrideMaterials.IsValidIndex(ElementIndex))
 	{
 		OutOwner = this;
 		OutPropertyPath = FString::Printf(TEXT("%s[%d]"), GET_MEMBER_NAME_STRING_CHECKED(UMeshComponent, OverrideMaterials), ElementIndex);
+		if (FArrayProperty* ArrayProperty = CastField<FArrayProperty>(UMeshComponent::StaticClass()->FindPropertyByName(GET_MEMBER_NAME_CHECKED(UMeshComponent, OverrideMaterials))))
+		{
+			OutProperty = ArrayProperty->Inner;
+		}
 		return true;
 	}
 	if (SkeletalMesh && SkeletalMesh->GetMaterials().IsValidIndex(ElementIndex))
 	{
 		OutOwner = SkeletalMesh;
 		OutPropertyPath = FString::Printf(TEXT("%s[%d].%s"), *USkeletalMesh::GetMaterialsMemberName().ToString(), ElementIndex, GET_MEMBER_NAME_STRING_CHECKED(FSkeletalMaterial, MaterialInterface));
+		OutProperty = FSkeletalMaterial::StaticStruct()->FindPropertyByName(GET_MEMBER_NAME_CHECKED(FSkeletalMaterial, MaterialInterface));
 		return true;
 	}
 
diff --git a/Engine/Source/Runtime/Engine/Private/Components/StaticMeshComponent.cpp b/Engine/Source/Runtime/Engine/Private/Components/StaticMeshComponent.cpp
--- a/Engine/Source/Runtime/Engine/Private/Components/StaticMeshComponent.cpp
+++ b/Engine/Source/Runtime/Engine/Private/Components/StaticMeshComponent.cpp
@@ -2528,12 +2528,16 @@ void UStaticMeshComponent::GetUsedMaterials(TArray<UMaterialInterface*>& OutMate
 }
 
 #if WITH_EDITOR
-bool UStaticMeshComponent::GetMaterialPropertyPath(int32 ElementIndex, UObject*& OutOwner, FString& OutPropertyPath)
+bool UStaticMeshComponent::GetMaterialPropertyPath(int32 ElementIndex, UObject*& OutOwner, FString& OutPropertyPath, FProperty*& OutProperty)
 {
 	if(OverrideMaterials.IsValidIndex(ElementIndex))
 	{				
 		OutOwner = this;
 		OutPropertyPath = FString::Printf(TEXT("%s[%d]"), GET_MEMBER_NAME_STRING_CHECKED(UMeshComponent, OverrideMaterials), ElementIndex);
+		if (FArrayProperty* ArrayProperty = CastField<FArrayProperty>(UMeshComponent::StaticClass()->FindPropertyByName(GET_MEMBER_NAME_CHECKED(UMeshComponent, OverrideMaterials))))
+		{
+			OutProperty = ArrayProperty->Inner;
+		}
 
 		return true;
 	}
@@ -2541,6 +2545,7 @@ bool UStaticMeshComponent::GetMaterialPropertyPath(int32 ElementIndex, UObject*&
 	{
 		OutOwner = GetStaticMesh();
 		OutPropertyPath = FString::Printf(TEXT("%s[%d].%s"), *UStaticMesh::GetStaticMaterialsName().ToString(), ElementIndex, GET_MEMBER_NAME_STRING_CHECKED(FStaticMaterial, MaterialInterface));
+		OutProperty = FStaticMaterial::StaticStruct()->FindPropertyByName(GET_MEMBER_NAME_CHECKED(FStaticMaterial, MaterialInterface));
 
 		return true;
 	}
diff --git a/Engine/Source/Runtime/Engine/Private/Components/TextRenderComponent.cpp b/Engine/Source/Runtime/Engine/Private/Components/TextRenderComponent.cpp
--- a/Engine/Source/Runtime/Engine/Private/Components/TextRenderComponent.cpp
+++ b/Engine/Source/Runtime/Engine/Private/Components/TextRenderComponent.cpp
@@ -1108,12 +1108,13 @@ int32 UTextRenderComponent::GetNumMaterials() const
 }
 
 #if WITH_EDITOR
-bool UTextRenderComponent::GetMaterialPropertyPath(int32 ElementIndex, UObject*& OutOwner, FString& OutPropertyPath)
+bool UTextRenderComponent::GetMaterialPropertyPath(int32 ElementIndex, UObject*& OutOwner, FString& OutPropertyPath, FProperty*& OutProperty)
 {
 	if (ElementIndex == 0)
 	{
 		OutOwner = this;
 		OutPropertyPath = GET_MEMBER_NAME_STRING_CHECKED(UTextRenderComponent, TextMaterial);
+		OutProperty = UTextRenderComponent::StaticClass()->FindPropertyByName(GET_MEMBER_NAME_CHECKED(UTextRenderComponent, TextMaterial));
 		return true;
 	}
 
diff --git a/Engine/Source/Runtime/Engine/Private/Particles/ParticleComponents.cpp b/Engine/Source/Runtime/Engine/Private/Particles/ParticleComponents.cpp
--- a/Engine/Source/Runtime/Engine/Private/Particles/ParticleComponents.cpp
+++ b/Engine/Source/Runtime/Engine/Private/Particles/ParticleComponents.cpp
@@ -4445,12 +4445,17 @@ int32 UParticleSystemComponent::GetNumMaterials() const
 }
 
 #if WITH_EDITOR
-bool UParticleSystemComponent::GetMaterialPropertyPath(int32 ElementIndex, UObject*& OutOwner, FString& OutPropertyPath)
+bool UParticleSystemComponent::GetMaterialPropertyPath(int32 ElementIndex, UObject*& OutOwner, FString& OutPropertyPath, FProperty*& OutProperty)
 {
 	if (EmitterMaterials.IsValidIndex(ElementIndex))
 	{
 		OutOwner = this;
 		OutPropertyPath = FString::Printf(TEXT("%s[%d]"), GET_MEMBER_NAME_STRING_CHECKED(UParticleSystemComponent, EmitterMaterials), ElementIndex);
+
+		if (FArrayProperty* ArrayProperty = CastField<FArrayProperty>(UParticleSystemComponent::StaticClass()->FindPropertyByName(GET_MEMBER_NAME_CHECKED(UParticleSystemComponent, EmitterMaterials))))
+		{
+			OutProperty = ArrayProperty->Inner;
+		}
 		return true;
 	}
 	if (Template && Template->Emitters.IsValidIndex(ElementIndex))
@@ -4463,6 +4468,7 @@ bool UParticleSystemComponent::GetMaterialPropertyPath(int32 ElementIndex, UObje
 			{
 				OutOwner = EmitterLODLevel->RequiredModule;
 				OutPropertyPath = GET_MEMBER_NAME_STRING_CHECKED(UParticleModuleRequired, Material);
+				OutProperty = UParticleModuleRequired::StaticClass()->FindPropertyByName(GET_MEMBER_NAME_CHECKED(UParticleModuleRequired, Material));
 				return true;
 			}
 		}
diff --git a/Engine/Source/Runtime/Landscape/Classes/LandscapeComponent.h b/Engine/Source/Runtime/Landscape/Classes/LandscapeComponent.h
--- a/Engine/Source/Runtime/Landscape/Classes/LandscapeComponent.h
+++ b/Engine/Source/Runtime/Landscape/Classes/LandscapeComponent.h
@@ -806,7 +806,7 @@ public:
 
 	//~ Begin USceneComponent Interface.
 #if WITH_EDITOR
-	virtual bool GetMaterialPropertyPath(int32 ElementIndex, UObject*& OutOwner, FString& OutPropertyPath) override;
+	virtual bool GetMaterialPropertyPath(int32 ElementIndex, UObject*& OutOwner, FString& OutPropertyPath, FProperty*& OutProperty) override;
 #endif // WITH_EDITOR
 	virtual void DestroyComponent(bool bPromoteChildren = false) override;
 	virtual FBoxSphereBounds CalcBounds(const FTransform& LocalToWorld) const override;
diff --git a/Engine/Source/Runtime/Landscape/Private/LandscapeEdit.cpp b/Engine/Source/Runtime/Landscape/Private/LandscapeEdit.cpp
--- a/Engine/Source/Runtime/Landscape/Private/LandscapeEdit.cpp
+++ b/Engine/Source/Runtime/Landscape/Private/LandscapeEdit.cpp
@@ -577,7 +577,7 @@ int32 ULandscapeComponent::GetNumMaterials() const
 }
 
 #if WITH_EDITOR
-bool ULandscapeComponent::GetMaterialPropertyPath(int32 ElementIndex, UObject*& OutOwner, FString& OutPropertyPath)
+bool ULandscapeComponent::GetMaterialPropertyPath(int32 ElementIndex, UObject*& OutOwner, FString& OutPropertyPath, FProperty*& OutProperty)
 {
 	if (ElementIndex == 0)
 	{
@@ -585,6 +585,7 @@ bool ULandscapeComponent::GetMaterialPropertyPath(int32 ElementIndex, UObject*&
 		{
 			OutOwner = this;
 			OutPropertyPath = GET_MEMBER_NAME_STRING_CHECKED(ULandscapeComponent, OverrideMaterial);
+			OutProperty = ULandscapeComponent::StaticClass()->FindPropertyByName(GET_MEMBER_NAME_CHECKED(ULandscapeComponent, OverrideMaterial));
 	
 			return true;
 		}
@@ -592,6 +593,7 @@ bool ULandscapeComponent::GetMaterialPropertyPath(int32 ElementIndex, UObject*&
 		{
 			OutOwner = Proxy;
 			OutPropertyPath = GET_MEMBER_NAME_STRING_CHECKED(ALandscapeProxy, LandscapeHoleMaterial);
+			OutProperty = ALandscapeProxy::StaticClass()->FindPropertyByName(GET_MEMBER_NAME_CHECKED(ALandscapeProxy, LandscapeHoleMaterial));
 			return true;
 		}
 	}
