[Remote Control] Fix rebuind issue for material assets expose

diff --git a/Engine/Plugins/VirtualProduction/RemoteControl/Source/RemoteControlUI/Private/RemoteControlUIModule.cpp b/Engine/Plugins/VirtualProduction/RemoteControl/Source/RemoteControlUI/Private/RemoteControlUIModule.cpp
--- a/Engine/Plugins/VirtualProduction/RemoteControl/Source/RemoteControlUI/Private/RemoteControlUIModule.cpp
+++ b/Engine/Plugins/VirtualProduction/RemoteControl/Source/RemoteControlUI/Private/RemoteControlUIModule.cpp
@@ -149,6 +149,20 @@ namespace RemoteControlUIModule
 		Toolkit->InitRemoteControlPresetEditor(Mode, EditWithinLevelEditor, Preset);
 	}
 
+	int32 GetNumStaticMaterials(UObject* InOuterObject)
+	{
+		if (UStaticMesh* StaticMesh = Cast<UStaticMesh>(InOuterObject))
+		{
+			return StaticMesh->GetStaticMaterials().Num();
+		}
+		if (USkeletalMesh* SkeletalMesh = Cast<USkeletalMesh>(InOuterObject))
+		{
+			return SkeletalMesh->GetMaterials().Num();
+		}
+
+		return 0;
+	};
+
 	FObjectProperty* GetObjectProperty(const FRCExposesPropertyArgs& InPropertyArgs)
 	{
 		if (!ensureMsgf(InPropertyArgs.IsValid(), TEXT("Extension Property Args not valid")))
@@ -825,81 +839,19 @@ void FRemoteControlUIModule::TryOverridingMaterials(const FRCExposesPropertyArgs
 	{
 		return;
 	}
-
+	
 	const FRCExposesPropertyArgs::EType ExtensionArgsType = InPropertyArgs.GetType();
 
 	if (ExtensionArgsType == FRCExposesPropertyArgs::EType::E_Handle)
 	{
-		UObject* OriginalMaterialAsObject = nullptr;
-		TSharedPtr<IPropertyHandle> PropertyHandle = InPropertyArgs.PropertyHandle;
-
-		if (PropertyHandle->GetValue(OriginalMaterialAsObject) == FPropertyAccess::Success)
-		{
-			if (UMaterialInterface* OriginalMaterial = Cast<UMaterialInterface>(OriginalMaterialAsObject))
-			{
-				FName MaterialSlotNameToBeOverriden;
-
-				// NOTE : Obtain the source object information from the property.
-
-				if (PropertyHandle->GetNumOuterObjects() == 1)
-				{
-					TArray<UObject*> OuterObjects;
-
-					PropertyHandle->GetOuterObjects(OuterObjects);
-
-					check(OuterObjects.Num());
-
-					if (UStaticMesh* StaticMesh = Cast<UStaticMesh>(OuterObjects[0]))
-					{
-						// NOTE : As 'StaticMaterials' must be protected for async build, always use the accessors even internally.
-						for (const FStaticMaterial StaticMaterial : StaticMesh->GetStaticMaterials())
-						{
-							if (StaticMaterial.MaterialInterface == OriginalMaterial)
-							{
-								MaterialSlotNameToBeOverriden = StaticMaterial.MaterialSlotName;
-
-								break;
-							}
-						}
-					}
-					else if (USkeletalMesh* SkeletalMesh = Cast<USkeletalMesh>(OuterObjects[0]))
-					{
-						// NOTE : As 'Materials' must be protected for async build, always use the accessors even internally.
-						for (const FSkeletalMaterial& SkeletalMaterial : SkeletalMesh->GetMaterials())
-						{
-							if (SkeletalMaterial.MaterialInterface == OriginalMaterial)
-							{
-								MaterialSlotNameToBeOverriden = SkeletalMaterial.MaterialSlotName;
-
-								break;
-							}
-						}
-					}
-
-					if (!ensureMsgf(!MaterialSlotNameToBeOverriden.IsNone(), TEXT("Material Property could not be exposed because the property does not contain any valid slot names.")))
-					{
-						return;
-					}
-
-					if (UMeshComponent* MeshComponentToBeModified = GetSelectedMeshComponentToBeModified(OuterObjects[0], OriginalMaterial))
-					{
-						FScopedTransaction Transaction(LOCTEXT("OverrideMaterial", "Override Material"));
-						const int32 TargetMaterialIndex = MeshComponentToBeModified->GetMaterialIndex(MaterialSlotNameToBeOverriden);
-
-						if (FComponentEditorUtils::AttemptApplyMaterialToComponent(MeshComponentToBeModified, OriginalMaterial, TargetMaterialIndex))
-						{
-							RefreshPanels();
-						}
-					}
-				}
-			}
-		}
+		ensureMsgf(false, TEXT("Override materials can't be done with property handle arguments type"));
 	}
+	// Override material only if PropertyArgs is OwnerObject
 	else if (ExtensionArgsType == FRCExposesPropertyArgs::EType::E_OwnerObject)
 	{
 		FRCFieldPathInfo RCFieldPathInfo(InPropertyArgs.PropertyPath);
 		RCFieldPathInfo.Resolve(InPropertyArgs.OwnerObject);
-		FRCFieldResolvedData ResolvedData = RCFieldPathInfo.GetResolvedData();
+		const FRCFieldResolvedData ResolvedData = RCFieldPathInfo.GetResolvedData();
 		FObjectProperty* ObjectProperty = CastField<FObjectProperty>(InPropertyArgs.Property);
 
 		if (!ObjectProperty)
@@ -912,31 +864,17 @@ void FRemoteControlUIModule::TryOverridingMaterials(const FRCExposesPropertyArgs
 		OriginalMaterial = OriginalMaterial ? OriginalMaterial : UMaterial::GetDefaultMaterial(MD_Surface);
 		if (OriginalMaterial)
 		{
-			int32 MaterialIndex = INDEX_NONE;
-
-			int32 StartPathArrayIndex = -1;
-			int32 EndPathArrayIndex = -1;
-			InPropertyArgs.PropertyPath.FindChar('[', StartPathArrayIndex);
-			InPropertyArgs.PropertyPath.FindChar(']', EndPathArrayIndex);
-
-			if (StartPathArrayIndex != INDEX_NONE && EndPathArrayIndex != INDEX_NONE)
-			{
-				const FString MaterialSlotNameToBeOverriden = InPropertyArgs.PropertyPath.Mid(StartPathArrayIndex + 1, EndPathArrayIndex - 1 - StartPathArrayIndex);
-				LexFromString(MaterialIndex, *MaterialSlotNameToBeOverriden);
-			}
-
-			if (!ensureMsgf(MaterialIndex > -1, TEXT("Material Property could not be exposed because the property does not contain any valid slot index.")))
-			{
-				return;
-			}
-
 			if (UMeshComponent* MeshComponentToBeModified = GetSelectedMeshComponentToBeModified(InPropertyArgs.OwnerObject, OriginalMaterial))
 			{
-				FScopedTransaction Transaction(LOCTEXT("OverrideMaterial", "Override Material"));
-
-				if (FComponentEditorUtils::AttemptApplyMaterialToComponent(MeshComponentToBeModified, OriginalMaterial, MaterialIndex))
+				const int32 NumStaticMaterials = RemoteControlUIModule::GetNumStaticMaterials(InPropertyArgs.OwnerObject);
+				if (NumStaticMaterials > 0)
 				{
-					RefreshPanels();
+					FScopedTransaction Transaction(LOCTEXT("OverrideMaterial", "Override Material"));
+
+					if (FComponentEditorUtils::AttemptApplyMaterialToComponent(MeshComponentToBeModified, OriginalMaterial, NumStaticMaterials - 1))
+					{
+						RefreshPanels();
+					}
 				}
 			}
 		}
