- Adding special handler for ndisplay properties in widget registry

diff --git a/Engine/Plugins/VirtualProduction/RemoteControl/Source/RemoteControlUI/Private/UI/RCPanelWidgetRegistry.cpp b/Engine/Plugins/VirtualProduction/RemoteControl/Source/RemoteControlUI/Private/UI/RCPanelWidgetRegistry.cpp
--- a/Engine/Plugins/VirtualProduction/RemoteControl/Source/RemoteControlUI/Private/UI/RCPanelWidgetRegistry.cpp
+++ b/Engine/Plugins/VirtualProduction/RemoteControl/Source/RemoteControlUI/Private/UI/RCPanelWidgetRegistry.cpp
@@ -2,6 +2,8 @@
 
 #include "RCPanelWidgetRegistry.h"
 
+#include "Engine/BlueprintGeneratedClass.h"
+#include "GameFramework/Actor.h"
 #include "IDetailTreeNode.h"
 #include "IPropertyRowGenerator.h"
 #include "Modules/ModuleManager.h"
@@ -82,6 +84,14 @@ namespace WidgetRegistryUtils
 	}
 }
 
+FRCPanelWidgetRegistry::FRCPanelWidgetRegistry()
+{
+	FRCTreeNodeFinderHandler NewHandler;
+	NewHandler.IsExceptionFunc = [this](UObject* InObject, const FString& InField, ERCFindNodeMethod InFindMethod) { return IsNDisplayObject(InObject, InField, InFindMethod); };
+	NewHandler.FinderFunction = [this](UObject* InObject, const FString& InField, ERCFindNodeMethod InFindMethod) { return FindNDisplayTreeNode(InObject, InField, InFindMethod); };
+	SpecialTreeNodeHandlers.Add(MoveTemp(NewHandler));
+}
+
 TSharedPtr<IDetailTreeNode> FRCPanelWidgetRegistry::GetObjectTreeNode(UObject* InObject, const FString& InField, ERCFindNodeMethod InFindMethod)
 {
 	const TPair<TWeakObjectPtr<UObject>, FString> CacheKey{InObject, InField};
@@ -96,7 +106,18 @@ TSharedPtr<IDetailTreeNode> FRCPanelWidgetRegistry::GetObjectTreeNode(UObject* I
 			TreeNodeCache.Remove(CacheKey);
 		}
 	}
-	
+
+	//Verify if desired object has special handling and cache result if it does
+	for (const FRCTreeNodeFinderHandler& Handler : SpecialTreeNodeHandlers)
+	{
+		if (Handler.IsExceptionFunc(InObject, InField, InFindMethod))
+		{
+			TSharedPtr<IDetailTreeNode> Node = Handler.FinderFunction(InObject, InField, InFindMethod);
+			TreeNodeCache.Add(CacheKey, Node);
+			return Node;
+		}
+	}
+		
 	TSharedPtr<IPropertyRowGenerator> Generator;
 	TWeakObjectPtr<UObject> WeakObject = InObject;
 	
@@ -110,14 +131,13 @@ TSharedPtr<IDetailTreeNode> FRCPanelWidgetRegistry::GetObjectTreeNode(UObject* I
 		// We can override the validation with a lambda since the validation function in PRG is not necessary for our implementation
 		auto ValidationLambda = ([](const FRootPropertyNodeList& PropertyNodeList) { return true; });
 		FPropertyRowGeneratorArgs Args;
-		Args.bShouldShowHiddenProperties = true;
 		Generator = FModuleManager::GetModuleChecked<FPropertyEditorModule>("PropertyEditor").CreatePropertyRowGenerator(Args);
 		Generator->SetCustomValidatePropertyNodesFunction(FOnValidatePropertyRowGeneratorNodes::CreateLambda(MoveTemp(ValidationLambda)));
 		Generator->SetObjects({InObject});
 		ObjectToRowGenerator.Add(WeakObject, Generator);
 	}
 	
-	TSharedPtr<IDetailTreeNode> Node = WidgetRegistryUtils::FindNode(Generator->GetRootTreeNodes(), InField, ERCFindNodeMethod::Path);
+	TSharedPtr<IDetailTreeNode> Node = WidgetRegistryUtils::FindNode(Generator->GetRootTreeNodes(), InField, InFindMethod);
 	// Cache the node to avoid having to do the recursive find again.
 	TreeNodeCache.Add(CacheKey, Node);
 	
@@ -168,3 +188,65 @@ void FRCPanelWidgetRegistry::Clear()
 	StructToRowGenerator.Empty();
 	TreeNodeCache.Empty();
 }
+
+bool FRCPanelWidgetRegistry::IsNDisplayObject(UObject* InObject, const FString& InField, ERCFindNodeMethod InFindMethod)
+{
+	if (InObject)
+	{
+		AActor* Actor = InObject->GetTypedOuter<AActor>();
+		if (Actor != nullptr)
+		{
+			UBlueprintGeneratedClass* BlueprintClass = Cast<UBlueprintGeneratedClass>(Actor->GetClass());
+			if (BlueprintClass && BlueprintClass->ClassDefaultObject)
+			{
+				static constexpr TCHAR ConfigurationDataName[] = TEXT("DisplayClusterConfigurationData");
+				const FString ObjectClassName = InObject->GetClass()->GetName();
+				if (ObjectClassName.Contains(ConfigurationDataName))
+				{
+					//Can't easily test for actor class without depending on ndisplay
+					//If user have inheritance in the way naming is not what's expected.
+					//If we know we are dealing with the config object and it's under a blueprint generated actor, good chance we can handle it
+					return true;
+				}
+			}
+		}
+	}
+
+	return false;
+}
+
+TSharedPtr<IDetailTreeNode> FRCPanelWidgetRegistry::FindNDisplayTreeNode(UObject* InObject, const FString& InField, ERCFindNodeMethod InFindMethod)
+{
+	//To support ndisplay customization, we need to go back to the root actor to trigger all required customizations to see the treenode
+	AActor* Actor = InObject->GetTypedOuter<AActor>();
+	if (Actor == nullptr)
+	{
+		return nullptr;
+	}
+
+	//Run the generation using the owner actor instead of the subobject
+	InObject = Actor;
+
+	TSharedPtr<IPropertyRowGenerator> Generator;
+	TWeakObjectPtr<UObject> WeakObject = InObject;
+
+	if (TSharedPtr<IPropertyRowGenerator>* FoundGenerator = ObjectToRowGenerator.Find(WeakObject))
+	{
+		Generator = *FoundGenerator;
+	}
+	else
+	{
+		// Since we must keep many PRG objects alive in order to access the handle data, validating the nodes each tick is very taxing.
+		// We can override the validation with a lambda since the validation function in PRG is not necessary for our implementation
+		auto ValidationLambda = ([](const FRootPropertyNodeList& PropertyNodeList) { return true; });
+		FPropertyRowGeneratorArgs Args;
+		Args.bShouldShowHiddenProperties = false;
+		Generator = FModuleManager::GetModuleChecked<FPropertyEditorModule>("PropertyEditor").CreatePropertyRowGenerator(Args);
+		Generator->SetCustomValidatePropertyNodesFunction(FOnValidatePropertyRowGeneratorNodes::CreateLambda(MoveTemp(ValidationLambda)));
+
+		Generator->SetObjects({ InObject });
+		ObjectToRowGenerator.Add(WeakObject, Generator);
+	}
+
+	return WidgetRegistryUtils::FindNode(Generator->GetRootTreeNodes(), InField, InFindMethod);
+}
diff --git a/Engine/Plugins/VirtualProduction/RemoteControl/Source/RemoteControlUI/Private/UI/RCPanelWidgetRegistry.h b/Engine/Plugins/VirtualProduction/RemoteControl/Source/RemoteControlUI/Private/UI/RCPanelWidgetRegistry.h
--- a/Engine/Plugins/VirtualProduction/RemoteControl/Source/RemoteControlUI/Private/UI/RCPanelWidgetRegistry.h
+++ b/Engine/Plugins/VirtualProduction/RemoteControl/Source/RemoteControlUI/Private/UI/RCPanelWidgetRegistry.h
@@ -16,9 +16,21 @@ enum class ERCFindNodeMethod
 	Path  // Find a node by providing a property path.
 };
 
+/** Local special case handlers to generate widgets to be displayed */
+using FRCTreeNodeFinderIsException = TFunction< bool(UObject* /*InObject*/, const FString& /*InField*/, ERCFindNodeMethod /*InFindMethod*/) >;
+using FRCTreeNodeFinderFunction = TFunction< TSharedPtr<IDetailTreeNode>(UObject* /*InObject*/, const FString& /*InField*/, ERCFindNodeMethod /*InFindMethod*/) >;
+struct FRCTreeNodeFinderHandler
+{
+	FRCTreeNodeFinderIsException IsExceptionFunc;
+	FRCTreeNodeFinderFunction FinderFunction;
+};
+
 class FRCPanelWidgetRegistry
 {
 public:
+	
+	FRCPanelWidgetRegistry();
+
 	/**
 	 * Get a detail tree node for a given object and property.
 	 * @param InObject the object used to generate the detail row.
@@ -50,6 +62,12 @@ public:
 	 */
 	void Clear();
 
+private:
+
+	/** Special handlers for ndisplay root actor and its details customizations */
+	bool IsNDisplayObject(UObject* InObject, const FString& InField, ERCFindNodeMethod InFindMethod);
+	TSharedPtr<IDetailTreeNode> FindNDisplayTreeNode(UObject* InObject, const FString& InField, ERCFindNodeMethod InFindMethod);
+
 private:
 	/** Map of objects to row generator, used to have one row generator per object. */
 	TMap<TWeakObjectPtr<UObject>, TSharedPtr<IPropertyRowGenerator>> ObjectToRowGenerator;
@@ -57,4 +75,6 @@ private:
 	TMap<TSharedPtr<FStructOnScope>, TSharedPtr<IPropertyRowGenerator>> StructToRowGenerator;
 	/** Cache of Object&Field to tree nodes. */
 	TMap<TPair<TWeakObjectPtr<UObject>, FString>, TWeakPtr<IDetailTreeNode>> TreeNodeCache;
+	/** List of tree node finder handlers for certain type */
+	TArray<FRCTreeNodeFinderHandler> SpecialTreeNodeHandlers;
 };
