Remote Control - Fix crash when changing lightcard location

diff --git a/Engine/Plugins/VirtualProduction/RemoteControl/Source/WebRemoteControl/Private/WebSocketMessageHandler.cpp b/Engine/Plugins/VirtualProduction/RemoteControl/Source/WebRemoteControl/Private/WebSocketMessageHandler.cpp
--- a/Engine/Plugins/VirtualProduction/RemoteControl/Source/WebRemoteControl/Private/WebSocketMessageHandler.cpp
+++ b/Engine/Plugins/VirtualProduction/RemoteControl/Source/WebRemoteControl/Private/WebSocketMessageHandler.cpp
@@ -386,12 +386,18 @@ void FWebSocketMessageHandler::ProcessChangedProperties()
 		// Each client will have a custom payload that doesnt contain the events it triggered.
 		for (const TPair<FGuid, TSet<FGuid>>& ClientToEventsPair : Entry.Value)
 		{
-			TArray<uint8> WorkingBuffer;
-			if (ClientToEventsPair.Value.Num() && WritePropertyChangeEventPayload(Preset, ClientToEventsPair.Value, WorkingBuffer))
+			// This should be improved in the future, we create one message per modified property to avoid
+			// sending a list of non-uniform properties (ie. Color, Transform), ideally these should be grouped by underlying
+			// property class. See UE-139683
+			for (const FGuid& Id : ClientToEventsPair.Value)
 			{
-				TArray<uint8> Payload;
-				WebRemoteControlUtils::ConvertToUTF8(WorkingBuffer, Payload);
-				Server->Send(ClientToEventsPair.Key, Payload);
+				TArray<uint8> WorkingBuffer;
+				if (ClientToEventsPair.Value.Num() && WritePropertyChangeEventPayload(Preset, { Id }, WorkingBuffer))
+				{
+					TArray<uint8> Payload;
+					WebRemoteControlUtils::ConvertToUTF8(WorkingBuffer, Payload);
+					Server->Send(ClientToEventsPair.Key, Payload);
+				}	
 			}
 		}
 	}
