Fix issues with remote control rebinding

diff --git a/Engine/Plugins/VirtualProduction/RemoteControl/Source/RemoteControl/Private/RemoteControlBinding.cpp b/Engine/Plugins/VirtualProduction/RemoteControl/Source/RemoteControl/Private/RemoteControlBinding.cpp
--- a/Engine/Plugins/VirtualProduction/RemoteControl/Source/RemoteControl/Private/RemoteControlBinding.cpp
+++ b/Engine/Plugins/VirtualProduction/RemoteControl/Source/RemoteControl/Private/RemoteControlBinding.cpp
@@ -5,7 +5,10 @@
 #include "Engine/Engine.h"
 #include "Engine/EngineTypes.h"
 #include "Engine/World.h"
-#include "GameFramework/Actor.h"
+#include "HAL/IConsoleManager.h"
+#include "IRemoteControlModule.h"
+#include "RemoteControlPreset.h"
+#include "RemoteControlSettings.h"
 #include "UObject/SoftObjectPath.h"
 #include "UObject/SoftObjectPtr.h"
 
@@ -16,6 +19,8 @@
 #include "Misc/App.h"
 #endif
 
+#define LOCTEXT_NAMESPACE "RemoteControlBinding"
+
 namespace
 {
 	/**
@@ -53,8 +58,32 @@ namespace
 #endif
 		return CounterpartObject ? CounterpartObject : Object;
 	}
+
+	void DumpBindings(const TArray<FString>& Args)
+	{
+		if (Args.Num())
+		{
+			if (URemoteControlPreset* Preset = FindObject<URemoteControlPreset>(ANY_PACKAGE, *Args[0]))
+			{
+				TStringBuilder<1000> Output;
+				for (URemoteControlBinding* Binding : Preset->Bindings)
+				{
+					if (GEngine)
+					{
+						GEngine->Exec(nullptr, *FString::Format(TEXT("obj dump {0} hide=\"actor,object,lighting,movement\""), { Binding->GetPathName() }));
+					}
+				}
+			}
+		}
+	}
 }
 
+static FAutoConsoleCommand CCmdDumpBindings = FAutoConsoleCommand(
+	TEXT("RemoteControl.DumpBindings"),
+	TEXT("Dumps all binding info on the preset passed as argument."),
+	FConsoleCommandWithArgsDelegate::CreateStatic(&DumpBindings)
+);
+
 void URemoteControlLevelIndependantBinding::SetBoundObject(const TSoftObjectPtr<UObject>& InObject)
 {
 	BoundObject = InObject;
@@ -105,6 +134,20 @@ void URemoteControlLevelDependantBinding::SetBoundObject(const TSoftObjectPtr<UO
 		SubLevelSelectionMap.FindOrAdd(OuterLevel->GetWorld()) = OuterLevel;
 		
 		Name = EditorObject->GetName();
+		
+		if (BindingContext.OwnerActorName.IsNone() || !GetDefault<URemoteControlSettings>()->bUseRebindingContext) 
+		{
+			InitializeBindingContext(InObject.Get());
+		}
+	}
+}
+
+void URemoteControlLevelDependantBinding::SetBoundObject_OverrideContext(const TSoftObjectPtr<UObject>& InObject)
+{
+	SetBoundObject(InObject);
+	if (InObject)
+	{
+		InitializeBindingContext(InObject.Get());
 	}
 }
 
@@ -159,7 +202,21 @@ UObject* URemoteControlLevelDependantBinding::Resolve() const
 
 	if (Object)
 	{
+		// Make sure we don't resolve on a subobject of a dying parent actor.
+		if (AActor* OwnerActor = Object->GetTypedOuter<AActor>())
+		{
+			if (!::IsValid(OwnerActor))
+			{
+				return nullptr;
+			}
+		}
+
 		LevelWithLastSuccessfulResolve = Object->GetTypedOuter<ULevel>();
+
+		if (BindingContext.OwnerActorName.IsNone())
+		{
+			InitializeBindingContext(Object);
+		}
 	}
 
 	return FindObjectInCounterpartWorld(Object, ECounterpartWorldTarget::PIE);
@@ -222,6 +279,7 @@ TSoftObjectPtr<UObject> URemoteControlLevelDependantBinding::ResolveForCurrentWo
 			if (TSoftObjectPtr<UObject> ObjectPtr = BoundObjectMap.FindRef(WeakLevel))
 			{
 				SubLevelSelectionMap.FindOrAdd(World) = WeakLevel;
+
 				return ObjectPtr;
 			}
 		}
@@ -229,6 +287,41 @@ TSoftObjectPtr<UObject> URemoteControlLevelDependantBinding::ResolveForCurrentWo
 	return nullptr;
 }
 
+void URemoteControlLevelDependantBinding::InitializeBindingContext(UObject* InObject) const
+{
+	if (InObject)
+	{
+		BindingContext.SupportedClass = InObject->GetClass();
+
+		if (InObject->IsA<AActor>())
+		{
+			BindingContext.OwnerActorClass = InObject->GetClass();
+			BindingContext.OwnerActorName = InObject->GetFName();
+		}
+		else if (InObject->IsA<UActorComponent>())
+		{
+			AActor* Owner = InObject->GetTypedOuter<AActor>();
+			check(Owner);
+			BindingContext.ComponentName = InObject->GetFName();
+			BindingContext.OwnerActorClass = Owner->GetClass();
+			BindingContext.OwnerActorName = Owner->GetFName();
+		}
+		else
+		{
+			AActor* Owner = InObject->GetTypedOuter<AActor>();
+			check(Owner);
+			BindingContext.ComponentName = InObject->GetFName();
+			BindingContext.OwnerActorClass = Owner->GetClass();
+			BindingContext.OwnerActorName = Owner->GetFName();
+		}
+	}
+}
+
+UClass* URemoteControlLevelDependantBinding::GetSupportedOwnerClass() const
+{
+	return BindingContext.OwnerActorClass.LoadSynchronous();
+}
+
 UWorld* URemoteControlLevelDependantBinding::GetCurrentWorld()
 {
 	// Since this is used to retrieve the binding in the map, we never use the PIE world in editor.
@@ -272,3 +365,5 @@ void URemoteControlLevelDependantBinding::SetBoundObject(const TSoftObjectPtr<UL
 	ensure(ObjectName.Len());
 	Name = MoveTemp(ObjectName);
 }
+
+#undef LOCTEXT_NAMESPACE 
\ No newline at end of file
diff --git a/Engine/Plugins/VirtualProduction/RemoteControl/Source/RemoteControl/Private/RemoteControlEntity.cpp b/Engine/Plugins/VirtualProduction/RemoteControl/Source/RemoteControl/Private/RemoteControlEntity.cpp
--- a/Engine/Plugins/VirtualProduction/RemoteControl/Source/RemoteControl/Private/RemoteControlEntity.cpp
+++ b/Engine/Plugins/VirtualProduction/RemoteControl/Source/RemoteControl/Private/RemoteControlEntity.cpp
@@ -17,6 +17,18 @@ TArray<UObject*> FRemoteControlEntity::GetBoundObjects() const
 	return ResolvedObjects.FilterByPredicate([](const UObject* Object){ return !!Object; });
 }
 
+UObject* FRemoteControlEntity::GetBoundObject() const
+{
+	TArray<UObject*> ResolvedObjects;
+	ResolvedObjects.Reserve(Bindings.Num());
+	Algo::TransformIf(Bindings, ResolvedObjects,
+		[](TWeakObjectPtr<URemoteControlBinding> WeakBinding) { return WeakBinding.IsValid(); },
+		[](TWeakObjectPtr<URemoteControlBinding> WeakBinding) { return WeakBinding->Resolve(); });
+
+	UObject** Obj = ResolvedObjects.FindByPredicate([](const UObject* Object) { return !!Object; });
+	return Obj ? *Obj : nullptr;
+}
+
 const TArray<TWeakObjectPtr<URemoteControlBinding>>& FRemoteControlEntity::GetBindings() const
 {
 	return Bindings;
@@ -46,15 +58,26 @@ void FRemoteControlEntity::BindObject(UObject* InObjectToBind)
 		return;
 	}
 
-	URemoteControlBinding* Binding = Owner->FindOrAddBinding(InObjectToBind);
-
 	if (Bindings.Num() == 0)
 	{
-		Bindings.Emplace(Binding);
+		Bindings.Emplace(Owner->FindOrAddBinding(InObjectToBind));
 	}
 	else
 	{
-		Bindings[0] = Binding;
+		// Don't modify the binding itself since that would modify all other properties pointing to this binding.
+		// Try to first find a binding that has the same bound object map to preserve that information.
+		if (URemoteControlBinding* MatchingBinding = Owner->FindMatchingBinding(Bindings[0].Get(), InObjectToBind))
+		{
+			Bindings[0] = MatchingBinding;
+			
+		}
+		else
+		{
+			Bindings[0] = DuplicateObject<URemoteControlBinding>(Bindings[0].Get(), Owner.Get());
+			Owner->Bindings.Add(Bindings[0].Get());
+			Bindings[0]->SetBoundObject(InObjectToBind);
+		}
+
 	}
 
 	OnEntityModifiedDelegate.ExecuteIfBound(Id);
diff --git a/Engine/Plugins/VirtualProduction/RemoteControl/Source/RemoteControl/Private/RemoteControlPreset.cpp b/Engine/Plugins/VirtualProduction/RemoteControl/Source/RemoteControl/Private/RemoteControlPreset.cpp
--- a/Engine/Plugins/VirtualProduction/RemoteControl/Source/RemoteControl/Private/RemoteControlPreset.cpp
+++ b/Engine/Plugins/VirtualProduction/RemoteControl/Source/RemoteControl/Private/RemoteControlPreset.cpp
@@ -695,6 +695,8 @@ void URemoteControlPreset::PostLoad()
 	CreatePropertyWatchers();
 
 	PostLoadProperties();
+
+	RemoveUnusedBindings();
 }
 
 void URemoteControlPreset::PostDuplicate(bool bDuplicateForPIE)
@@ -930,6 +932,61 @@ URemoteControlBinding* URemoteControlPreset::FindOrAddBinding(const TSoftObjectP
 	return NewBinding;
 }
 
+URemoteControlBinding* URemoteControlPreset::FindMatchingBinding(const URemoteControlBinding* InBinding, UObject* InObject)
+{
+	for (URemoteControlBinding* Binding : Bindings)
+	{
+		URemoteControlLevelDependantBinding* LevelDependantBindingIt = Cast<URemoteControlLevelDependantBinding>(Binding);
+		const URemoteControlLevelDependantBinding* InLevelDependingBinding = Cast<URemoteControlLevelDependantBinding>(InBinding);
+
+		if (InBinding == Binding
+			|| Binding->Resolve() != InObject)
+		{
+			continue;
+		}
+
+		if (LevelDependantBindingIt->BindingContext != InLevelDependingBinding->BindingContext)
+		{
+			continue;
+		}
+
+		if (LevelDependantBindingIt && InLevelDependingBinding)
+		{
+			TSoftObjectPtr<ULevel> CurrentWorldLevel = LevelDependantBindingIt->SubLevelSelectionMap.FindRef(InObject->GetWorld());
+
+			bool bSameBoundObjectMap = true;
+			// Check if the binding it has the same bound object map except for the current level.
+			for (const TPair<TSoftObjectPtr<ULevel>, TSoftObjectPtr<UObject>>& Pair : InLevelDependingBinding->BoundObjectMap)
+			{
+				if (Pair.Key != CurrentWorldLevel)
+				{
+					TSoftObjectPtr<UObject>* BoundObject = LevelDependantBindingIt->BoundObjectMap.Find(Pair.Key);
+					if (BoundObject)
+					{
+						if (*BoundObject != Pair.Value)
+						{
+							bSameBoundObjectMap = false;
+							break;
+						}
+					}
+					else
+					{
+						bSameBoundObjectMap = false;
+						break;
+					}
+				}
+			}
+
+			if (bSameBoundObjectMap)
+			{
+				return Binding;
+			}
+		}
+	}
+
+	return nullptr;
+}
+
 void URemoteControlPreset::OnEntityModified(const FGuid& EntityId)
 {
 	PerFrameUpdatedEntities.Add(EntityId);
@@ -1054,6 +1111,26 @@ void URemoteControlPreset::CreatePropertyWatchers()
 	}
 }
 
+void URemoteControlPreset::RemoveUnusedBindings()
+{
+	TSet<TWeakObjectPtr<URemoteControlBinding>> ReferencedBindings;
+	for (TSharedPtr<FRemoteControlEntity> Entity : Registry->GetExposedEntities())
+	{
+		for (TWeakObjectPtr<URemoteControlBinding> Binding : Entity->GetBindings())
+		{
+			ReferencedBindings.Add(Binding);
+		}
+	}
+
+	for (auto It = Bindings.CreateIterator(); It; It++)
+	{
+		if (!ReferencedBindings.Contains(*It))
+		{
+			It.RemoveCurrent();
+		}
+	}
+}
+
 void URemoteControlPreset::PostLoadProperties()
 {
 	for (const TSharedPtr<FRemoteControlProperty>& ExposedProperty : Registry->GetExposedEntities<FRemoteControlProperty>())
@@ -1264,15 +1341,16 @@ TArray<UObject*> URemoteControlPreset::ResolvedBoundObjects(FName FieldLabel)
 
 void URemoteControlPreset::RebindUnboundEntities()
 {
+	Modify();
 	RebindingManager->Rebind(this);
 	Algo::Transform(Registry->GetExposedEntities(), PerFrameUpdatedEntities, [](const TSharedPtr<FRemoteControlEntity>& Entity) { return Entity->GetId(); });
 }
 
-void URemoteControlPreset::RebindAllEntitiesUnderSameActor(const FGuid& EntityId, AActor* NewActor)
+void URemoteControlPreset::RebindAllEntitiesUnderSameActor(const FGuid& EntityId, AActor* NewActor, bool bUseRebindingContext)
 {
 	if (TSharedPtr<FRemoteControlEntity> Entity = Registry->GetExposedEntity(EntityId))
 	{
-		RebindingManager->RebindAllEntitiesUnderSameActor(this, Entity, NewActor);
+		RebindingManager->RebindAllEntitiesUnderSameActor(this, Entity, NewActor, bUseRebindingContext);
 	}
 }
 
@@ -2139,12 +2217,6 @@ void URemoteControlPreset::CleanUpBindings()
 				{
 					PerFrameUpdatedEntities.Add(Entity->GetId());
 				}
-
-				if (BindingsToDelete.Contains(It->Get()))
-				{
-					It.RemoveCurrent();
-					break;
-				}
 			}
 		}
 	}
diff --git a/Engine/Plugins/VirtualProduction/RemoteControl/Source/RemoteControl/Private/RemoteControlPresetRebindingManager.cpp b/Engine/Plugins/VirtualProduction/RemoteControl/Source/RemoteControl/Private/RemoteControlPresetRebindingManager.cpp
--- a/Engine/Plugins/VirtualProduction/RemoteControl/Source/RemoteControl/Private/RemoteControlPresetRebindingManager.cpp
+++ b/Engine/Plugins/VirtualProduction/RemoteControl/Source/RemoteControl/Private/RemoteControlPresetRebindingManager.cpp
@@ -2,9 +2,11 @@
 
 #include "RemoteControlPresetRebindingManager.h"
 
+#include "Algo/AllOf.h"
 #include "Algo/Transform.h"
 #include "Components/ActorComponent.h"
 #include "Engine/Brush.h"
+#include "Engine/Classes/Components/ActorComponent.h"
 #include "Engine/Engine.h"
 #include "EngineUtils.h"
 #include "Misc/Char.h"
@@ -18,6 +20,8 @@
 #include "ActorEditorUtils.h"
 #endif
 
+static TAutoConsoleVariable<int32> CVarRemoteControlRebindingUseLegacyAlgo(TEXT("RemoteControl.UseLegacyRebinding"), 0, TEXT("Whether to revert to the old algorithm when doing a Rebind All."));
+
 namespace RCPresetRebindingManager
 {
 	UWorld* GetCurrentWorld()
@@ -130,6 +134,29 @@ namespace RCPresetRebindingManager
 		return ObjectMap;
 	}
 
+	bool GetActorsOfClass(UClass* InTargetClass, TArray<AActor*>& OutActors)
+	{
+		auto IsObjectOfClass = [InTargetClass](UObject* Object)
+		{
+			return Object
+				&& IsValidActorForRebinding(CastChecked<AActor>(Object))
+				&& Object->IsA(InTargetClass);
+		};
+
+		if (UWorld* World = GetCurrentWorld())
+		{
+			for (TActorIterator<AActor> It(World, AActor::StaticClass(), EActorIteratorFlags::SkipPendingKill); It; ++It)
+			{
+				if (IsObjectOfClass(*It))
+				{
+					OutActors.Add(*It);
+				}
+			}
+		}
+
+		return OutActors.Num() != 0;
+	}
+
 	TMap<UClass*, TMap<uint32, TSet<UObject*>>> CreateInitialBoundObjectMap(TConstArrayView<TSharedPtr<FRemoteControlEntity>> BoundEntities)
 	{
 		TMap<UClass*, TMap<uint32, TSet<UObject*>>> BoundObjectMap;
@@ -196,6 +223,17 @@ UObject* FNameBasedRebindingPolicy::FindMatch(const FFindBindingForEntityArgs& A
 }
 
 void FRemoteControlPresetRebindingManager::Rebind(URemoteControlPreset* Preset)
+{
+	if (CVarRemoteControlRebindingUseLegacyAlgo.GetValueOnAnyThread() == 1)
+	{
+		Rebind_Legacy(Preset);
+	}
+	else
+	{
+		Rebind_NewAlgo(Preset);
+	}
+}
+void FRemoteControlPresetRebindingManager::Rebind_Legacy(URemoteControlPreset* Preset)
 {
 	check(Preset);
 
@@ -228,14 +266,317 @@ void FRemoteControlPresetRebindingManager::Rebind(URemoteControlPreset* Preset)
 
 	for (TPair<UClass*, TArray<TSharedPtr<FRemoteControlEntity>>>& Entry : EntitiesGroupedBySupportedOwnerClass)
 	{
-		RebindEntitiesForClass(Entry.Key, Entry.Value);
+		if (Entry.Key->IsChildOf<AActor>())
+		{
+			RebindEntitiesForActorClass_Legacy(Preset, Entry.Key, Entry.Value);
+		}
+	}
+
+	for (TPair<UClass*, TArray<TSharedPtr<FRemoteControlEntity>>>& Entry : EntitiesGroupedBySupportedOwnerClass)
+	{
+		// Do an additional pass in case the bindings did not have the necessary info to use the previous rebinding process. 
+		RebindEntitiesForClass_Legacy(Entry.Key, Entry.Value);
+	}
+}
+
+
+void FRemoteControlPresetRebindingManager::Rebind_NewAlgo(URemoteControlPreset* Preset)
+{
+	check(Preset);
+
+	// Separate bindings based on if they are bound or not.
+	TArray<URemoteControlLevelDependantBinding*> ValidBindings;
+	TArray<URemoteControlLevelDependantBinding*> InvalidBindings;
+	
+	auto ClassifyBindings = [&ValidBindings, &InvalidBindings, Preset] ()
+	{
+		ValidBindings.Reset();
+		InvalidBindings.Reset();
+
+		for (URemoteControlBinding* Binding : Preset->Bindings)
+		{
+			if (URemoteControlLevelDependantBinding* LevelDependantBinding = Cast<URemoteControlLevelDependantBinding>(Binding))
+			{
+				if (LevelDependantBinding->Resolve())
+				{
+					ValidBindings.Add(LevelDependantBinding);
+				}
+				else
+				{
+					InvalidBindings.Add(LevelDependantBinding);
+				}
+			}
+		}
+	};
+
+	ClassifyBindings();
+	const bool bAllContextsEmpty = Algo::AllOf(InvalidBindings, [](URemoteControlLevelDependantBinding* Binding) { return Binding->BindingContext.IsEmpty(); });
+	if (bAllContextsEmpty)
+	{
+		Rebind_Legacy(Preset);
+		return;
+	}
+
+	// First, try to 'rebind for all properties' bindings that are already bound.
+	// to make sure that we don't accidently rebind a property to a different actor
+	for (URemoteControlLevelDependantBinding* ValidBinding : ValidBindings)
+	{
+		UObject* ResolvedObject = ValidBinding->Resolve();
+		if (AActor* ResolvedActor = Cast<AActor>(ResolvedObject))
+		{
+			RebindAllEntitiesUnderSameActor(Preset, ValidBinding, ResolvedActor);
+		}
+		else
+		{
+			RebindAllEntitiesUnderSameActor(Preset, ValidBinding, ResolvedObject->GetTypedOuter<AActor>());
+		}
+	}
+
+	// Reclassify bindings after we rebound them.
+	ClassifyBindings();
+
+	// Create a list of what objects and components are already bound.
+	TSet<TPair<AActor*, UObject*>> BoundObjects;
+	for (URemoteControlLevelDependantBinding* ValidBinding : ValidBindings)
+	{
+		const bool bTargetingActor = !ValidBinding->BindingContext.OwnerActorName.IsNone() && ValidBinding->BindingContext.ComponentName.IsNone();
+		const bool bTargetingComponent = !ValidBinding->BindingContext.OwnerActorName.IsNone() && !ValidBinding->BindingContext.ComponentName.IsNone();
+
+		if (bTargetingActor)
+		{
+			BoundObjects.Add(TPair<AActor*, UObject*>{ Cast<AActor>(ValidBinding->Resolve()), nullptr });
+		}
+		else if (bTargetingComponent)
+		{
+			UObject* Object = ValidBinding->Resolve();
+			AActor* Actor = Object->GetTypedOuter<AActor>();
+			BoundObjects.Add(TPair<AActor*, UObject*>{ Actor, Object });
+		}
+	}
+
+	auto FindByName = [](const TArray<AActor*>& PotentialMatches, FName TargetName) -> UObject*
+	{
+		for (UObject* Object : PotentialMatches)
+		{
+			// Attempt finding by name.
+			if (Object->GetFName() == TargetName)
+			{
+				return Object;
+			}
+		}
+		return nullptr;
+	};
+
+	auto TryRebindPair = [&BoundObjects](URemoteControlBinding* InBindingToRebind, AActor* InActor, UObject* InActorComponent)
+	{
+		TPair<AActor*, UObject*> Pair = TPair<AActor*, UObject*> { InActor, InActorComponent };
+		if (!BoundObjects.Contains(Pair))
+		{
+			BoundObjects.Add(Pair);
+			InBindingToRebind->Modify();
+			if (InActorComponent)
+			{
+				InBindingToRebind->SetBoundObject(InActorComponent);
+			}
+			else
+			{
+				InBindingToRebind->SetBoundObject(InActor);
+			}
+			return true;
+		}
+		return false;
+	};
+
+	auto GetComponentBasedOnName = [](URemoteControlLevelDependantBinding* InBinding, UObject* InObject) -> UObject*
+	{
+		FName InitialComponentName = InBinding->BindingContext.ComponentName;
+		UObject* TargetComponent = FindObject<UObject>(InObject, *InitialComponentName.ToString());
+		if (TargetComponent && TargetComponent->GetClass()->IsChildOf(InBinding->BindingContext.SupportedClass.LoadSynchronous()))
+		{
+			return TargetComponent;
+		}
+		return nullptr;
+	};
+
+	auto RebindComponentBasedOnClass = [TryRebindPair](URemoteControlLevelDependantBinding* InBinding, AActor* InActorMatch)
+	{
+		if (UClass* SupportedClass = InBinding->BindingContext.SupportedClass.LoadSynchronous())
+		{
+			if (SupportedClass->IsChildOf(UActorComponent::StaticClass()))
+			{
+				TArray<UActorComponent*> Components;
+				Cast<AActor>(InActorMatch)->GetComponents(InBinding->BindingContext.SupportedClass.LoadSynchronous(), Components);
+				
+				for (UActorComponent* Component : Components)
+				{
+					if (TryRebindPair(InBinding, Component->GetOwner(), Component))
+					{
+						return true;
+					}
+				}
+			}
+		}
+		return false;
+	};
+
+	auto RebindUsingActorName = [FindByName, &BoundObjects, TryRebindPair](URemoteControlLevelDependantBinding* BindingToRebind, const TArray<AActor*>& ObjectsWithSupportedClass)
+	{
+		if (UObject* Match = FindByName(ObjectsWithSupportedClass, BindingToRebind->BindingContext.OwnerActorName))
+		{
+			return TryRebindPair(BindingToRebind, Cast<AActor>(Match), nullptr);
+		}
+		return false;
+	};
+
+	auto RebindUsingClass = [FindByName, &BoundObjects, TryRebindPair](URemoteControlLevelDependantBinding* BindingToRebind, const TArray<AActor*>& ObjectsWithSupportedClass)
+	{
+		for (UObject* PotentialMatch : ObjectsWithSupportedClass)
+		{
+			if (TryRebindPair(BindingToRebind, Cast<AActor>(PotentialMatch), nullptr))
+			{
+				return true;
+			}
+		}
+		return false;
+	};
+
+	auto TryRebindComponent = [GetComponentBasedOnName, TryRebindPair, RebindComponentBasedOnClass](URemoteControlLevelDependantBinding* BindingToRebind, UObject* PotentialMatch)
+	{
+		if (UObject* TargetComponent = GetComponentBasedOnName(BindingToRebind, PotentialMatch))
+		{
+			return TryRebindPair(BindingToRebind, TargetComponent->GetTypedOuter<AActor>(), TargetComponent);
+		}
+		else
+		{
+			return RebindComponentBasedOnClass(BindingToRebind, Cast<AActor>(PotentialMatch));
+		}
+	};
+
+	// Core of the rebinding algo
+	for (URemoteControlLevelDependantBinding* InvalidBinding : InvalidBindings)
+	{
+		if (UClass* OwnerClass = InvalidBinding->BindingContext.OwnerActorClass.LoadSynchronous())
+		{
+			const bool bRebindingComponent = !InvalidBinding->BindingContext.ComponentName.IsNone();
+			TArray<AActor*> ObjectsWithSupportedClass;
+			if (!RCPresetRebindingManager::GetActorsOfClass(OwnerClass, ObjectsWithSupportedClass))
+			{
+				continue;
+			}
+
+			if (bRebindingComponent)
+			{
+				if (UObject* Match = FindByName(ObjectsWithSupportedClass, InvalidBinding->BindingContext.OwnerActorName))
+				{
+					// Found an owner object with matching name. Try to find a component under it with a matching name.
+					TryRebindComponent(InvalidBinding, Match);
+				}
+				else
+				{
+					// Could not find an actor with the same name as the initial binding, rely only on class instead.
+					for (UObject* Object : ObjectsWithSupportedClass)
+					{
+						if (TryRebindComponent(InvalidBinding, Object))
+						{
+							break;
+						}
+					}
+				}
+			}
+			else
+			{
+				if (!RebindUsingActorName(InvalidBinding, ObjectsWithSupportedClass))
+				{
+					RebindUsingClass(InvalidBinding, ObjectsWithSupportedClass);
+				}
+			}
+		}
 	}
+
+	Preset->RemoveUnusedBindings();
 }
 
-void FRemoteControlPresetRebindingManager::RebindAllEntitiesUnderSameActor(URemoteControlPreset* Preset, const TSharedPtr<FRemoteControlEntity>& Entity, AActor* NewActor)
+
+TArray<FGuid> FRemoteControlPresetRebindingManager::RebindAllEntitiesUnderSameActor(URemoteControlPreset* Preset, URemoteControlBinding* InitialBinding, AActor* NewActor)
+{
+	check(Preset && InitialBinding && NewActor);
+
+	TArray<FGuid> ReboundEntityIds;
+
+	URemoteControlLevelDependantBinding* InitialLevelDependantBinding = Cast<URemoteControlLevelDependantBinding>(InitialBinding);
+	if (!InitialLevelDependantBinding)
+	{
+		return ReboundEntityIds;
+	}
+
+	TArray<URemoteControlLevelDependantBinding*> BindingsToRebind;
+
+	// Find bindings with the same actor as the original binding
+	for (URemoteControlBinding* Binding : Preset->Bindings)
+	{
+		if (URemoteControlLevelDependantBinding* LevelDependantBinding = Cast<URemoteControlLevelDependantBinding>(Binding))
+		{
+			// Todo: Make this more robust by comparing full actor paths instead of names.
+			if (LevelDependantBinding->BindingContext.OwnerActorClass == InitialLevelDependantBinding->BindingContext.OwnerActorClass
+				&& LevelDependantBinding->BindingContext.OwnerActorName == InitialLevelDependantBinding->BindingContext.OwnerActorName)
+			{
+				BindingsToRebind.Add(LevelDependantBinding);
+			}
+		}
+	}
+
+	// Gather entities to rebind
+	TArray<TSharedPtr<FRemoteControlEntity>> EntitiesToRebind;
+	for (const TWeakPtr<FRemoteControlEntity>& WeakEntity : Preset->GetExposedEntities<FRemoteControlEntity>())
+	{
+		TSharedPtr<FRemoteControlEntity> RCEntity = WeakEntity.Pin();
+		if (RCEntity && RCEntity->GetBindings().Num())
+		{
+			if (BindingsToRebind.Contains(RCEntity->GetBindings()[0].Get()))
+			{
+				EntitiesToRebind.Add(RCEntity);
+			}
+		}
+	}
+
+	ReboundEntityIds.Reserve(EntitiesToRebind.Num());
+	Algo::Transform(EntitiesToRebind, ReboundEntityIds, [](const TSharedPtr<FRemoteControlEntity>& RCEntityPtr) { return RCEntityPtr->GetId();  });
+
+	// Find the new binding target, then rebind.
+	for (const TSharedPtr<FRemoteControlEntity>& EntityToRebind : EntitiesToRebind)
+	{
+		URemoteControlLevelDependantBinding* EntityBinding = Cast<URemoteControlLevelDependantBinding>(EntityToRebind->GetBindings()[0].Get());
+
+		UObject* RebindTarget = NewActor;
+
+		// If the binding is targeting a component rather than an actor
+		if (!EntityBinding->BindingContext.ComponentName.IsNone())
+		{
+			if (UObject* ResolvedComponent = StaticFindObject(UObject::StaticClass(), NewActor, *EntityBinding->Name, false))
+			{
+				RebindTarget = ResolvedComponent;
+			}
+		}
+
+		EntityToRebind->BindObject(RebindTarget);
+
+		if (Preset->GetExposedEntityType(EntityToRebind->GetId()) == FRemoteControlProperty::StaticStruct())
+		{
+			StaticCastSharedPtr<FRemoteControlProperty>(EntityToRebind)->EnableEditCondition();
+		}
+	}
+
+	Preset->RemoveUnusedBindings();
+
+	return ReboundEntityIds;
+}
+
+TArray<FGuid> FRemoteControlPresetRebindingManager::RebindAllEntitiesUnderSameActor_Legacy(URemoteControlPreset* Preset, const TSharedPtr<FRemoteControlEntity>& Entity, AActor* NewActor)
 {
 	check(Preset && Entity && NewActor);
 
+	TArray<FGuid> ReboundEntities;
+
 	TMap<FName, TSet<URemoteControlBinding*>> ActorsToBindings;
 	ActorsToBindings.Reserve(Preset->Bindings.Num());
 
@@ -332,12 +673,37 @@ void FRemoteControlPresetRebindingManager::RebindAllEntitiesUnderSameActor(URemo
 		}
 
 		EntityToRebind->BindObject(RebindTarget);
+		ReboundEntities.Add(EntityToRebind->GetId());
 
 		if (Preset->GetExposedEntityType(EntityToRebind->GetId()) == FRemoteControlProperty::StaticStruct())
 		{
 			StaticCastSharedPtr<FRemoteControlProperty>(EntityToRebind)->EnableEditCondition();
 		}
 	}
+
+	Preset->RemoveUnusedBindings();
+
+	return ReboundEntities;
+}
+
+TArray<FGuid> FRemoteControlPresetRebindingManager::RebindAllEntitiesUnderSameActor(URemoteControlPreset* Preset, const TSharedPtr<FRemoteControlEntity>& Entity, AActor* NewActor, bool bUseRebindingContext)
+{
+	check(Preset && Entity && NewActor);
+	TArray<FGuid> ReboundEntityIds;
+
+	if (Entity->GetBindings().Num())
+	{
+		if (bUseRebindingContext)
+		{
+			ReboundEntityIds = RebindAllEntitiesUnderSameActor(Preset, Entity->GetBindings()[0].Get(), NewActor);
+		}
+		else
+		{
+			ReboundEntityIds = RebindAllEntitiesUnderSameActor_Legacy(Preset, Entity, NewActor);
+		}
+	}
+
+	return ReboundEntityIds;
 }
 
 TMap<UClass*, TArray<TSharedPtr<FRemoteControlEntity>>> FRemoteControlPresetRebindingManager::GroupByEntitySupportedOwnerClass(TConstArrayView<TSharedPtr<FRemoteControlEntity>> Entities) const
@@ -364,7 +730,7 @@ void FRemoteControlPresetRebindingManager::ReinitializeBindings(URemoteControlPr
 	}
 }
 
-void FRemoteControlPresetRebindingManager::RebindEntitiesForClass(UClass* Class, const TArray<TSharedPtr<FRemoteControlEntity>>& UnboundEntities)
+void FRemoteControlPresetRebindingManager::RebindEntitiesForClass_Legacy(UClass* Class, const TArray<TSharedPtr<FRemoteControlEntity>>& UnboundEntities)
 {
 	// Keep track of matched properties by keeping a record of FieldPaths hash -> Objects.
 	// This is to make sure we don't bind the same property to the same object twice.
@@ -375,11 +741,44 @@ void FRemoteControlPresetRebindingManager::RebindEntitiesForClass(UClass* Class,
 	{
 		for (const TSharedPtr<FRemoteControlEntity>& RCEntity : UnboundEntities)
 		{
-			if (UObject* Match = FDefaultRebindingPolicy::FindMatch({RCEntity, *ObjectsToConsider, MatchedProperties}))
+			if (!RCEntity->IsBound())
+			{
+				if (UObject* Match = FDefaultRebindingPolicy::FindMatch({RCEntity, *ObjectsToConsider, MatchedProperties}))
+				{
+					MatchedProperties.FindOrAdd(RCEntity->GetUnderlyingEntityIdentifier()).Add(Match);
+					RCEntity->BindObject(Match);
+				}
+			}
+		}
+	}
+}
+
+void FRemoteControlPresetRebindingManager::RebindEntitiesForActorClass_Legacy(URemoteControlPreset* Preset, UClass* Class, const TArray<TSharedPtr<FRemoteControlEntity>>& UnboundEntities)
+{
+	TSet<FGuid> BoundProperties;
+	TSet<AActor*> BoundActors;
+
+	if (const TArray<UObject*>* ObjectsToConsider = Context.ObjectsGroupedByRelevantClass.Find(Class))
+	{
+		for (const TSharedPtr<FRemoteControlEntity>& RCEntity : UnboundEntities)
+		{
+			if (!BoundProperties.Contains(RCEntity->GetId()))
 			{
-				MatchedProperties.FindOrAdd(RCEntity->GetUnderlyingEntityIdentifier()).Add(Match);
-				RCEntity->BindObject(Match);
+				for (UObject* ObjToConsider : *ObjectsToConsider)
+				{
+					if (AActor* ActorToConsider = Cast<AActor>(ObjToConsider))
+					{
+						if (!BoundActors.Contains(ActorToConsider))
+						{
+							constexpr bool bUseRebindingContext = false;
+							BoundProperties.Append(RebindAllEntitiesUnderSameActor(Preset, RCEntity, ActorToConsider, bUseRebindingContext));
+							BoundActors.Add(ActorToConsider);
+							break;
+						}
+					}
+				}
 			}
 		}
 	}
 }
+
diff --git a/Engine/Plugins/VirtualProduction/RemoteControl/Source/RemoteControl/Private/RemoteControlPresetRebindingManager.h b/Engine/Plugins/VirtualProduction/RemoteControl/Source/RemoteControl/Private/RemoteControlPresetRebindingManager.h
--- a/Engine/Plugins/VirtualProduction/RemoteControl/Source/RemoteControl/Private/RemoteControlPresetRebindingManager.h
+++ b/Engine/Plugins/VirtualProduction/RemoteControl/Source/RemoteControl/Private/RemoteControlPresetRebindingManager.h
@@ -58,15 +58,20 @@ public:
 	/**
 	 * Attempt rebinding a Preset's entities by first trying to find objects with the same name
 	 * in the current map, then by searching for objects that support the underlying entity.
+	 * @Note: Set RemoteControl.UseLegacyRebinding to 1 in order to use the legacy algorthm.
 	 */
 	void Rebind(URemoteControlPreset* Preset);
 
 	/**
 	 * Given a RC Entity, rebind all entities with the same owner to a new actor
+	 * @return the unique identifiers of the rc entities that were rebound.
 	 */
-	void RebindAllEntitiesUnderSameActor(URemoteControlPreset* Preset, const TSharedPtr<FRemoteControlEntity>& Entity, AActor* NewActor);
+	TArray<FGuid> RebindAllEntitiesUnderSameActor(URemoteControlPreset* Preset, const TSharedPtr<FRemoteControlEntity>& Entity, AActor* NewActor, bool bUseRebindingContext);
 
 private:
+	void Rebind_Legacy(URemoteControlPreset* Preset);
+	void Rebind_NewAlgo(URemoteControlPreset* Preset);
+
 	/** Group RC Entities by the entity's supported owner class. */
 	TMap<UClass*, TArray<TSharedPtr<FRemoteControlEntity>>> GroupByEntitySupportedOwnerClass(TConstArrayView<TSharedPtr<FRemoteControlEntity>> Entities) const;
 
@@ -77,7 +82,12 @@ private:
 	void ReinitializeBindings(URemoteControlPreset* Preset) const;
 
 	/**  Attempt to rebind all unbound entities. */
-	void RebindEntitiesForClass(UClass* Class, const TArray<TSharedPtr<FRemoteControlEntity>>& UnboundEntities);
+	void RebindEntitiesForClass_Legacy(UClass* Class, const TArray<TSharedPtr<FRemoteControlEntity>>& UnboundEntities);
+
+	void RebindEntitiesForActorClass_Legacy(URemoteControlPreset* Preset, UClass* Class, const TArray<TSharedPtr<FRemoteControlEntity>>& UnboundEntities);
+
+	TArray<FGuid> RebindAllEntitiesUnderSameActor(URemoteControlPreset* Preset, URemoteControlBinding* InitialBinding, AActor* NewActor);
+	TArray<FGuid> RebindAllEntitiesUnderSameActor_Legacy(URemoteControlPreset* Preset, const TSharedPtr<FRemoteControlEntity>& Entity, AActor* NewActor);
 private:
 	struct FRebindingContext
 	{
diff --git a/Engine/Plugins/VirtualProduction/RemoteControl/Source/RemoteControl/Public/RemoteControlBinding.h b/Engine/Plugins/VirtualProduction/RemoteControl/Source/RemoteControl/Public/RemoteControlBinding.h
--- a/Engine/Plugins/VirtualProduction/RemoteControl/Source/RemoteControl/Public/RemoteControlBinding.h
+++ b/Engine/Plugins/VirtualProduction/RemoteControl/Source/RemoteControl/Public/RemoteControlBinding.h
@@ -3,6 +3,9 @@
 #pragma once
 
 #include "UObject/Object.h"
+
+#include "GameFramework/Actor.h"
+#include "Engine/Classes/Components/ActorComponent.h"
 #include "UObject/ObjectMacros.h"
 #include "UObject/SoftObjectPtr.h"
 #include "UObject/WeakObjectPtrTemplates.h"
@@ -79,6 +82,57 @@ private:
 	TSoftObjectPtr<UObject> BoundObject;
 };
 
+USTRUCT()
+struct FRemoteControlInitialBindingContext
+{
+	GENERATED_BODY()
+
+	/*
+	 * The class that's supported by this binding.
+	 */
+	UPROPERTY()
+	TSoftClassPtr<UObject> SupportedClass;
+
+	/**
+	 *  Name of the component if any.
+	 */
+	UPROPERTY()
+	FName ComponentName;
+
+	/**
+	 * The class of the actor that's targeted by this binding or the class of the actor that owns the component that is targeted.
+	 */
+	UPROPERTY()
+	TSoftClassPtr<AActor> OwnerActorClass;
+
+	/**
+	 * Name of the initial actor that was targetted by this binding.
+	 */
+	UPROPERTY()
+	FName OwnerActorName;
+
+	bool IsEmpty() const
+	{
+		return SupportedClass.GetUniqueID().ToString().IsEmpty()
+			&& ComponentName.IsNone() 
+			&& OwnerActorClass.GetUniqueID().ToString().IsEmpty()
+			&& OwnerActorName.IsNone();
+	}
+
+	friend bool operator==(const FRemoteControlInitialBindingContext& LHS, const FRemoteControlInitialBindingContext& RHS)
+	{
+		return LHS.SupportedClass == RHS.SupportedClass
+			&& LHS.ComponentName == RHS.ComponentName
+			&& LHS.OwnerActorClass == RHS.OwnerActorClass
+			&& LHS.OwnerActorName == RHS.OwnerActorName;
+	}
+
+	friend bool operator!=(const FRemoteControlInitialBindingContext& LHS, const FRemoteControlInitialBindingContext& RHS)
+	{
+		return !(LHS == RHS);
+	}
+};
+
 UCLASS(BlueprintType)
 class REMOTECONTROL_API URemoteControlLevelDependantBinding : public URemoteControlBinding
 {
@@ -100,17 +154,29 @@ public:
 	 */
 	void SetBoundObject(const TSoftObjectPtr<ULevel>& Level, const TSoftObjectPtr<UObject>& BoundObject);
 
+	/**
+	 * Set the bound object and reinitialize the context used for rebinding purposes.
+	 */
+	void SetBoundObject_OverrideContext(const TSoftObjectPtr<UObject>& InObject);
+
 	/**
 	 * Initialize this binding to be used with the new current level.
 	 * Copies the binding from the last successful resolve in case the level was duplicated.
 	 */
 	void InitializeForNewLevel();
 
+	/**
+	 * Get the suppported class of the bound object or of the owner of the bound object in the case of a bound component.
+	 */
+	UClass* GetSupportedOwnerClass() const;
+
 	/** Get the current world. */
 	static UWorld* GetCurrentWorld();
 private:
 	TSoftObjectPtr<UObject> ResolveForCurrentWorld() const;
 
+	void InitializeBindingContext(UObject* InObject) const;
+
 private:
 	/**
 	 *	The map bound objects with their level as key.
@@ -133,5 +199,9 @@ private:
 	UPROPERTY()
 	mutable TSoftObjectPtr<ULevel> LevelWithLastSuccessfulResolve;
 
+	UPROPERTY()
+	mutable FRemoteControlInitialBindingContext BindingContext;
+
 	friend class FRemoteControlPresetRebindingManager;
+	friend class URemoteControlPreset;
 };
diff --git a/Engine/Plugins/VirtualProduction/RemoteControl/Source/RemoteControl/Public/RemoteControlEntity.h b/Engine/Plugins/VirtualProduction/RemoteControl/Source/RemoteControl/Public/RemoteControlEntity.h
--- a/Engine/Plugins/VirtualProduction/RemoteControl/Source/RemoteControl/Public/RemoteControlEntity.h
+++ b/Engine/Plugins/VirtualProduction/RemoteControl/Source/RemoteControl/Public/RemoteControlEntity.h
@@ -52,6 +52,11 @@ struct REMOTECONTROL_API FRemoteControlEntity
 	 */
 	TArray<UObject*> GetBoundObjects() const;
 
+	/** 
+	 * Resolve the first binding for this entity. 
+	 */
+	UObject* GetBoundObject() const;
+
 	/**
 	 * Get the unresolved bindings of this exposed entity.
 	 */
diff --git a/Engine/Plugins/VirtualProduction/RemoteControl/Source/RemoteControl/Public/RemoteControlPreset.h b/Engine/Plugins/VirtualProduction/RemoteControl/Source/RemoteControl/Public/RemoteControlPreset.h
--- a/Engine/Plugins/VirtualProduction/RemoteControl/Source/RemoteControl/Public/RemoteControlPreset.h
+++ b/Engine/Plugins/VirtualProduction/RemoteControl/Source/RemoteControl/Public/RemoteControlPreset.h
@@ -739,7 +739,7 @@ public:
 	/**
 	 * Given a RC Entity, rebind all entities with the same owner to a new actor.
 	 */
-	void RebindAllEntitiesUnderSameActor(const FGuid& EntityId, AActor* NewActor);
+	void RebindAllEntitiesUnderSameActor(const FGuid& EntityId, AActor* NewActor, bool bUseRebindingContext = true);
 
 	DECLARE_MULTICAST_DELEGATE_TwoParams(FOnPresetEntityEvent, URemoteControlPreset* /*Preset*/, const FGuid& /*EntityId*/);
 	FOnPresetEntityEvent& OnEntityExposed() { return OnEntityExposedDelegate; }
@@ -881,6 +881,9 @@ public:
 	URemoteControlBinding* FindOrAddBinding(const TSoftObjectPtr<UObject>& Object);
 private:
 
+	/** Find a binding that has the same boundobjectmap but that currently points to the object passed as argument. */
+	URemoteControlBinding* FindMatchingBinding(const URemoteControlBinding* InBinding, UObject* InObject);
+
 	/** Handler called upon an entity being modified. */
 	void OnEntityModified(const FGuid& EntityId);
 
@@ -905,6 +908,9 @@ private:
 	/** Create property watchers for exposed properties that need them. */
 	void CreatePropertyWatchers();
 
+	/** Remove bindings that do not have properties pointing to them. */
+	void RemoveUnusedBindings();
+
 	/** Call post load function for exposed properties. */
 	void PostLoadProperties();
 	
@@ -1017,4 +1023,5 @@ PRAGMA_DISABLE_DEPRECATION_WARNINGS
 PRAGMA_ENABLE_DEPRECATION_WARNINGS
 	friend FRemoteControlPresetLayout;
 	friend FRemoteControlEntity;
+	friend class FRemoteControlPresetRebindingManager;
 };
diff --git a/Engine/Plugins/VirtualProduction/RemoteControl/Source/RemoteControlCommon/Public/RemoteControlSettings.h b/Engine/Plugins/VirtualProduction/RemoteControl/Source/RemoteControlCommon/Public/RemoteControlSettings.h
--- a/Engine/Plugins/VirtualProduction/RemoteControl/Source/RemoteControlCommon/Public/RemoteControlSettings.h
+++ b/Engine/Plugins/VirtualProduction/RemoteControl/Source/RemoteControlCommon/Public/RemoteControlSettings.h
@@ -83,4 +83,7 @@ public:
 	/** The split widget control ratio between entity tree and details/protocol binding list. */
 	UPROPERTY(config)
 	float TreeBindingSplitRatio = 0.7;
+
+	UPROPERTY(config)
+	bool bUseRebindingContext = true;
 };
diff --git a/Engine/Plugins/VirtualProduction/RemoteControl/Source/RemoteControlUI/Private/UI/Customizations/RemoteControlEntityCustomization.cpp b/Engine/Plugins/VirtualProduction/RemoteControl/Source/RemoteControlUI/Private/UI/Customizations/RemoteControlEntityCustomization.cpp
--- a/Engine/Plugins/VirtualProduction/RemoteControl/Source/RemoteControlUI/Private/UI/Customizations/RemoteControlEntityCustomization.cpp
+++ b/Engine/Plugins/VirtualProduction/RemoteControl/Source/RemoteControlUI/Private/UI/Customizations/RemoteControlEntityCustomization.cpp
@@ -123,6 +123,55 @@ void FRemoteControlEntityCustomization::CreateRangeWidget(URemoteControlPreset*
     ];
 }
 
+void FRemoteControlEntityCustomization::CreateBindingWidget(const FRemoteControlEntity* RCEntity, IDetailLayoutBuilder& LayoutBuilder, IDetailCategoryBuilder& CategoryBuilder) const
+{
+	const FSlateFontInfo FontInfo = LayoutBuilder.GetDetailFont();
+	FText ToolTipText = LOCTEXT("RCBindingToolTip", "The object this is currently bound to.");
+
+	FString PathName;
+
+	if (ensure(RCEntity))
+	{
+		if (UObject* Obj = RCEntity->GetBoundObject())
+		{
+			PathName = Obj->GetPathName();
+		}
+	}
+
+	FText PathNameText;
+
+	if (PathName.IsEmpty())
+	{
+		PathNameText = LOCTEXT("BindingPath", "This entity is not bound.");
+	}
+	else
+	{
+		PathNameText = FText::FromString(MoveTemp(PathName));
+	}
+
+	CategoryBuilder.AddCustomRow(LOCTEXT("BindingLabel", "Binding"))
+		.NameContent()
+		[
+			SNew(STextBlock)
+			.Text(LOCTEXT("BindingLabel", "Binding"))
+			.ToolTipText(ToolTipText)
+			.Font(FontInfo)
+		]
+		.ValueContent()
+		[
+			SNew(SHorizontalBox)
+			.ToolTipText(ToolTipText)
+			+ SHorizontalBox::Slot()
+			.FillWidth(1)
+			[
+				SNew(SEditableTextBox)
+				.IsReadOnly(true)
+				.Text(PathNameText)
+				.Font(FontInfo)
+			]
+		];
+}
+
 void FRemoteControlEntityCustomization::OnMetadataKeyCommitted(const FText& Text, ETextCommit::Type Type, FName MetadataKey)
 {
 	if (FRemoteControlEntity* Entity = GetEntityPtr())
@@ -190,6 +239,8 @@ void FRemoteControlEntityCustomization::GeneratePropertyRows(IDetailLayoutBuilde
 				AddPropertyRow(TEXT("FieldName"), FieldName.ToString(), /*bReadOnly=*/true);
 			}
 
+			CreateBindingWidget(Entity, DetailBuilder, CategoryBuilder);
+
 			if (URemoteControlPreset* Preset = Entity->GetOwner())
 			{
 				// Add Entity metadata rows.
diff --git a/Engine/Plugins/VirtualProduction/RemoteControl/Source/RemoteControlUI/Private/UI/Customizations/RemoteControlEntityCustomization.h b/Engine/Plugins/VirtualProduction/RemoteControl/Source/RemoteControlUI/Private/UI/Customizations/RemoteControlEntityCustomization.h
--- a/Engine/Plugins/VirtualProduction/RemoteControl/Source/RemoteControlUI/Private/UI/Customizations/RemoteControlEntityCustomization.h
+++ b/Engine/Plugins/VirtualProduction/RemoteControl/Source/RemoteControlUI/Private/UI/Customizations/RemoteControlEntityCustomization.h
@@ -29,6 +29,9 @@ private:
 	/** Create a widget that handles modifying the min and max properties. */
 	void CreateRangeWidget(URemoteControlPreset* Preset, const FGuid& DisplayedEntityId, IDetailLayoutBuilder& LayoutBuilder, IDetailCategoryBuilder& CategoryBuilder);
 
+	/** Create a widget that displays what's the entity currently bound to. */
+	void CreateBindingWidget(const FRemoteControlEntity* RCEntity, IDetailLayoutBuilder& LayoutBuilder, IDetailCategoryBuilder& CategoryBuilder) const;
+
 	/** Handler called upon modifying the metadata of an exposed entity. */
 	void OnMetadataKeyCommitted(const FText& Text, ETextCommit::Type Type, FName MetadataKey);
 
diff --git a/Engine/Plugins/VirtualProduction/RemoteControl/Source/RemoteControlUI/Private/UI/SRCPanelExposedEntity.cpp b/Engine/Plugins/VirtualProduction/RemoteControl/Source/RemoteControlUI/Private/UI/SRCPanelExposedEntity.cpp
--- a/Engine/Plugins/VirtualProduction/RemoteControl/Source/RemoteControlUI/Private/UI/SRCPanelExposedEntity.cpp
+++ b/Engine/Plugins/VirtualProduction/RemoteControl/Source/RemoteControlUI/Private/UI/SRCPanelExposedEntity.cpp
@@ -3,25 +3,30 @@
 #include "SRCPanelExposedEntity.h"
 
 #include "ActorTreeItem.h"
-#include "Components/ActorComponent.h"
+#include "Editor.h"
 #include "EditorFontGlyphs.h"
+#include "Engine/Selection.h"
+#include "Engine/Classes/Components/ActorComponent.h"
 #include "Framework/Application/SlateApplication.h"
 #include "Framework/MultiBox/MultiBoxBuilder.h"
 #include "GameFramework/Actor.h"
-#include "Modules/ModuleManager.h"
+#include "RemoteControlBinding.h"
 #include "RemoteControlEntity.h"
+#include "RemoteControlPreset.h"
 #include "RemoteControlField.h"
+#include "RemoteControlSettings.h"
 #include "RemoteControlPanelStyle.h"
-#include "RemoteControlPreset.h"
+#include "SRCPanelDragHandle.h"
 #include "SceneOutlinerFilters.h"
 #include "SceneOutlinerModule.h"
 #include "ScopedTransaction.h"
-#include "SRCPanelDragHandle.h"
 #include "Styling/SlateIconFinder.h"
+#include "Modules/ModuleManager.h"
 #include "Widgets/Input/SButton.h"
 #include "Widgets/Input/SComboButton.h"
-#include "Widgets/Layout/SBorder.h"
+#include "Widgets/Input/SCheckBox.h"
 #include "Widgets/Layout/SBox.h"
+#include "Widgets/Layout/SBorder.h"
 #include "Widgets/SBoxPanel.h"
 #include "Widgets/Text/SInlineEditableTextBlock.h"
 
@@ -51,6 +56,10 @@ TSharedPtr<FRemoteControlEntity> SRCPanelExposedEntity::GetEntity() const
 TSharedPtr<SWidget> SRCPanelExposedEntity::GetContextMenu()
 {
 	FMenuBuilder MenuBuilder(true, TSharedPtr<const FUICommandList>());
+
+	constexpr bool bNoIndent = true;
+	MenuBuilder.AddWidget(CreateUseContextCheckbox(), LOCTEXT("UseContextLabel", "Use Context"), bNoIndent);
+
 	MenuBuilder.AddSubMenu(
 		LOCTEXT("EntityRebindSubmenuLabel", "Rebind"),
 		LOCTEXT("EntityRebindSubmenuToolTip", "Pick an object to rebind this exposed entity."),
@@ -187,6 +196,7 @@ TSharedRef<SWidget> SRCPanelExposedEntity::CreateRebindMenuContent()
 	FSceneOutlinerModule& SceneOutlinerModule = FModuleManager::Get().LoadModuleChecked<FSceneOutlinerModule>("SceneOutliner");
 	FSceneOutlinerInitializationOptions Options;
 	Options.Filters = MakeShared<FSceneOutlinerFilters>();
+
 	Options.Filters->AddFilterPredicate<FActorTreeItem>(FActorTreeItem::FFilterPredicate::CreateRaw(this, &SRCPanelExposedEntity::IsActorSelectable));
 
 	return SNew(SBox)
@@ -241,7 +251,30 @@ bool SRCPanelExposedEntity::IsActorSelectable(const AActor* Actor) const
 {
 	if (TSharedPtr<FRemoteControlEntity> Entity = GetEntity())
 	{
-		return Entity->CanBindObject(Actor);
+		if (Entity->GetBindings().Num())
+		{
+			// Don't show what it's already bound to.
+			if (UObject* Component = Entity->GetBindings()[0]->Resolve())
+			{
+				if (Component == Actor || Component->GetTypedOuter<AActor>() == Actor)
+				{
+					return false;
+				}
+			}
+
+			if (ShouldUseRebindingContext())
+			{
+				if (URemoteControlLevelDependantBinding* Binding = Cast<URemoteControlLevelDependantBinding>(Entity->GetBindings()[0].Get()))
+				{
+					if (UClass* SupportedClass = Binding->GetSupportedOwnerClass())
+					{
+						return Actor->GetClass()->IsChildOf(SupportedClass);
+					}
+				}
+			}
+
+			return Entity->CanBindObject(Actor);
+		}
 	}
 	return false;
 }
@@ -320,13 +353,15 @@ void SRCPanelExposedEntity::OnActorSelectedForRebindAllProperties(AActor* InActo
 {
 	if (TSharedPtr<FRemoteControlEntity> Entity = GetEntity())
 	{
-		Preset->RebindAllEntitiesUnderSameActor(Entity->GetId(), InActor);
+		const bool bShouldUseRebindingContext = ShouldUseRebindingContext();
+
+		Preset->RebindAllEntitiesUnderSameActor(Entity->GetId(), InActor, bShouldUseRebindingContext);
+		SelectActor(InActor);
 	}
 
 	FSlateApplication::Get().DismissAllMenus();
 }
 
-
 void SRCPanelExposedEntity::HandleUnexposeEntity()
 {
 	if (URemoteControlPreset* RCPreset = Preset.Get())
@@ -337,5 +372,72 @@ void SRCPanelExposedEntity::HandleUnexposeEntity()
 	}
 }
 
+void SRCPanelExposedEntity::SelectActor(AActor* InActor) const
+{
+	if (GEditor)
+	{
+		// Don't change selection if the target's component is already selected
+		USelection* Selection = GEditor->GetSelectedComponents();
+
+		const bool bComponentSelected = Selection->Num() == 1
+			&& Selection->GetSelectedObject(0) != nullptr
+			&& Selection->GetSelectedObject(0)->GetTypedOuter<AActor>() == InActor;
+
+		if (!bComponentSelected)
+		{
+			constexpr bool bNoteSelectionChange = false;
+			constexpr bool bDeselectBSPSurfs = true;
+			constexpr bool WarnAboutManyActors = false;
+			GEditor->SelectNone(bNoteSelectionChange, bDeselectBSPSurfs, WarnAboutManyActors);
+
+			constexpr bool bInSelected = true;
+			constexpr bool bNotify = true;
+			constexpr bool bSelectEvenIfHidden = true;
+			GEditor->SelectActor(InActor, bInSelected, bNotify, bSelectEvenIfHidden);
+		}
+	}
+}
+
+TSharedRef<SWidget> SRCPanelExposedEntity::CreateUseContextCheckbox()
+{
+	return SNew(SCheckBox)
+		.ToolTipText(LOCTEXT("UseRebindingContextTooltip", "Unchecking this will allow you to rebind this property to any object regardless of the underlying supported class."))
+
+		// Bind the button's "on checked" event to our object's method for this
+		.OnCheckStateChanged_Raw(this, &SRCPanelExposedEntity::OnUseContextChanged)
+
+		// Bind the check box's "checked" state to our user interface action
+		.IsChecked_Raw(this, &SRCPanelExposedEntity::IsUseContextEnabled);
+}
+
+void SRCPanelExposedEntity::OnUseContextChanged(ECheckBoxState State)
+{
+	if (URemoteControlSettings* Settings = GetMutableDefault<URemoteControlSettings>())
+	{
+		if (State == ECheckBoxState::Unchecked)
+		{
+			Settings->bUseRebindingContext = false;
+		}
+		else if (State == ECheckBoxState::Checked)
+		{
+			Settings->bUseRebindingContext = true;
+		}
+	}
+}
+
+ECheckBoxState SRCPanelExposedEntity::IsUseContextEnabled() const
+{
+	return ShouldUseRebindingContext() ? ECheckBoxState::Checked : ECheckBoxState::Unchecked;
+}
+
+bool SRCPanelExposedEntity::ShouldUseRebindingContext() const
+{
+	if (const URemoteControlSettings* Settings = GetDefault<URemoteControlSettings>())
+	{
+		return Settings->bUseRebindingContext;
+	}
+
+	return false;
+}
 
 #undef LOCTEXT_NAMESPACE
diff --git a/Engine/Plugins/VirtualProduction/RemoteControl/Source/RemoteControlUI/Private/UI/SRCPanelExposedEntity.h b/Engine/Plugins/VirtualProduction/RemoteControl/Source/RemoteControlUI/Private/UI/SRCPanelExposedEntity.h
--- a/Engine/Plugins/VirtualProduction/RemoteControl/Source/RemoteControlUI/Private/UI/SRCPanelExposedEntity.h
+++ b/Engine/Plugins/VirtualProduction/RemoteControl/Source/RemoteControlUI/Private/UI/SRCPanelExposedEntity.h
@@ -66,6 +66,7 @@ private:
 
 	/** Handle selecting an actor for a rebind for all properties under an actor. */
 	void OnActorSelectedForRebindAllProperties(AActor* InActor) const;
+
 	/** Verifies that the entity's label doesn't already exist. */
 	bool OnVerifyItemLabelChanged(const FText& InLabel, FText& OutErrorMessage);
 	/** Handles committing a entity label. */
@@ -77,6 +78,14 @@ private:
 	/** Handle clicking on the unexpose button. */
 	void HandleUnexposeEntity();
 
+	void SelectActor(AActor* InActor) const;
+
+	//~ Use context widget functions.
+	TSharedRef<SWidget> CreateUseContextCheckbox();
+	void OnUseContextChanged(ECheckBoxState State);
+	ECheckBoxState IsUseContextEnabled() const;
+	bool ShouldUseRebindingContext() const;
+
 	/** The textbox for the row's name. */
 	TSharedPtr<SInlineEditableTextBlock> NameTextBox;
 };
diff --git a/Engine/Plugins/VirtualProduction/RemoteControl/Source/RemoteControlUI/Private/UI/SRemoteControlPanel.cpp b/Engine/Plugins/VirtualProduction/RemoteControl/Source/RemoteControlUI/Private/UI/SRemoteControlPanel.cpp
--- a/Engine/Plugins/VirtualProduction/RemoteControl/Source/RemoteControlUI/Private/UI/SRemoteControlPanel.cpp
+++ b/Engine/Plugins/VirtualProduction/RemoteControl/Source/RemoteControlUI/Private/UI/SRemoteControlPanel.cpp
@@ -184,6 +184,8 @@ void SRemoteControlPanel::Construct(const FArguments& InArgs, URemoteControlPres
 						.Text(LOCTEXT("RebindButtonText", "Rebind All"))
 					]
 				]
+
+				+ SHorizontalBox::Slot()
 				.VAlign(VAlign_Center)
 				.AutoWidth()
 				[
