- Changing default webapp port to avoid conflicts on mac

#b amir.keren

diff --git a/Engine/Plugins/VirtualProduction/RemoteControl/Source/RemoteControlCommon/Public/RemoteControlSettings.h b/Engine/Plugins/VirtualProduction/RemoteControl/Source/RemoteControlCommon/Public/RemoteControlSettings.h
--- a/Engine/Plugins/VirtualProduction/RemoteControl/Source/RemoteControlCommon/Public/RemoteControlSettings.h
+++ b/Engine/Plugins/VirtualProduction/RemoteControl/Source/RemoteControlCommon/Public/RemoteControlSettings.h
@@ -50,7 +50,7 @@ public:
 
 	/** The remote control web app http port. */
 	UPROPERTY(config, EditAnywhere, Category = "Remote Control Web Interface", DisplayName = "Remote Control Web Interface http Port")
-	uint32 RemoteControlWebInterfacePort = 7000;
+	uint32 RemoteControlWebInterfacePort = 30000;
 
 	/** Should force a build of the WebApp at startup. */
 	UPROPERTY(config, EditAnywhere, Category = "Remote Control Web Interface", DisplayName = "Force WebApp build at startup")
