Remote Control delta API work and associated tests

diff --git a/Engine/Plugins/VirtualProduction/RemoteControl/Source/RemoteControl/Private/RemoteControlInterceptionHelpers.h b/Engine/Plugins/VirtualProduction/RemoteControl/Source/RemoteControl/Private/RemoteControlInterceptionHelpers.h
--- a/Engine/Plugins/VirtualProduction/RemoteControl/Source/RemoteControl/Private/RemoteControlInterceptionHelpers.h
+++ b/Engine/Plugins/VirtualProduction/RemoteControl/Source/RemoteControl/Private/RemoteControlInterceptionHelpers.h
@@ -42,6 +42,30 @@ constexpr ERCPayloadType ToInternal(ERCIPayloadType Value)
 	return ERCPayloadType::Cbor;
 }
 
+// Internal helper for ERCIModifyOperation->ERCModifyOperation conversion
+constexpr ERCModifyOperation ToInternal(ERCIModifyOperation Value)
+{
+	switch (Value)
+	{
+	case ERCIModifyOperation::EQUAL:
+		return ERCModifyOperation::EQUAL;
+
+	case ERCIModifyOperation::ADD:
+		return ERCModifyOperation::ADD;
+
+	case ERCIModifyOperation::SUBTRACT:
+		return ERCModifyOperation::SUBTRACT;
+
+	case ERCIModifyOperation::MULTIPLY:
+		return ERCModifyOperation::MULTIPLY;
+
+	case ERCIModifyOperation::DIVIDE:
+		return ERCModifyOperation::DIVIDE;
+	}
+
+	return ERCModifyOperation::EQUAL;
+}
+
 // Internal helper for ERCAccess->ERCIAccess conversion
 constexpr ERCIAccess ToExternal(ERCAccess Value)
 {
@@ -77,3 +101,27 @@ constexpr ERCIPayloadType ToExternal(ERCPayloadType Value)
 
 	return ERCIPayloadType::Cbor;
 }
+
+// Internal helper for ERCModifyOperation->ERCIModifyOperation conversion
+constexpr ERCIModifyOperation ToExternal(ERCModifyOperation Value)
+{
+	switch (Value)
+	{
+	case ERCModifyOperation::EQUAL:
+		return ERCIModifyOperation::EQUAL;
+
+	case ERCModifyOperation::ADD:
+		return ERCIModifyOperation::ADD;
+
+	case ERCModifyOperation::SUBTRACT:
+		return ERCIModifyOperation::SUBTRACT;
+
+	case ERCModifyOperation::MULTIPLY:
+		return ERCIModifyOperation::MULTIPLY;
+
+	case ERCModifyOperation::DIVIDE:
+		return ERCIModifyOperation::DIVIDE;
+	}
+
+	return ERCIModifyOperation::EQUAL;
+}
diff --git a/Engine/Plugins/VirtualProduction/RemoteControl/Source/RemoteControl/Private/RemoteControlInterceptionProcessor.cpp b/Engine/Plugins/VirtualProduction/RemoteControl/Source/RemoteControl/Private/RemoteControlInterceptionProcessor.cpp
--- a/Engine/Plugins/VirtualProduction/RemoteControl/Source/RemoteControl/Private/RemoteControlInterceptionProcessor.cpp
+++ b/Engine/Plugins/VirtualProduction/RemoteControl/Source/RemoteControl/Private/RemoteControlInterceptionProcessor.cpp
@@ -42,7 +42,7 @@ void FRemoteControlInterceptionProcessor::SetObjectProperties(FRCIPropertiesMeta
 		check(StructDeserializerBackend.IsValid());
 
 		// Deserialize without replication
-		IRemoteControlModule::Get().SetObjectProperties(ObjectRef, *StructDeserializerBackend);
+		IRemoteControlModule::Get().SetObjectProperties(ObjectRef, *StructDeserializerBackend, ERCPayloadType::Json, TArray<uint8>(), ToInternal(PropsMetadata.Operation));
 	}
 }
 
diff --git a/Engine/Plugins/VirtualProduction/RemoteControl/Source/RemoteControl/Private/RemoteControlModule.cpp b/Engine/Plugins/VirtualProduction/RemoteControl/Source/RemoteControl/Private/RemoteControlModule.cpp
--- a/Engine/Plugins/VirtualProduction/RemoteControl/Source/RemoteControl/Private/RemoteControlModule.cpp
+++ b/Engine/Plugins/VirtualProduction/RemoteControl/Source/RemoteControl/Private/RemoteControlModule.cpp
@@ -14,6 +14,7 @@
 #include "RemoteControlInterceptionProcessor.h"
 #include "RemoteControlInstanceMaterial.h"
 #include "RemoteControlPreset.h"
+#include "Serialization/PropertyMapStructDeserializerBackendWrapper.h"
 #include "StructDeserializer.h"
 #include "StructSerializer.h"
 #include "AssetRegistry/AssetRegistryModule.h"
@@ -187,21 +188,81 @@ namespace RemoteControlUtil
 	{
 		return Access == ERCAccess::WRITE_ACCESS || Access == ERCAccess::WRITE_TRANSACTION_ACCESS;
 	}
+
+	/** Helper function to check if a value is 0 without always calling FMath::IsNearlyZero, since that results in an ambiguous call for non-float values. */
+	template <typename T>
+	bool IsNearlyZero(T Value)
+	{
+		return Value == 0;
+	}
+
+	template <>
+	bool IsNearlyZero(double Value)
+	{
+		return FMath::IsNearlyZero(Value);
+	}
+
+	/**
+	 * Apply a property modification based on the given type's overloaded operators.
+	 * 
+	 * @param Operation The operation to perform.
+	 * @param BasePropertyData A pointer to the original value before modification.
+	 * @param DeltaPropertyData A pointer to the value of the operand and which will store the result.
+	 * @param Getter Function which, given a pointer to property data, will return the property's value.
+	 * @param Setter Function which, given a pointer to property data and a value, will set the property's value to the new value.
+	 */
+	template <typename T>
+	bool ApplySimpleDeltaOperation(ERCModifyOperation Operation, const void* BasePropertyData, void* DeltaPropertyData, FProperty* Property, TFunction<T(const void*)> Getter, TFunction<void(void*, T)> Setter)
+	{
+		const T BaseValue = Getter(BasePropertyData);
+		const T DeltaValue = Getter(DeltaPropertyData);
+
+		switch (Operation)
+		{
+		case ERCModifyOperation::ADD:
+			Setter(DeltaPropertyData, BaseValue + DeltaValue);
+			break;
+
+		case ERCModifyOperation::SUBTRACT:
+			Setter(DeltaPropertyData, BaseValue - DeltaValue);
+			break;
+
+		case ERCModifyOperation::MULTIPLY:
+			Setter(DeltaPropertyData, BaseValue * DeltaValue);
+			break;
+
+		case ERCModifyOperation::DIVIDE:
+			if (IsNearlyZero<T>(DeltaValue))
+			{
+				return false;
+			}
+			Setter(DeltaPropertyData, BaseValue / DeltaValue);
+			break;
+
+		default:
+			// Unsupported operation
+			return false;
+		}
+
+		return true;
+	}
 }
 
 namespace RemoteControlSetterUtils
 {
 	struct FConvertToFunctionCallArgs
 	{
-		FConvertToFunctionCallArgs(const FRCObjectReference& InObjectReference, IStructDeserializerBackend& InReaderBackend, FRCCall& OutCall)
+		FConvertToFunctionCallArgs(const FRCObjectReference& InObjectReference, IStructDeserializerBackend& InReaderBackend, FRCCall& OutCall, const void* InValuePtrOverride = nullptr)
 			: ObjectReference(InObjectReference)
 			, ReaderBackend(InReaderBackend)
 			, Call(OutCall)
+			, ValuePtrOverride(InValuePtrOverride)
 		{ }
 
 		const FRCObjectReference& ObjectReference;
 		IStructDeserializerBackend& ReaderBackend;
 		FRCCall& Call;
+		const void* ValuePtrOverride;
 	};
 
 	void CreateRCCall(FConvertToFunctionCallArgs& InOutArgs, UFunction* InFunction, FStructOnScope&& InFunctionArguments, FRCInterceptionPayload& OutPayload)
@@ -238,14 +299,9 @@ namespace RemoteControlSetterUtils
 		bool bSuccess = false;
 		if (FProperty* SetterArgument = RemoteControlPropertyUtilities::FindSetterArgument(InSetterFunction, InOutArgs.ObjectReference.Property.Get()))
 		{
-			FStructOnScope ArgsOnScope{InSetterFunction};
-
-			// First put the complete property value from the object in the struct on scope
-			// in case the user only a part of the incoming structure (ie. Providing only { "x": 2 } in the case of a vector.
-			const uint8* ContainerAddress = InOutArgs.ObjectReference.Property->ContainerPtrToValuePtr<uint8>(InOutArgs.ObjectReference.ContainerAdress);
-			InOutArgs.ObjectReference.Property->CopyCompleteValue(ArgsOnScope.GetStructMemory(), ContainerAddress);
+			FStructOnScope ArgsOnScope{ InSetterFunction };
 
-			// Temporarily rename the setter argument in order to deserialize the incoming property modification on top of it
+			// Temporarily rename the setter argument in order to copy/deserialize the incoming property modification on top of it
 			// regardless of the argument name.
 			FName OldSetterArgumentName = SetterArgument->GetFName();
 			SetterArgument->Rename(InOutArgs.ObjectReference.Property->GetFName());
@@ -255,13 +311,35 @@ namespace RemoteControlSetterUtils
 					SetterArgument->Rename(OldSetterArgumentName);
 				};
 
-				// Then deserialize the input value on top of it and reset the setter property name.
-				bSuccess = FStructDeserializer::Deserialize((void*)ArgsOnScope.GetStructMemory(), *const_cast<UStruct*>(ArgsOnScope.GetStruct()), InOutArgs.ReaderBackend, FStructDeserializerPolicies());
-			}
+				if (InOutArgs.ValuePtrOverride)
+				{
+					// We already have a pointer directly to the data we want, so copy the data into the function arguments
+					const UStruct* ArgsStruct = ArgsOnScope.GetStruct();
 
-			if (bSuccess)
-			{
-				OptionalArgsOnScope = MoveTemp(ArgsOnScope);
+					check(ArgsStruct);
+
+					uint8* SetterArgData = SetterArgument->ContainerPtrToValuePtr<uint8>(ArgsOnScope.GetStructMemory());
+					InOutArgs.ObjectReference.Property->CopyCompleteValue(SetterArgData, InOutArgs.ValuePtrOverride);
+
+					bSuccess = true;
+				}
+				else
+				{
+					// Data needs to be deserialized from the passed reader backend
+
+					// First put the complete property value from the object in the struct on scope
+					// in case the user only a part of the incoming structure (ie. Providing only { "x": 2 } in the case of a vector.
+					const uint8* ContainerAddress = InOutArgs.ObjectReference.Property->ContainerPtrToValuePtr<uint8>(InOutArgs.ObjectReference.ContainerAdress);
+					InOutArgs.ObjectReference.Property->CopyCompleteValue(ArgsOnScope.GetStructMemory(), ContainerAddress);
+
+					// Deserialize on top of the setter argument
+					bSuccess = FStructDeserializer::Deserialize((void*)ArgsOnScope.GetStructMemory(), *const_cast<UStruct*>(ArgsOnScope.GetStruct()), InOutArgs.ReaderBackend, FStructDeserializerPolicies());
+				}
+
+				if (bSuccess)
+				{
+					OptionalArgsOnScope = MoveTemp(ArgsOnScope);
+				}
 			}
 		}
 
@@ -704,10 +782,11 @@ bool FRemoteControlModule::GetObjectProperties(const FRCObjectReference& ObjectA
 	return false;
 }
 
-bool FRemoteControlModule::SetObjectProperties(const FRCObjectReference& ObjectAccess, IStructDeserializerBackend& Backend, ERCPayloadType InPayloadType, const TArray<uint8>& InPayload)
+bool FRemoteControlModule::SetObjectProperties(const FRCObjectReference& ObjectAccess, IStructDeserializerBackend& Backend, ERCPayloadType InPayloadType, const TArray<uint8>& InPayload, ERCModifyOperation Operation)
 {
 	TRACE_CPUPROFILER_EVENT_SCOPE(FRemoteControlModule::SetObjectProperties);
 	UE_LOG(LogRemoteControl, VeryVerbose, TEXT("Set Object Properties"));
+
 	// Check the replication path before applying property values
 	if (InPayload.Num() != 0 && ObjectAccess.Object.IsValid())
 	{
@@ -718,6 +797,19 @@ bool FRemoteControlModule::SetObjectProperties(const FRCObjectReference& ObjectA
 			FRCInterceptionPayload InterceptionPayload;
 			constexpr bool bCreateInterceptionPayload = true;
 			RemoteControlSetterUtils::FConvertToFunctionCallArgs Args(ObjectAccess, Backend, Call);
+
+			// If this is a delta operation, deserialize the value into and apply the operation on a temporary buffer
+			TArray<uint8> DeltaData;
+			if (Operation != ERCModifyOperation::EQUAL)
+			{
+				if (!DeserializeDeltaModificationData(ObjectAccess, Backend, Operation, DeltaData))
+				{
+					return false;
+				}
+
+				Args.ValuePtrOverride = ObjectAccess.Property->ContainerPtrToValuePtr<void>(DeltaData.GetData());
+			}
+
 			if (RemoteControlSetterUtils::ConvertModificationToFunctionCall(Args, InterceptionPayload))
 			{
 				const bool bResult = InvokeCall(Call, InterceptionPayload.Type, InterceptionPayload.Payload);
@@ -734,7 +826,7 @@ bool FRemoteControlModule::SetObjectProperties(const FRCObjectReference& ObjectA
 
 		// Build interception command
 		FString PropertyPathString = TFieldPath<FProperty>(ObjectAccess.Property.Get()).ToString();
-		FRCIPropertiesMetadata PropsMetadata(ObjectAccess.Object->GetPathName(), PropertyPathString, ObjectAccess.PropertyPathInfo.ToString(), ToExternal(ObjectAccess.Access), ToExternal(InPayloadType), InPayload);
+		FRCIPropertiesMetadata PropsMetadata(ObjectAccess.Object->GetPathName(), PropertyPathString, ObjectAccess.PropertyPathInfo.ToString(), ToExternal(ObjectAccess.Access), ToExternal(InPayloadType), ToExternal(Operation), InPayload);
 
 		// Initialization
 		IModularFeatures& ModularFeatures = IModularFeatures::Get();
@@ -766,6 +858,19 @@ bool FRemoteControlModule::SetObjectProperties(const FRCObjectReference& ObjectA
 	{
 		FRCCall Call;
 		RemoteControlSetterUtils::FConvertToFunctionCallArgs Args(ObjectAccess, Backend, Call);
+
+		// If this is a delta operation, deserialize the value into and apply the operation on a temporary buffer
+		TArray<uint8> DeltaData;
+		if (Operation != ERCModifyOperation::EQUAL)
+		{
+			if (!DeserializeDeltaModificationData(ObjectAccess, Backend, Operation, DeltaData))
+			{
+				return false;
+			}
+
+			Args.ValuePtrOverride = ObjectAccess.Property->ContainerPtrToValuePtr<void>(DeltaData.GetData());
+		}
+
 		if (RemoteControlSetterUtils::ConvertModificationToFunctionCall(Args))
 		{
 			const bool bResult = InvokeCall(Call);
@@ -857,17 +962,34 @@ bool FRemoteControlModule::SetObjectProperties(const FRCObjectReference& ObjectA
 			};
 		}
 
-		//Serialize the element if we're looking for a member or serialize the full object if not
 		bool bSuccess = false;
-		if (MutableObjectReference.PropertyPathInfo.IsResolved())
+
+		if (Operation == ERCModifyOperation::EQUAL)
 		{
-			const FRCFieldPathSegment& LastSegment = MutableObjectReference.PropertyPathInfo.GetFieldSegment(MutableObjectReference.PropertyPathInfo.GetSegmentCount() - 1);
-			int32 Index = LastSegment.ArrayIndex != INDEX_NONE ? LastSegment.ArrayIndex : LastSegment.ResolvedData.MapIndex;
-			bSuccess = FStructDeserializer::DeserializeElement(MutableObjectReference.ContainerAdress, *LastSegment.ResolvedData.Struct, Index, Backend, Policies);
+			//Serialize the element if we're looking for a member or serialize the full object if not
+			if (MutableObjectReference.PropertyPathInfo.IsResolved())
+			{
+				const FRCFieldPathSegment& LastSegment = MutableObjectReference.PropertyPathInfo.GetFieldSegment(MutableObjectReference.PropertyPathInfo.GetSegmentCount() - 1);
+				int32 Index = LastSegment.ArrayIndex != INDEX_NONE ? LastSegment.ArrayIndex : LastSegment.ResolvedData.MapIndex;
+				bSuccess = FStructDeserializer::DeserializeElement(MutableObjectReference.ContainerAdress, *LastSegment.ResolvedData.Struct, Index, Backend, Policies);
+			}
+			else
+			{
+				bSuccess = FStructDeserializer::Deserialize(MutableObjectReference.ContainerAdress, *ContainerType, Backend, Policies);
+			}
 		}
 		else
 		{
-			bSuccess = FStructDeserializer::Deserialize(MutableObjectReference.ContainerAdress, *ContainerType, Backend, Policies);
+			// This is a delta operation, so deserialize the value into and apply the operation on a temporary buffer
+			TArray<uint8> DeltaData;
+			if (!DeserializeDeltaModificationData(MutableObjectReference, Backend, Operation, DeltaData))
+			{
+				return false;
+			}
+
+			// Copy data to the actual object
+			MutableObjectReference.Property->CopyCompleteValue_InContainer(MutableObjectReference.ContainerAdress, DeltaData.GetData());
+			bSuccess = true;
 		}
 
 #if WITH_EDITOR
@@ -893,9 +1015,9 @@ bool FRemoteControlModule::SetObjectProperties(const FRCObjectReference& ObjectA
 						{
 							World->Scene->UpdateLightColorAndBrightness(Cast<ULightComponent>(MutableObjectReference.Object.Get()));
 						}
+					}
 				}
 			}
-			}
 			else
 			{
 				OngoingModification = MutableObjectReference;
@@ -917,6 +1039,8 @@ bool FRemoteControlModule::SetObjectProperties(const FRCObjectReference& ObjectA
 		{
 			EntityFactoryPair.Value->PostSetObjectProperties(ObjectAccess.Object.Get(), bSuccess);
 		}
+
+		return bSuccess;
 	}
 	return false;
 }
@@ -1244,6 +1368,82 @@ bool FRemoteControlModule::PropertyModificationShouldUseSetter(UObject* Object,
 	return !!RemoteControlPropertyUtilities::FindSetterFunction(Property, Object->GetClass());
 }
 
+bool FRemoteControlModule::DeserializeDeltaModificationData(const FRCObjectReference& ObjectAccess, IStructDeserializerBackend& Backend, ERCModifyOperation Operation, TArray<uint8>& OutData)
+{
+	// Allocate data to deserialize the request data into
+	const int32 StructureSize = ObjectAccess.ContainerType->GetStructureSize();
+	OutData.SetNumUninitialized(StructureSize);
+	void* OutContainerAddress = OutData.GetData();
+
+	// Copy existing property data into the delta container so unchanged values remain the same
+	ObjectAccess.Property->CopyCompleteValue_InContainer(OutContainerAddress, ObjectAccess.ContainerAdress);
+
+	// Deserialize the data on top of what we just copied so that modified properties contain delta values
+	FStructDeserializerPolicies Policies;
+	Policies.PropertyFilter = [&ObjectAccess](const FProperty* CurrentProp, const FProperty* ParentProp)
+	{
+		return CurrentProp == ObjectAccess.Property || ParentProp != nullptr;
+	};
+
+	// Wrap the backend so we can track which properties were actually changed by deserialization
+	FPropertyMapStructDeserializerBackendWrapper BackendWrapper(Backend);
+
+	bool bSuccess;
+	if (ObjectAccess.PropertyPathInfo.IsResolved())
+	{
+		const FRCFieldPathSegment& LastSegment = ObjectAccess.PropertyPathInfo.GetFieldSegment(ObjectAccess.PropertyPathInfo.GetSegmentCount() - 1);
+		int32 Index = LastSegment.ArrayIndex != INDEX_NONE ? LastSegment.ArrayIndex : LastSegment.ResolvedData.MapIndex;
+		bSuccess = FStructDeserializer::DeserializeElement(OutContainerAddress, *LastSegment.ResolvedData.Struct, Index, BackendWrapper, Policies);
+	}
+	else
+	{
+		bSuccess = FStructDeserializer::Deserialize(OutContainerAddress, *ObjectAccess.ContainerType, BackendWrapper, Policies);
+	}
+
+	if (!bSuccess)
+	{
+		return false;
+	}
+
+	// Apply delta operation to each property that was changed (where possible)
+	for (const FPropertyMapStructDeserializerBackendWrapper::FReadPropertyData& ReadProperty : BackendWrapper.GetReadProperties())
+	{
+		// Pointer to the property in OutData
+		void* OutPropertyValue = ReadProperty.Data;
+
+		// Offset from start of OutData struct to the property that was changed
+		ptrdiff_t Offset = ((const uint8*)OutPropertyValue) - ((const uint8*)OutContainerAddress);
+
+		// Pointer to the equivalent property in the original object
+		const void* BasePropertyValue = ((uint8*)ObjectAccess.Object.Get()) + Offset;
+
+		if (FNumericProperty* NumericProperty = CastField<FNumericProperty>(ReadProperty.Property))
+		{
+			// Float property
+			if (NumericProperty->IsFloatingPoint())
+			{
+				bSuccess &= RemoteControlUtil::ApplySimpleDeltaOperation<double>(Operation, BasePropertyValue, OutPropertyValue, ReadProperty.Property,
+					[&NumericProperty](const void* Data) { return NumericProperty->GetFloatingPointPropertyValue(Data); },
+					[&NumericProperty](void* Data, double Value) { NumericProperty->SetFloatingPointPropertyValue(Data, Value); });
+			}
+			// Integer property
+			else if (NumericProperty->IsInteger() && !NumericProperty->IsEnum())
+			{
+				bSuccess &= RemoteControlUtil::ApplySimpleDeltaOperation<int64>(Operation, BasePropertyValue, OutPropertyValue, ReadProperty.Property,
+					[&NumericProperty](const void* Data) { return NumericProperty->GetSignedIntPropertyValue(Data); },
+					[&NumericProperty](void* Data, int64 Value) { NumericProperty->SetIntPropertyValue(Data, Value); });
+			}
+		}
+
+		if (!bSuccess)
+		{
+			break;
+		}
+	}
+
+	return bSuccess;
+}
+
 #if WITH_EDITOR
 void FRemoteControlModule::TestOrFinalizeOngoingChange(bool bForceEndChange)
 {
diff --git a/Engine/Plugins/VirtualProduction/RemoteControl/Source/RemoteControl/Private/RemoteControlModule.h b/Engine/Plugins/VirtualProduction/RemoteControl/Source/RemoteControl/Private/RemoteControlModule.h
--- a/Engine/Plugins/VirtualProduction/RemoteControl/Source/RemoteControl/Private/RemoteControlModule.h
+++ b/Engine/Plugins/VirtualProduction/RemoteControl/Source/RemoteControl/Private/RemoteControlModule.h
@@ -31,7 +31,7 @@ public:
 	virtual bool ResolveObject(ERCAccess AccessType, const FString& ObjectPath, const FString& PropertyName, FRCObjectReference& OutObjectRef, FString* OutErrorText = nullptr) override;
 	virtual bool ResolveObjectProperty(ERCAccess AccessType, UObject* Object, FRCFieldPathInfo PropertyPath, FRCObjectReference& OutObjectRef, FString* OutErrorText = nullptr) override;
 	virtual bool GetObjectProperties(const FRCObjectReference& ObjectAccess, IStructSerializerBackend& Backend) override;
-	virtual bool SetObjectProperties(const FRCObjectReference& ObjectAccess, IStructDeserializerBackend& Backend, ERCPayloadType InPayloadType, const TArray<uint8>& InPayload) override;
+	virtual bool SetObjectProperties(const FRCObjectReference& ObjectAccess, IStructDeserializerBackend& Backend, ERCPayloadType InPayloadType, const TArray<uint8>& InPayload, ERCModifyOperation Operation) override;
 	virtual bool ResetObjectProperties(const FRCObjectReference& ObjectAccess, const bool bAllowIntercept) override;
 	virtual TOptional<FExposedFunction> ResolvePresetFunction(const FResolvePresetFieldArgs& Args) const override;
 	virtual TOptional<FExposedProperty> ResolvePresetProperty(const FResolvePresetFieldArgs& Args) const override;
@@ -60,6 +60,17 @@ private:
 	/** Determines if a property modification should use a setter or default to deserializing directly onto an object. */
 	static bool PropertyModificationShouldUseSetter(UObject* Object, FProperty* Property);
 
+	/**
+	 * Deserialize data for a non-EQUAL modification request and apply the operation to the resulting data.
+	 * 
+	 * @param ObjectAccess Data about the object/property for which modification was requested.
+	 * @param Backend Deserialization backend for the modification request.
+	 * @param Operation Type of operation to apply to the value.
+	 * @param OutData Buffer used to deserialize data from the backend and contain the result value.
+	 * @return True if the data was successfully deserialized and modified.
+	 */
+	static bool DeserializeDeltaModificationData(const FRCObjectReference& ObjectAccess, IStructDeserializerBackend& Backend, ERCModifyOperation Operation, TArray<uint8>& OutData);
+
 #if WITH_EDITOR
 	/** Finalize an ongoing change, triggering post edit change on the tracked object. */
 	void TestOrFinalizeOngoingChange(bool bForceEndChange = false);
diff --git a/Engine/Plugins/VirtualProduction/RemoteControl/Source/RemoteControl/Private/Serialization/PropertyMapStructDeserializerBackendWrapper.cpp b/Engine/Plugins/VirtualProduction/RemoteControl/Source/RemoteControl/Private/Serialization/PropertyMapStructDeserializerBackendWrapper.cpp
new file mode 100644
--- /dev/null
+++ b/Engine/Plugins/VirtualProduction/RemoteControl/Source/RemoteControl/Private/Serialization/PropertyMapStructDeserializerBackendWrapper.cpp
@@ -0,0 +1,20 @@
+// Copyright Epic Games, Inc. All Rights Reserved.
+
+#include "PropertyMapStructDeserializerBackendWrapper.h"
+#include "UObject/UnrealType.h"
+
+bool FPropertyMapStructDeserializerBackendWrapper::ReadProperty(FProperty* Property, FProperty* Outer, void* Data, int32 ArrayIndex)
+{
+	if (InnerBackend.ReadProperty(Property, Outer, Data, ArrayIndex))
+	{
+		FReadPropertyData ReadPropertyData;
+		ReadPropertyData.Property = Property;
+		ReadPropertyData.Data = Property->ContainerPtrToValuePtr<void>(Data, ArrayIndex);
+
+		ReadProperties.Add(ReadPropertyData);
+		
+		return true;
+	}
+
+	return false;
+}
\ No newline at end of file
diff --git a/Engine/Plugins/VirtualProduction/RemoteControl/Source/RemoteControl/Private/Serialization/PropertyMapStructDeserializerBackendWrapper.h b/Engine/Plugins/VirtualProduction/RemoteControl/Source/RemoteControl/Private/Serialization/PropertyMapStructDeserializerBackendWrapper.h
new file mode 100644
--- /dev/null
+++ b/Engine/Plugins/VirtualProduction/RemoteControl/Source/RemoteControl/Private/Serialization/PropertyMapStructDeserializerBackendWrapper.h
@@ -0,0 +1,47 @@
+// Copyright Epic Games, Inc. All Rights Reserved.
+
+#pragma once
+
+#include "IStructDeserializerBackend.h"
+
+/**
+ * Wraps an existing IStructDeserializerBackend and passes calls through to it, but tracks which properties
+ * were successfully read in the process of deserialization.
+ */
+class FPropertyMapStructDeserializerBackendWrapper
+	: public IStructDeserializerBackend
+{
+public:
+	struct FReadPropertyData
+	{
+		FProperty* Property;
+		void* Data;
+	};
+
+	/**
+	 * Creates and initializes a new instance.
+	 *
+	 * @param InnerBackend The deserializer to wrap. This will actually handle deserializing the struct.
+	 */
+	FPropertyMapStructDeserializerBackendWrapper(IStructDeserializerBackend& InnerBackend)
+		: InnerBackend(InnerBackend)
+	{ }
+
+	// IStructDeserializerBackend interface
+
+	virtual const FString& GetCurrentPropertyName() const override { return InnerBackend.GetCurrentPropertyName(); }
+	virtual FString GetDebugString() const override { return InnerBackend.GetDebugString(); }
+	virtual const FString& GetLastErrorMessage() const override { return InnerBackend.GetLastErrorMessage(); }
+	virtual bool GetNextToken(EStructDeserializerBackendTokens& OutToken) override { return InnerBackend.GetNextToken(OutToken); }
+	virtual void SkipArray() override { InnerBackend.SkipArray(); }
+	virtual void SkipStructure() override { InnerBackend.SkipStructure(); }
+	virtual bool ReadProperty(FProperty* Property, FProperty* Outer, void* Data, int32 ArrayIndex) override;
+
+	/** Get the array of properties that have been read so far during deserialization. */
+	const TArray<FReadPropertyData>& GetReadProperties() const { return ReadProperties; }
+
+
+private:
+	IStructDeserializerBackend& InnerBackend;
+	TArray<FReadPropertyData> ReadProperties;
+};
diff --git a/Engine/Plugins/VirtualProduction/RemoteControl/Source/RemoteControl/Private/Tests/RemoteControlDeltaAPITestData.cpp b/Engine/Plugins/VirtualProduction/RemoteControl/Source/RemoteControl/Private/Tests/RemoteControlDeltaAPITestData.cpp
new file mode 100644
--- /dev/null
+++ b/Engine/Plugins/VirtualProduction/RemoteControl/Source/RemoteControl/Private/Tests/RemoteControlDeltaAPITestData.cpp
@@ -0,0 +1,10 @@
+// Copyright Epic Games, Inc. All Rights Reserved.
+
+#include "RemoteControlDeltaAPITestData.h"
+
+const FLinearColor FRemoteControlDeltaAPITestStruct::ColorValueDefault = FLinearColor(0.5f, 0.1f, 0.3f, 0.8f);
+
+const int32 URemoteControlDeltaAPITestObject::Int32ValueDefault = 523;
+const float URemoteControlDeltaAPITestObject::FloatValueDefault = 0.3f;
+const FVector URemoteControlDeltaAPITestObject::VectorValueDefault = FVector(112.f, -40.f, 37.f);
+const FLinearColor URemoteControlDeltaAPITestObject::ColorValueDefault = FLinearColor(0.2f, 0.6f, 0.9f, 0.5f);
\ No newline at end of file
diff --git a/Engine/Plugins/VirtualProduction/RemoteControl/Source/RemoteControl/Private/Tests/RemoteControlDeltaAPITestData.h b/Engine/Plugins/VirtualProduction/RemoteControl/Source/RemoteControl/Private/Tests/RemoteControlDeltaAPITestData.h
new file mode 100644
--- /dev/null
+++ b/Engine/Plugins/VirtualProduction/RemoteControl/Source/RemoteControl/Private/Tests/RemoteControlDeltaAPITestData.h
@@ -0,0 +1,60 @@
+// Copyright Epic Games, Inc. All Rights Reserved.
+#pragma once
+
+#include "CoreMinimal.h"
+
+#include "RemoteControlDeltaAPITestData.generated.h"
+
+USTRUCT()
+struct FRemoteControlDeltaAPITestStruct
+{
+	GENERATED_BODY()
+
+public:
+	static const FLinearColor ColorValueDefault;
+
+	UPROPERTY(EditAnywhere, Category = "RC")
+	FLinearColor ColorValue = ColorValueDefault;
+};
+
+UCLASS()
+class URemoteControlDeltaAPITestObject : public UObject
+{
+	GENERATED_BODY()
+
+public:
+	static const int32 Int32ValueDefault;
+	static const float FloatValueDefault;
+	static const FVector VectorValueDefault;
+	static const FLinearColor ColorValueDefault;
+	
+	static const FName GetInt32WithSetterValuePropertyName()
+	{
+		return GET_MEMBER_NAME_CHECKED(URemoteControlDeltaAPITestObject, Int32WithSetterValue);
+	}
+
+	UFUNCTION(BlueprintGetter)
+	int32 GetInt32WithSetterValue() { return Int32WithSetterValue; }
+
+	UFUNCTION(BlueprintSetter)
+	void SetInt32WithSetterValue(const int32 NewValue) { Int32WithSetterValue = NewValue; }
+
+	UPROPERTY(EditAnywhere, Category = "RC")
+	int32 Int32Value = Int32ValueDefault;
+
+	UPROPERTY(EditAnywhere, Category = "RC")
+	float FloatValue = FloatValueDefault;
+
+	UPROPERTY(EditAnywhere, Category = "RC")
+	FVector VectorValue = VectorValueDefault;
+
+	UPROPERTY(EditAnywhere, Category = "RC")
+	FLinearColor ColorValue = ColorValueDefault;
+
+	UPROPERTY(EditAnywhere, Category = "RC")
+	FRemoteControlDeltaAPITestStruct StructValue;
+
+private:
+	UPROPERTY(EditAnywhere, Category = "RC", BlueprintGetter = GetInt32WithSetterValue, BlueprintSetter = SetInt32WithSetterValue)
+	int32 Int32WithSetterValue;
+};
\ No newline at end of file
diff --git a/Engine/Plugins/VirtualProduction/RemoteControl/Source/RemoteControl/Private/Tests/RemoteControlDeltaAPITests.spec.cpp b/Engine/Plugins/VirtualProduction/RemoteControl/Source/RemoteControl/Private/Tests/RemoteControlDeltaAPITests.spec.cpp
new file mode 100644
--- /dev/null
+++ b/Engine/Plugins/VirtualProduction/RemoteControl/Source/RemoteControl/Private/Tests/RemoteControlDeltaAPITests.spec.cpp
@@ -0,0 +1,585 @@
+// Copyright Epic Games, Inc. All Rights Reserved.
+
+#include "Engine/Engine.h"
+#include "Misc/AutomationTest.h"
+
+#include "Backends/JsonStructDeserializerBackend.h"
+#include "Backends/JsonStructSerializerBackend.h"
+#include "Dom/JsonObject.h"
+#include "Policies/CondensedJsonPrintPolicy.h"
+#include "Serialization/JsonSerializer.h"
+#include "StructSerializer.h"
+
+#include "IRemoteControlModule.h"
+#include "RemoteControlDeltaAPITestData.h"
+
+BEGIN_DEFINE_SPEC(FRemoteControlDeltaAPISpec, "Plugin.RemoteControl.DeltaAPI", EAutomationTestFlags::EditorContext | EAutomationTestFlags::EngineFilter | EAutomationTestFlags::ApplicationContextMask)
+
+URemoteControlDeltaAPITestObject* TestObject;
+
+TArray<uint8> JsonBuffer;
+FMemoryReader Reader = FMemoryReader(JsonBuffer);
+FMemoryWriter Writer = FMemoryWriter(JsonBuffer);
+FJsonStructSerializerBackend SerializerBackend = FJsonStructSerializerBackend(Writer, EStructSerializerBackendFlags::Default);
+FJsonStructDeserializerBackend DeserializerBackend = FJsonStructDeserializerBackend(Reader);
+
+/** Fill out JsonBuffer with data from a string containing JSON data. */
+void SerializeBufferFromJsonString(const FString& JsonString);
+
+/** Fill out JsonBuffer with data for the given property using the data stored in TestObject. */
+void SerializeElementAndReset(FProperty* Property);
+
+/** Reset all values in TestObject to their defaults. */
+void ResetTestObject();
+
+/**
+ * Apply an operation to a property by calling IRemoteControlModule::SetObjectProperties.
+ * 
+ * @param Proprety The property to modify on TestObject.
+ * @param Operation The operation to apply to the property.
+ * @param bIncludeInterceptPayload If true, JsonBuffer will be passed as the InPayload argument of SetObjectProperties. Otherwise, an empty array will be passed.
+ */
+void ApplyOperation(FProperty* Property, ERCModifyOperation Operation, bool bIncludeInterceptPayload = true);
+
+END_DEFINE_SPEC(FRemoteControlDeltaAPISpec)
+
+void FRemoteControlDeltaAPISpec::SerializeBufferFromJsonString(const FString& JsonString)
+{
+	JsonBuffer.SetNumUninitialized(JsonString.GetAllocatedSize());
+	FMemory::Memcpy(JsonBuffer.GetData(), JsonString.GetCharArray().GetData(), JsonString.GetAllocatedSize());
+}
+
+void FRemoteControlDeltaAPISpec::SerializeElementAndReset(FProperty* Property)
+{
+	FStructSerializer::SerializeElement(TestObject, Property, INDEX_NONE, SerializerBackend, FStructSerializerPolicies());
+
+	// We use the test object to temporarily store the value we want to serialize, so reset it afterward since we'll also use it to test the change
+	ResetTestObject();
+}
+
+void FRemoteControlDeltaAPISpec::ResetTestObject()
+{
+	TestObject->Int32Value = URemoteControlDeltaAPITestObject::Int32ValueDefault;
+	TestObject->FloatValue = URemoteControlDeltaAPITestObject::FloatValueDefault;
+	TestObject->VectorValue = URemoteControlDeltaAPITestObject::VectorValueDefault;
+	TestObject->ColorValue = URemoteControlDeltaAPITestObject::ColorValueDefault;
+	TestObject->SetInt32WithSetterValue(URemoteControlDeltaAPITestObject::Int32ValueDefault);
+
+	TestObject->StructValue.ColorValue = FRemoteControlDeltaAPITestStruct::ColorValueDefault;
+}
+
+void FRemoteControlDeltaAPISpec::ApplyOperation(FProperty* Property, ERCModifyOperation Operation, bool bIncludeInterceptPayload)
+{
+	const FRCFieldPathInfo FieldPath(Property);
+	FRCObjectReference ObjectRef;
+	IRemoteControlModule::Get().ResolveObjectProperty(ERCAccess::WRITE_ACCESS, TestObject, FieldPath, ObjectRef);
+	IRemoteControlModule::Get().SetObjectProperties(ObjectRef, DeserializerBackend, ERCPayloadType::Json, bIncludeInterceptPayload ? JsonBuffer : TArray<uint8>(), Operation);
+}
+
+void FRemoteControlDeltaAPISpec::Define()
+{
+	BeforeEach([this]
+	{
+		 TestObject = NewObject<URemoteControlDeltaAPITestObject>();
+
+		 JsonBuffer.Empty();
+
+		 Writer.Reset();
+		 Writer.Seek(0);
+
+		 Reader.Reset();
+		 Reader.Seek(0);
+
+		 SerializerBackend = FJsonStructSerializerBackend(Writer, EStructSerializerBackendFlags::Default);
+		 DeserializerBackend = FJsonStructDeserializerBackend(Reader);
+	});
+
+	Describe("Add", [this]
+	{
+		It("should add integers", [this]
+		{
+			FProperty* ValueProperty = FindFProperty<FProperty>(URemoteControlDeltaAPITestObject::StaticClass(), GET_MEMBER_NAME_CHECKED(URemoteControlDeltaAPITestObject, Int32Value));
+
+			const int32 AddedValue = 42;
+			TestObject->Int32Value = AddedValue;
+
+			SerializeElementAndReset(ValueProperty);
+			ApplyOperation(ValueProperty, ERCModifyOperation::ADD);
+
+			TestEqual("Combined value", TestObject->Int32Value, URemoteControlDeltaAPITestObject::Int32ValueDefault + AddedValue);
+		});
+
+		It("should add floats", [this]
+		{
+			FProperty* ValueProperty = FindFProperty<FProperty>(URemoteControlDeltaAPITestObject::StaticClass(), GET_MEMBER_NAME_CHECKED(URemoteControlDeltaAPITestObject, FloatValue));
+
+			const float AddedValue = 0.6f;
+			TestObject->FloatValue = AddedValue;
+
+			SerializeElementAndReset(ValueProperty);
+			ApplyOperation(ValueProperty, ERCModifyOperation::ADD);
+
+			TestEqual("Combined value", TestObject->FloatValue, URemoteControlDeltaAPITestObject::FloatValueDefault + AddedValue);
+		});
+
+		It("should add vectors", [this]
+		{
+			FProperty* ValueProperty = FindFProperty<FProperty>(URemoteControlDeltaAPITestObject::StaticClass(), GET_MEMBER_NAME_CHECKED(URemoteControlDeltaAPITestObject, VectorValue));
+
+			const FVector AddedValue = FVector(-20.f, 0.7f, 101.f);
+			TestObject->VectorValue = AddedValue;
+
+			SerializeElementAndReset(ValueProperty);
+			ApplyOperation(ValueProperty, ERCModifyOperation::ADD);
+
+			TestEqual("Combined value", TestObject->VectorValue, URemoteControlDeltaAPITestObject::VectorValueDefault + AddedValue);
+		});
+
+		It("should add colors", [this]
+		{
+			FProperty* ValueProperty = FindFProperty<FProperty>(URemoteControlDeltaAPITestObject::StaticClass(), GET_MEMBER_NAME_CHECKED(URemoteControlDeltaAPITestObject, ColorValue));
+
+			const FLinearColor AddedValue = FLinearColor(0.4f, 0.1f, 0.2f, 0.8f);
+			TestObject->ColorValue = AddedValue;
+
+			SerializeElementAndReset(ValueProperty);
+			ApplyOperation(ValueProperty, ERCModifyOperation::ADD);
+
+			TestEqual("Combined value", TestObject->ColorValue, URemoteControlDeltaAPITestObject::ColorValueDefault + AddedValue);
+		});
+
+		It("should add a partial struct's elements only to the provided field", [this]
+		{
+			FProperty* ValueProperty = FindFProperty<FProperty>(URemoteControlDeltaAPITestObject::StaticClass(), GET_MEMBER_NAME_CHECKED(URemoteControlDeltaAPITestObject, VectorValue));
+
+			SerializeBufferFromJsonString("{ \"VectorValue\": { \"Y\": 12.3 } }");
+			ApplyOperation(ValueProperty, ERCModifyOperation::ADD);
+
+			TestEqual("X value", TestObject->VectorValue.X, URemoteControlDeltaAPITestObject::VectorValueDefault.X);
+			TestEqual("Combined Y value", TestObject->VectorValue.Y, URemoteControlDeltaAPITestObject::VectorValueDefault.Y + 12.3f);
+			TestEqual("Z value", TestObject->VectorValue.Z, URemoteControlDeltaAPITestObject::VectorValueDefault.Z);
+		});
+
+		It("should add property with a getter/setter", [this]
+		{
+			FProperty* ValueProperty = FindFProperty<FProperty>(URemoteControlDeltaAPITestObject::StaticClass(), URemoteControlDeltaAPITestObject::GetInt32WithSetterValuePropertyName());
+
+			const int32 AddedValue = 12;
+			TestObject->SetInt32WithSetterValue(AddedValue);
+
+			SerializeElementAndReset(ValueProperty);
+			ApplyOperation(ValueProperty, ERCModifyOperation::ADD);
+
+			TestEqual("Combined value", TestObject->GetInt32WithSetterValue(), URemoteControlDeltaAPITestObject::Int32ValueDefault + AddedValue);
+		});
+
+		It("should add property with a getter/setter and no intercept payload", [this]
+		{
+			FProperty* ValueProperty = FindFProperty<FProperty>(URemoteControlDeltaAPITestObject::StaticClass(), URemoteControlDeltaAPITestObject::GetInt32WithSetterValuePropertyName());
+
+			const int32 AddedValue = 12;
+			TestObject->SetInt32WithSetterValue(AddedValue);
+
+			SerializeElementAndReset(ValueProperty);
+			ApplyOperation(ValueProperty, ERCModifyOperation::ADD, false);
+
+			TestEqual("Combined value", TestObject->GetInt32WithSetterValue(), URemoteControlDeltaAPITestObject::Int32ValueDefault + AddedValue);
+		});
+
+		It("should add nested struct values", [this]
+		{
+			FProperty* ValueProperty = FindFProperty<FProperty>(URemoteControlDeltaAPITestObject::StaticClass(), GET_MEMBER_NAME_CHECKED(URemoteControlDeltaAPITestObject, StructValue));
+
+			const FLinearColor AddedValue = FLinearColor(0.2f, 0.5f, 0.1f, 0.3f);
+			TestObject->StructValue.ColorValue = AddedValue;
+
+			SerializeElementAndReset(ValueProperty);
+			ApplyOperation(ValueProperty, ERCModifyOperation::ADD);
+
+			TestEqual("Combined value", TestObject->StructValue.ColorValue, FRemoteControlDeltaAPITestStruct::ColorValueDefault + AddedValue);
+		});
+
+		It("should add a nested partial struct's elements only to the provided field", [this]
+		{
+			FProperty* ValueProperty = FindFProperty<FProperty>(URemoteControlDeltaAPITestObject::StaticClass(), GET_MEMBER_NAME_CHECKED(URemoteControlDeltaAPITestObject, StructValue));
+
+			SerializeBufferFromJsonString("{ \"StructValue\": { \"ColorValue\": { \"G\": 0.2 } } }");
+			ApplyOperation(ValueProperty, ERCModifyOperation::ADD);
+
+			TestEqual("R value", TestObject->StructValue.ColorValue.R, FRemoteControlDeltaAPITestStruct::ColorValueDefault.R);
+			TestEqual("Combined G value", TestObject->StructValue.ColorValue.G, FRemoteControlDeltaAPITestStruct::ColorValueDefault.G + 0.2f);
+			TestEqual("B value", TestObject->StructValue.ColorValue.B, FRemoteControlDeltaAPITestStruct::ColorValueDefault.B);
+			TestEqual("A value", TestObject->StructValue.ColorValue.A, FRemoteControlDeltaAPITestStruct::ColorValueDefault.A);
+		});
+	});
+
+	Describe("Subtract", [this]
+	{
+		It("should subtract integers", [this]
+		{
+			FProperty* ValueProperty = FindFProperty<FProperty>(URemoteControlDeltaAPITestObject::StaticClass(), GET_MEMBER_NAME_CHECKED(URemoteControlDeltaAPITestObject, Int32Value));
+
+			const int32 SubtractedValue = 42;
+			TestObject->Int32Value = SubtractedValue;
+
+			SerializeElementAndReset(ValueProperty);
+			ApplyOperation(ValueProperty, ERCModifyOperation::SUBTRACT);
+
+			TestEqual("Combined value", TestObject->Int32Value, URemoteControlDeltaAPITestObject::Int32ValueDefault - SubtractedValue);
+		});
+
+		It("should subtract floats", [this]
+		{
+			FProperty* ValueProperty = FindFProperty<FProperty>(URemoteControlDeltaAPITestObject::StaticClass(), GET_MEMBER_NAME_CHECKED(URemoteControlDeltaAPITestObject, FloatValue));
+
+			const float SubtractedValue = 0.6f;
+			TestObject->FloatValue = SubtractedValue;
+
+			SerializeElementAndReset(ValueProperty);
+			ApplyOperation(ValueProperty, ERCModifyOperation::SUBTRACT);
+
+			TestEqual("Combined value", TestObject->FloatValue, URemoteControlDeltaAPITestObject::FloatValueDefault - SubtractedValue);
+		});
+
+		It("should subtract vectors", [this]
+		{
+			FProperty* ValueProperty = FindFProperty<FProperty>(URemoteControlDeltaAPITestObject::StaticClass(), GET_MEMBER_NAME_CHECKED(URemoteControlDeltaAPITestObject, VectorValue));
+
+			const FVector SubtractedValue = FVector(-20.f, 0.7f, 101.f);
+			TestObject->VectorValue = SubtractedValue;
+
+			SerializeElementAndReset(ValueProperty);
+			ApplyOperation(ValueProperty, ERCModifyOperation::SUBTRACT);
+
+			TestEqual("Combined value", TestObject->VectorValue, URemoteControlDeltaAPITestObject::VectorValueDefault - SubtractedValue);
+		});
+
+		It("should subtract colors", [this]
+		{
+			FProperty* ValueProperty = FindFProperty<FProperty>(URemoteControlDeltaAPITestObject::StaticClass(), GET_MEMBER_NAME_CHECKED(URemoteControlDeltaAPITestObject, ColorValue));
+
+			const FLinearColor SubtractedValue = FLinearColor(0.4f, 0.1f, 0.2f, 0.8f);
+			TestObject->ColorValue = SubtractedValue;
+
+			SerializeElementAndReset(ValueProperty);
+			ApplyOperation(ValueProperty, ERCModifyOperation::SUBTRACT);
+
+			TestEqual("Combined value", TestObject->ColorValue, URemoteControlDeltaAPITestObject::ColorValueDefault - SubtractedValue);
+		});
+
+		It("should subtract a partial struct's elements only from the provided field", [this]
+		{
+			FProperty* ValueProperty = FindFProperty<FProperty>(URemoteControlDeltaAPITestObject::StaticClass(), GET_MEMBER_NAME_CHECKED(URemoteControlDeltaAPITestObject, VectorValue));
+
+			SerializeBufferFromJsonString("{ \"VectorValue\": { \"Z\": 6.7 } }");
+			ApplyOperation(ValueProperty, ERCModifyOperation::SUBTRACT);
+
+			TestEqual("X value", TestObject->VectorValue.X, URemoteControlDeltaAPITestObject::VectorValueDefault.X);
+			TestEqual("Y value", TestObject->VectorValue.Y, URemoteControlDeltaAPITestObject::VectorValueDefault.Y);
+			TestEqual("Combined Z value", TestObject->VectorValue.Z, URemoteControlDeltaAPITestObject::VectorValueDefault.Z - 6.7f);
+		});
+
+		It("should subtract property with a getter/setter", [this]
+		{
+			FProperty* ValueProperty = FindFProperty<FProperty>(URemoteControlDeltaAPITestObject::StaticClass(), URemoteControlDeltaAPITestObject::GetInt32WithSetterValuePropertyName());
+
+			const int32 SubtractedValue = 27;
+			TestObject->SetInt32WithSetterValue(SubtractedValue);
+
+			SerializeElementAndReset(ValueProperty);
+			ApplyOperation(ValueProperty, ERCModifyOperation::SUBTRACT);
+
+			TestEqual("Combined value", TestObject->GetInt32WithSetterValue(), URemoteControlDeltaAPITestObject::Int32ValueDefault - SubtractedValue);
+		});
+
+		It("should subtract property with a getter/setter and no intercept payload", [this]
+		{
+			FProperty* ValueProperty = FindFProperty<FProperty>(URemoteControlDeltaAPITestObject::StaticClass(), URemoteControlDeltaAPITestObject::GetInt32WithSetterValuePropertyName());
+
+			const int32 SubtractedValue = 27;
+			TestObject->SetInt32WithSetterValue(SubtractedValue);
+
+			SerializeElementAndReset(ValueProperty);
+			ApplyOperation(ValueProperty, ERCModifyOperation::SUBTRACT, false);
+
+			TestEqual("Combined value", TestObject->GetInt32WithSetterValue(), URemoteControlDeltaAPITestObject::Int32ValueDefault - SubtractedValue);
+		});
+
+		It("should subtract nested struct values", [this]
+		{
+			FProperty* ValueProperty = FindFProperty<FProperty>(URemoteControlDeltaAPITestObject::StaticClass(), GET_MEMBER_NAME_CHECKED(URemoteControlDeltaAPITestObject, StructValue));
+
+			const FLinearColor SubtractedValue = FLinearColor(0.2f, 0.5f, 0.1f, 0.3f);
+			TestObject->StructValue.ColorValue = SubtractedValue;
+
+			SerializeElementAndReset(ValueProperty);
+			ApplyOperation(ValueProperty, ERCModifyOperation::SUBTRACT);
+
+			TestEqual("Combined value", TestObject->StructValue.ColorValue, FRemoteControlDeltaAPITestStruct::ColorValueDefault - SubtractedValue);
+		});
+
+		It("should subtract a nested partial struct's elements only to the provided field", [this]
+		{
+			FProperty* ValueProperty = FindFProperty<FProperty>(URemoteControlDeltaAPITestObject::StaticClass(), GET_MEMBER_NAME_CHECKED(URemoteControlDeltaAPITestObject, StructValue));
+
+			SerializeBufferFromJsonString("{ \"StructValue\": { \"ColorValue\": { \"B\": 0.7 } } }");
+			ApplyOperation(ValueProperty, ERCModifyOperation::SUBTRACT);
+
+			TestEqual("R value", TestObject->StructValue.ColorValue.R, FRemoteControlDeltaAPITestStruct::ColorValueDefault.R);
+			TestEqual("G value", TestObject->StructValue.ColorValue.G, FRemoteControlDeltaAPITestStruct::ColorValueDefault.G);
+			TestEqual("Combined B value", TestObject->StructValue.ColorValue.B, FRemoteControlDeltaAPITestStruct::ColorValueDefault.B - 0.7f);
+			TestEqual("A value", TestObject->StructValue.ColorValue.A, FRemoteControlDeltaAPITestStruct::ColorValueDefault.A);
+		});
+	});
+
+	Describe("Multiply", [this]
+	{
+		It("should multiply integers", [this]
+		{
+			FProperty* ValueProperty = FindFProperty<FProperty>(URemoteControlDeltaAPITestObject::StaticClass(), GET_MEMBER_NAME_CHECKED(URemoteControlDeltaAPITestObject, Int32Value));
+
+			const int32 Multiplier = 2;
+			TestObject->Int32Value = Multiplier;
+
+			SerializeElementAndReset(ValueProperty);
+			ApplyOperation(ValueProperty, ERCModifyOperation::MULTIPLY);
+
+			TestEqual("Combined value", TestObject->Int32Value, URemoteControlDeltaAPITestObject::Int32ValueDefault * Multiplier);
+		});
+
+		It("should multiply floats", [this]
+		{
+			FProperty* ValueProperty = FindFProperty<FProperty>(URemoteControlDeltaAPITestObject::StaticClass(), GET_MEMBER_NAME_CHECKED(URemoteControlDeltaAPITestObject, FloatValue));
+
+			const float Multiplier = 3.f;
+			TestObject->FloatValue = Multiplier;
+
+			SerializeElementAndReset(ValueProperty);
+			ApplyOperation(ValueProperty, ERCModifyOperation::MULTIPLY);
+
+			TestEqual("Combined value", TestObject->FloatValue, URemoteControlDeltaAPITestObject::FloatValueDefault * Multiplier);
+		});
+
+		It("should multiply vectors", [this]
+		{
+			FProperty* ValueProperty = FindFProperty<FProperty>(URemoteControlDeltaAPITestObject::StaticClass(), GET_MEMBER_NAME_CHECKED(URemoteControlDeltaAPITestObject, VectorValue));
+
+			const FVector Multiplier = FVector(-0.5f, 0.2f, 3.f);
+			TestObject->VectorValue = Multiplier;
+
+			SerializeElementAndReset(ValueProperty);
+			ApplyOperation(ValueProperty, ERCModifyOperation::MULTIPLY);
+
+			TestEqual("Combined value", TestObject->VectorValue, URemoteControlDeltaAPITestObject::VectorValueDefault * Multiplier);
+		});
+
+		It("should multiply colors", [this]
+		{
+			FProperty* ValueProperty = FindFProperty<FProperty>(URemoteControlDeltaAPITestObject::StaticClass(), GET_MEMBER_NAME_CHECKED(URemoteControlDeltaAPITestObject, ColorValue));
+
+			const FLinearColor Multiplier = FLinearColor(2.f, 0.8f, 1.6f, 1.3f);
+			TestObject->ColorValue = Multiplier;
+
+			SerializeElementAndReset(ValueProperty);
+			ApplyOperation(ValueProperty, ERCModifyOperation::MULTIPLY);
+
+			TestEqual("Combined value", TestObject->ColorValue, URemoteControlDeltaAPITestObject::ColorValueDefault * Multiplier);
+		});
+
+		It("should multiply a partial struct's elements only for the provided field", [this]
+		{
+			FProperty* ValueProperty = FindFProperty<FProperty>(URemoteControlDeltaAPITestObject::StaticClass(), GET_MEMBER_NAME_CHECKED(URemoteControlDeltaAPITestObject, VectorValue));
+
+			SerializeBufferFromJsonString("{ \"VectorValue\": { \"X\": 3.2 } }");
+			ApplyOperation(ValueProperty, ERCModifyOperation::MULTIPLY);
+
+			TestEqual("Combined X value", TestObject->VectorValue.X, URemoteControlDeltaAPITestObject::VectorValueDefault.X * 3.2f);
+			TestEqual("Y value", TestObject->VectorValue.Y, URemoteControlDeltaAPITestObject::VectorValueDefault.Y);
+			TestEqual("Z value", TestObject->VectorValue.Z, URemoteControlDeltaAPITestObject::VectorValueDefault.Z);
+		});
+
+		It("should multiply property with a getter/setter", [this]
+		{
+			FProperty* ValueProperty = FindFProperty<FProperty>(URemoteControlDeltaAPITestObject::StaticClass(), URemoteControlDeltaAPITestObject::GetInt32WithSetterValuePropertyName());
+
+			const int32 Multiplier = 6;
+			TestObject->SetInt32WithSetterValue(Multiplier);
+
+			SerializeElementAndReset(ValueProperty);
+			ApplyOperation(ValueProperty, ERCModifyOperation::MULTIPLY);
+
+			TestEqual("Combined value", TestObject->GetInt32WithSetterValue(), URemoteControlDeltaAPITestObject::Int32ValueDefault * Multiplier);
+		});
+
+		It("should multiply property with a getter/setter and no intercept payload", [this]
+		{
+			FProperty* ValueProperty = FindFProperty<FProperty>(URemoteControlDeltaAPITestObject::StaticClass(), URemoteControlDeltaAPITestObject::GetInt32WithSetterValuePropertyName());
+
+			const int32 Multiplier = 6;
+			TestObject->SetInt32WithSetterValue(Multiplier);
+
+			SerializeElementAndReset(ValueProperty);
+			ApplyOperation(ValueProperty, ERCModifyOperation::MULTIPLY, false);
+
+			TestEqual("Combined value", TestObject->GetInt32WithSetterValue(), URemoteControlDeltaAPITestObject::Int32ValueDefault * Multiplier);
+		});
+
+		It("should multiply nested struct values", [this]
+		{
+			FProperty* ValueProperty = FindFProperty<FProperty>(URemoteControlDeltaAPITestObject::StaticClass(), GET_MEMBER_NAME_CHECKED(URemoteControlDeltaAPITestObject, StructValue));
+
+			const FLinearColor Multiplier = FLinearColor(2.f, 0.8f, 1.6f, 1.3f);
+			TestObject->StructValue.ColorValue = Multiplier;
+
+			SerializeElementAndReset(ValueProperty);
+			ApplyOperation(ValueProperty, ERCModifyOperation::MULTIPLY);
+
+			TestEqual("Combined value", TestObject->StructValue.ColorValue, FRemoteControlDeltaAPITestStruct::ColorValueDefault * Multiplier);
+		});
+
+		It("should multiply a nested partial struct's elements only to the provided field", [this]
+		{
+			FProperty* ValueProperty = FindFProperty<FProperty>(URemoteControlDeltaAPITestObject::StaticClass(), GET_MEMBER_NAME_CHECKED(URemoteControlDeltaAPITestObject, StructValue));
+
+			SerializeBufferFromJsonString("{ \"StructValue\": { \"ColorValue\": { \"R\": 2.5 } } }");
+			ApplyOperation(ValueProperty, ERCModifyOperation::MULTIPLY);
+
+			TestEqual("Combined r value", TestObject->StructValue.ColorValue.R, FRemoteControlDeltaAPITestStruct::ColorValueDefault.R * 2.5f);
+			TestEqual("G value", TestObject->StructValue.ColorValue.G, FRemoteControlDeltaAPITestStruct::ColorValueDefault.G);
+			TestEqual("B value", TestObject->StructValue.ColorValue.B, FRemoteControlDeltaAPITestStruct::ColorValueDefault.B);
+			TestEqual("A value", TestObject->StructValue.ColorValue.A, FRemoteControlDeltaAPITestStruct::ColorValueDefault.A);
+		});
+	});
+
+	Describe("Divide", [this]
+	{
+		It("should divide integers", [this]
+		{
+			FProperty* ValueProperty = FindFProperty<FProperty>(URemoteControlDeltaAPITestObject::StaticClass(), GET_MEMBER_NAME_CHECKED(URemoteControlDeltaAPITestObject, Int32Value));
+
+			const int32 Divisor = 2;
+			TestObject->Int32Value = Divisor;
+
+			SerializeElementAndReset(ValueProperty);
+			ApplyOperation(ValueProperty, ERCModifyOperation::DIVIDE);
+
+			TestEqual("Combined value", TestObject->Int32Value, URemoteControlDeltaAPITestObject::Int32ValueDefault / Divisor);
+		});
+
+		It("should ignore divisions by zero", [this]
+		{
+			FProperty* ValueProperty = FindFProperty<FProperty>(URemoteControlDeltaAPITestObject::StaticClass(), GET_MEMBER_NAME_CHECKED(URemoteControlDeltaAPITestObject, Int32Value));
+
+			const int32 Divisor = 0;
+			TestObject->Int32Value = Divisor;
+
+			SerializeElementAndReset(ValueProperty);
+			ApplyOperation(ValueProperty, ERCModifyOperation::DIVIDE);
+
+			TestEqual("Unaffected value", TestObject->Int32Value, URemoteControlDeltaAPITestObject::Int32ValueDefault);
+		});
+
+		It("should divide floats", [this]
+		{
+			FProperty* ValueProperty = FindFProperty<FProperty>(URemoteControlDeltaAPITestObject::StaticClass(), GET_MEMBER_NAME_CHECKED(URemoteControlDeltaAPITestObject, FloatValue));
+
+			const float Divisor = 3.f;
+			TestObject->FloatValue = Divisor;
+
+			SerializeElementAndReset(ValueProperty);
+			ApplyOperation(ValueProperty, ERCModifyOperation::DIVIDE);
+
+			TestEqual("Combined value", TestObject->FloatValue, URemoteControlDeltaAPITestObject::FloatValueDefault / Divisor);
+		});
+
+		It("should divide vectors", [this]
+		{
+			FProperty* ValueProperty = FindFProperty<FProperty>(URemoteControlDeltaAPITestObject::StaticClass(), GET_MEMBER_NAME_CHECKED(URemoteControlDeltaAPITestObject, VectorValue));
+
+			const FVector Divisor = FVector(-0.5f, 0.2f, 3.f);
+			TestObject->VectorValue = Divisor;
+
+			SerializeElementAndReset(ValueProperty);
+			ApplyOperation(ValueProperty, ERCModifyOperation::DIVIDE);
+
+			TestEqual("Combined value", TestObject->VectorValue, URemoteControlDeltaAPITestObject::VectorValueDefault / Divisor);
+		});
+
+		It("should divide colors", [this]
+		{
+			FProperty* ValueProperty = FindFProperty<FProperty>(URemoteControlDeltaAPITestObject::StaticClass(), GET_MEMBER_NAME_CHECKED(URemoteControlDeltaAPITestObject, ColorValue));
+
+			const FLinearColor Divisor = FLinearColor(2.f, 0.8f, 1.6f, 1.3f);
+			TestObject->ColorValue = Divisor;
+
+			SerializeElementAndReset(ValueProperty);
+			ApplyOperation(ValueProperty, ERCModifyOperation::DIVIDE);
+
+			TestEqual("Combined value", TestObject->ColorValue, URemoteControlDeltaAPITestObject::ColorValueDefault / Divisor);
+		});
+
+		It("should divide a partial struct's elements only for the provided field", [this]
+		{
+			FProperty* ValueProperty = FindFProperty<FProperty>(URemoteControlDeltaAPITestObject::StaticClass(), GET_MEMBER_NAME_CHECKED(URemoteControlDeltaAPITestObject, VectorValue));
+
+			SerializeBufferFromJsonString("{ \"VectorValue\": { \"Y\": 2.2 } }");
+			ApplyOperation(ValueProperty, ERCModifyOperation::DIVIDE);
+
+			TestEqual("X value", TestObject->VectorValue.X, URemoteControlDeltaAPITestObject::VectorValueDefault.X);
+			TestEqual("Combined Y value", TestObject->VectorValue.Y, URemoteControlDeltaAPITestObject::VectorValueDefault.Y / 2.2f);
+			TestEqual("Z value", TestObject->VectorValue.Z, URemoteControlDeltaAPITestObject::VectorValueDefault.Z);
+		});
+
+		It("should multiply property with a getter/setter", [this]
+		{
+			FProperty* ValueProperty = FindFProperty<FProperty>(URemoteControlDeltaAPITestObject::StaticClass(), URemoteControlDeltaAPITestObject::GetInt32WithSetterValuePropertyName());
+
+			const int32 Divisor = 4;
+			TestObject->SetInt32WithSetterValue(Divisor);
+
+			SerializeElementAndReset(ValueProperty);
+			ApplyOperation(ValueProperty, ERCModifyOperation::DIVIDE);
+
+			TestEqual("Combined value", TestObject->GetInt32WithSetterValue(), URemoteControlDeltaAPITestObject::Int32ValueDefault / Divisor);
+		});
+
+		It("should multiply property with a getter/setter and no intercept payload", [this]
+		{
+			FProperty* ValueProperty = FindFProperty<FProperty>(URemoteControlDeltaAPITestObject::StaticClass(), URemoteControlDeltaAPITestObject::GetInt32WithSetterValuePropertyName());
+
+			const int32 Divisor = 4;
+			TestObject->SetInt32WithSetterValue(Divisor);
+
+			SerializeElementAndReset(ValueProperty);
+			ApplyOperation(ValueProperty, ERCModifyOperation::DIVIDE, false);
+
+			TestEqual("Combined value", TestObject->GetInt32WithSetterValue(), URemoteControlDeltaAPITestObject::Int32ValueDefault / Divisor);
+		});
+
+		It("should divide nested struct values", [this]
+		{
+			FProperty* ValueProperty = FindFProperty<FProperty>(URemoteControlDeltaAPITestObject::StaticClass(), GET_MEMBER_NAME_CHECKED(URemoteControlDeltaAPITestObject, StructValue));
+
+			const FLinearColor Divisor = FLinearColor(2.f, 0.8f, 1.6f, 1.3f);
+			TestObject->StructValue.ColorValue = Divisor;
+
+			SerializeElementAndReset(ValueProperty);
+			ApplyOperation(ValueProperty, ERCModifyOperation::DIVIDE);
+
+			TestEqual("Combined value", TestObject->StructValue.ColorValue, FRemoteControlDeltaAPITestStruct::ColorValueDefault / Divisor);
+		});
+
+		It("should divide a nested partial struct's elements only to the provided field", [this]
+		{
+			FProperty* ValueProperty = FindFProperty<FProperty>(URemoteControlDeltaAPITestObject::StaticClass(), GET_MEMBER_NAME_CHECKED(URemoteControlDeltaAPITestObject, StructValue));
+
+			SerializeBufferFromJsonString("{ \"StructValue\": { \"ColorValue\": { \"A\": 1.7 } } }");
+			ApplyOperation(ValueProperty, ERCModifyOperation::DIVIDE);
+
+			TestEqual("R value", TestObject->StructValue.ColorValue.R, FRemoteControlDeltaAPITestStruct::ColorValueDefault.R);
+			TestEqual("G value", TestObject->StructValue.ColorValue.G, FRemoteControlDeltaAPITestStruct::ColorValueDefault.G);
+			TestEqual("B value", TestObject->StructValue.ColorValue.B, FRemoteControlDeltaAPITestStruct::ColorValueDefault.B);
+			TestEqual("Combined A value", TestObject->StructValue.ColorValue.A, FRemoteControlDeltaAPITestStruct::ColorValueDefault.A / 1.7f);
+		});
+	});
+}
diff --git a/Engine/Plugins/VirtualProduction/RemoteControl/Source/RemoteControl/Public/IRemoteControlModule.h b/Engine/Plugins/VirtualProduction/RemoteControl/Source/RemoteControl/Public/IRemoteControlModule.h
--- a/Engine/Plugins/VirtualProduction/RemoteControl/Source/RemoteControl/Public/IRemoteControlModule.h
+++ b/Engine/Plugins/VirtualProduction/RemoteControl/Source/RemoteControl/Public/IRemoteControlModule.h
@@ -101,6 +101,19 @@ enum class ERCAccess : uint8
 	WRITE_TRANSACTION_ACCESS,
 };
 
+/**
+ * Type of operation to perform when setting a remote property's value
+ */
+UENUM()
+enum class ERCModifyOperation : uint8
+{
+	EQUAL,
+	ADD,
+	SUBTRACT,
+	MULTIPLY,
+	DIVIDE
+};
+
 /**
  * Reference to a UObject or one of its properties
  */
@@ -266,9 +279,10 @@ public:
 	 * @param Backend the struct deserializer backend to use to deserialize the object properties.
 	 * @param InPayloadType the payload type archive.
 	 * @param InInterceptPayload the payload reference archive for the interception.
+	 * @param Operation the type of operation to perform when setting the value.
 	 * @return true if the deserialization succeeded
 	 */
-	virtual bool SetObjectProperties(const FRCObjectReference& ObjectAccess, IStructDeserializerBackend& Backend, ERCPayloadType InPayloadType = ERCPayloadType::Json, const TArray<uint8>& InInterceptPayload = TArray<uint8>()) = 0;
+	virtual bool SetObjectProperties(const FRCObjectReference& ObjectAccess, IStructDeserializerBackend& Backend, ERCPayloadType InPayloadType = ERCPayloadType::Json, const TArray<uint8>& InInterceptPayload = TArray<uint8>(), ERCModifyOperation Operation = ERCModifyOperation::EQUAL) = 0;
 
 	/**
 	 * Reset the property or the object the Object Reference is pointing to
diff --git a/Engine/Plugins/VirtualProduction/RemoteControl/Source/WebRemoteControl/Private/RemoteControlRequest.h b/Engine/Plugins/VirtualProduction/RemoteControl/Source/WebRemoteControl/Private/RemoteControlRequest.h
--- a/Engine/Plugins/VirtualProduction/RemoteControl/Source/WebRemoteControl/Private/RemoteControlRequest.h
+++ b/Engine/Plugins/VirtualProduction/RemoteControl/Source/WebRemoteControl/Private/RemoteControlRequest.h
@@ -235,6 +235,12 @@ public:
 	UPROPERTY()
 	bool GenerateTransaction = false;
 
+	/**
+	 * Which type of operation should be performed if this is modifying a property.
+	 */
+	UPROPERTY()
+	ERCModifyOperation Operation = ERCModifyOperation::EQUAL;
+
 private:
 	/**
 	 * Indicates if the property should be read or written to.
diff --git a/Engine/Plugins/VirtualProduction/RemoteControl/Source/WebRemoteControl/Private/WebRemoteControl.cpp b/Engine/Plugins/VirtualProduction/RemoteControl/Source/WebRemoteControl/Private/WebRemoteControl.cpp
--- a/Engine/Plugins/VirtualProduction/RemoteControl/Source/WebRemoteControl/Private/WebRemoteControl.cpp
+++ b/Engine/Plugins/VirtualProduction/RemoteControl/Source/WebRemoteControl/Private/WebRemoteControl.cpp
@@ -866,7 +866,7 @@ bool FWebRemoteControlModule::HandleObjectPropertyRoute(const FHttpServerRequest
 			Reader.SetLimitSize(PropertyValueDelimiters.BlockEnd);
 			FRCJsonStructDeserializerBackend DeserializerBackend(Reader);
 			// Set a ERCPayloadType and TCHARBody in order to follow the replication path
-			if (IRemoteControlModule::Get().SetObjectProperties(ObjectRef, DeserializerBackend, ERCPayloadType::Json, DeserializedRequest.TCHARBody))
+			if (IRemoteControlModule::Get().SetObjectProperties(ObjectRef, DeserializerBackend, ERCPayloadType::Json, DeserializedRequest.TCHARBody, DeserializedRequest.Operation))
 			{
 				Response->Code = EHttpServerResponseCodes::Ok;
 			}
diff --git a/Engine/Plugins/VirtualProduction/RemoteControlInterception/Source/RemoteControlInterception/Public/IRemoteControlInterceptionCommands.h b/Engine/Plugins/VirtualProduction/RemoteControlInterception/Source/RemoteControlInterception/Public/IRemoteControlInterceptionCommands.h
--- a/Engine/Plugins/VirtualProduction/RemoteControlInterception/Source/RemoteControlInterception/Public/IRemoteControlInterceptionCommands.h
+++ b/Engine/Plugins/VirtualProduction/RemoteControlInterception/Source/RemoteControlInterception/Public/IRemoteControlInterceptionCommands.h
@@ -39,6 +39,18 @@ enum class ERCIAccess : uint8
 	WRITE_TRANSACTION_ACCESS,
 };
 
+/**
+ * Type of operation to perform when setting a remote property's value (proxy type for ERCModifyOperation to avoid module dependency)
+ */
+enum class ERCIModifyOperation : uint8
+{
+	EQUAL,
+	ADD,
+	SUBTRACT,
+	MULTIPLY,
+	DIVIDE
+};
+
 
 /**
  * The list of remote control commands available for interception
@@ -159,9 +171,10 @@ struct FRCIPropertiesMetadata : public FRCIObjectMetadata
 public:
 	FRCIPropertiesMetadata() = default;
 
-	FRCIPropertiesMetadata(const FString& InObjectPath, const FString& InPropertyPath, const FString& InPropertyPathInfo, const ERCIAccess InAccess, const ERCIPayloadType InPayloadType, const TArray<uint8>& InPayload)
+	FRCIPropertiesMetadata(const FString& InObjectPath, const FString& InPropertyPath, const FString& InPropertyPathInfo, const ERCIAccess InAccess, const ERCIPayloadType InPayloadType, const ERCIModifyOperation InOperation, const TArray<uint8>& InPayload)
 		: FRCIObjectMetadata(InObjectPath, InPropertyPath, InPropertyPathInfo, InAccess)
 		, PayloadType(InPayloadType)
+		, Operation(InOperation)
 		, Payload(InPayload)
 	{ }
 
@@ -199,6 +212,9 @@ public:
 	/** Intercepted payload type */
 	ERCIPayloadType PayloadType = ERCIPayloadType::Json;
 
+	/** Type of operation to perform on the property **/
+	ERCIModifyOperation Operation;
+
 	/** Property payload to intercept */
 	const TArray<uint8> Payload;
 
