RemoteControl: Editor crash when disabling Log Webapp requests handle duration

diff --git a/Engine/Plugins/VirtualProduction/RemoteControl/Source/WebRemoteControl/Private/WebRemoteControlExternalLogger.cpp b/Engine/Plugins/VirtualProduction/RemoteControl/Source/WebRemoteControl/Private/WebRemoteControlExternalLogger.cpp
--- a/Engine/Plugins/VirtualProduction/RemoteControl/Source/WebRemoteControl/Private/WebRemoteControlExternalLogger.cpp
+++ b/Engine/Plugins/VirtualProduction/RemoteControl/Source/WebRemoteControl/Private/WebRemoteControlExternalLogger.cpp
@@ -20,6 +20,12 @@ FWebRemoteControlExternalLogger::FWebRemoteControlExternalLogger(TSharedPtr<INet
 	WebSocketConnection->SetSocketClosedCallBack(ClosedCallback);
 }
 
+FWebRemoteControlExternalLogger::~FWebRemoteControlExternalLogger()
+{
+	FTSTicker::RemoveTicker(TickerHandle);
+	TickerHandle.Reset();
+}
+
 void FWebRemoteControlExternalLogger::Log(int32 RequestId, const TCHAR* Stage)
 {
 	if (IsConnected)
diff --git a/Engine/Plugins/VirtualProduction/RemoteControl/Source/WebRemoteControl/Private/WebRemoteControlExternalLogger.h b/Engine/Plugins/VirtualProduction/RemoteControl/Source/WebRemoteControl/Private/WebRemoteControlExternalLogger.h
--- a/Engine/Plugins/VirtualProduction/RemoteControl/Source/WebRemoteControl/Private/WebRemoteControlExternalLogger.h
+++ b/Engine/Plugins/VirtualProduction/RemoteControl/Source/WebRemoteControl/Private/WebRemoteControlExternalLogger.h
@@ -10,6 +10,7 @@ class FWebRemoteControlExternalLogger
 {
 public:
 	FWebRemoteControlExternalLogger(TSharedPtr<class INetworkingWebSocket> WebSocketConnection);
+	virtual ~FWebRemoteControlExternalLogger();
 
 	void Log(int32 RequestId, const TCHAR* Stage);
 
