Fix RC regression when exposing property

diff --git a/Engine/Plugins/VirtualProduction/RemoteControl/Source/RemoteControlUI/Private/UI/SRemoteControlPanel.cpp b/Engine/Plugins/VirtualProduction/RemoteControl/Source/RemoteControlUI/Private/UI/SRemoteControlPanel.cpp
--- a/Engine/Plugins/VirtualProduction/RemoteControl/Source/RemoteControlUI/Private/UI/SRemoteControlPanel.cpp
+++ b/Engine/Plugins/VirtualProduction/RemoteControl/Source/RemoteControlUI/Private/UI/SRemoteControlPanel.cpp
@@ -499,7 +499,6 @@ void SRemoteControlPanel::ToggleProperty(const FRCExposesPropertyArgs& InPropert
 
 	auto PostExpose = [this, InPropertyArgs]()
 	{
-		Refresh();
 		CachedExposedPropertyArgs.Emplace(InPropertyArgs);
 	};
 
