Added remote control support for FColor and FLinearColor

diff --git a/Engine/Plugins/VirtualProduction/RemoteControl/Source/RemoteControl/Private/RemoteControlPropertyHandle.cpp b/Engine/Plugins/VirtualProduction/RemoteControl/Source/RemoteControl/Private/RemoteControlPropertyHandle.cpp
--- a/Engine/Plugins/VirtualProduction/RemoteControl/Source/RemoteControl/Private/RemoteControlPropertyHandle.cpp
+++ b/Engine/Plugins/VirtualProduction/RemoteControl/Source/RemoteControl/Private/RemoteControlPropertyHandle.cpp
@@ -6,6 +6,7 @@
 #include "RemoteControlFieldPath.h"
 #include "RemoteControlPreset.h"
 
+#include "Algo/AllOf.h"
 #include "Backends/CborStructDeserializerBackend.h"
 #include "CborWriter.h"
 #include "Serialization/MemoryReader.h"
@@ -14,7 +15,7 @@
 
 namespace
 {
-	/** Check if the given propery is Array or Set or Map */
+	/** Check if the given property is Array or Set or Map */
 	bool IsArrayLike(const FProperty* Property)
 	{
 		return (Property->IsA<FArrayProperty>()
@@ -22,7 +23,7 @@ namespace
 			|| Property->IsA<FMapProperty>());
 	}
 
-	/** Get inner propery of Array or Set or Map */
+	/** Get inner property of Array or Set or Map */
 	FProperty* GetContainerInnerProperty(const FProperty* Property)
 	{
 		if (const FArrayProperty* ArrayProperty = CastField<FArrayProperty>(Property))
@@ -288,6 +289,14 @@ TSharedPtr<IRemoteControlPropertyHandle> FRemoteControlPropertyHandle::GetProper
 	{
 		PropertyHandle = MakeShared<FRemoteControlPropertyHandleRotator>(InRCProperty, InProperty, InParentProperty, InParentFieldPath, InArrayIndex);
 	}
+	else if (FRemoteControlPropertyHandleColor::Supports(InProperty))
+	{
+		PropertyHandle = MakeShared<FRemoteControlPropertyHandleColor>(InRCProperty, InProperty, InParentProperty, InParentFieldPath, InArrayIndex);
+	}
+	else if (FRemoteControlPropertyHandleLinearColor::Supports(InProperty))
+	{
+		PropertyHandle = MakeShared<FRemoteControlPropertyHandleLinearColor>(InRCProperty, InProperty, InParentProperty, InParentFieldPath, InArrayIndex);
+	}
 	else if (FRemoteControlPropertyHandleSet::Supports(InProperty))
 	{
 		PropertyHandle = MakeShared<FRemoteControlPropertyHandleSet>(InRCProperty, InProperty, InParentProperty, InParentFieldPath, InArrayIndex);
@@ -380,6 +389,8 @@ IMPLEMENT_PROPERTY_ACCESSOR(FVector2D)
 IMPLEMENT_PROPERTY_ACCESSOR(FVector4)
 IMPLEMENT_PROPERTY_ACCESSOR(FQuat)
 IMPLEMENT_PROPERTY_ACCESSOR(FRotator)
+IMPLEMENT_PROPERTY_ACCESSOR(FColor)
+IMPLEMENT_PROPERTY_ACCESSOR(FLinearColor)
 #undef IMPLEMENT_PROPERTY_ACCESSOR
 
 
@@ -1101,7 +1112,7 @@ bool FRemoteControlPropertyHandleByte::GetValue(uint8& OutValue) const
 		{
 			return false;
 		}
-		
+
 		const FEnumProperty* EnumProperty = CastField<FEnumProperty>(Property);
 		if (!ensure(EnumProperty))
 		{
@@ -1112,12 +1123,12 @@ bool FRemoteControlPropertyHandleByte::GetValue(uint8& OutValue) const
 
 		return true;
 	}
-	
+
 	return false;
 }
 
 bool FRemoteControlPropertyHandleByte::SetValue(uint8 InValue)
-{	
+{
 	FProperty* Property = GetProperty();
 	if (!ensure(Property))
 	{
@@ -1227,7 +1238,7 @@ bool FRemoteControlPropertyHandleVector::Supports(const FProperty* InProperty)
 
 bool FRemoteControlPropertyHandleVector::GetValue(FVector& OutValue) const
 {
-	if (!ensure(VectorComponents.Num() != 3 || IsComponentsValid()))
+	if (!ensure(VectorComponents.Num() != 3 || AreComponentsValid()))
 	{
 		return false;
 	}
@@ -1250,7 +1261,7 @@ bool FRemoteControlPropertyHandleVector::GetValue(FVector& OutValue) const
 
 bool FRemoteControlPropertyHandleVector::SetValue(FVector InValue)
 {
-	if (!ensure(VectorComponents.Num() != 3 || IsComponentsValid()))
+	if (!ensure(VectorComponents.Num() != 3 || AreComponentsValid()))
 	{
 		return false;
 	}
@@ -1270,7 +1281,7 @@ bool FRemoteControlPropertyHandleVector::SetValue(FVector InValue)
 
 bool FRemoteControlPropertyHandleVector::GetValue(FVector2D& OutValue) const
 {
-	if (!ensure(VectorComponents.Num() != 2 || IsComponentsValid()))
+	if (!ensure(VectorComponents.Num() != 2 || AreComponentsValid()))
 	{
 		return false;
 	}
@@ -1292,7 +1303,7 @@ bool FRemoteControlPropertyHandleVector::GetValue(FVector2D& OutValue) const
 
 bool FRemoteControlPropertyHandleVector::SetValue(FVector2D InValue)
 {
-	if (!ensure(VectorComponents.Num() != 2 || IsComponentsValid()))
+	if (!ensure(VectorComponents.Num() != 2 || AreComponentsValid()))
 	{
 		return false;
 	}
@@ -1311,7 +1322,7 @@ bool FRemoteControlPropertyHandleVector::SetValue(FVector2D InValue)
 
 bool FRemoteControlPropertyHandleVector::GetValue(FVector4& OutValue) const
 {
-	if (!ensure(VectorComponents.Num() != 4 || IsComponentsValid()))
+	if (!ensure(VectorComponents.Num() != 4 || AreComponentsValid()))
 	{
 		return false;
 	}
@@ -1335,7 +1346,7 @@ bool FRemoteControlPropertyHandleVector::GetValue(FVector4& OutValue) const
 
 bool FRemoteControlPropertyHandleVector::SetValue(FVector4 InValue)
 {
-	if (!ensure(VectorComponents.Num() != 4 || IsComponentsValid()))
+	if (!ensure(VectorComponents.Num() != 4 || AreComponentsValid()))
 	{
 		return false;
 	}
@@ -1420,7 +1431,7 @@ bool FRemoteControlPropertyHandleVector::SetW(float InValue)
 	return false;
 }
 
-bool FRemoteControlPropertyHandleVector::IsComponentsValid() const
+bool FRemoteControlPropertyHandleVector::AreComponentsValid() const
 {
 	for (TWeakPtr<IRemoteControlPropertyHandle> VectorComponentPtr : VectorComponents)
 	{
@@ -1456,7 +1467,7 @@ bool FRemoteControlPropertyHandleRotator::Supports(const FProperty* InProperty)
 
 bool FRemoteControlPropertyHandleRotator::GetValue(FRotator& OutValue) const
 {
-	if (!ensure(IsComponentsValid()))
+	if (!ensure(AreComponentsValid()))
 	{
 		return false;
 	}
@@ -1476,7 +1487,7 @@ bool FRemoteControlPropertyHandleRotator::GetValue(FRotator& OutValue) const
 
 bool FRemoteControlPropertyHandleRotator::SetValue(FRotator InValue)
 {
-	if (!ensure(IsComponentsValid()))
+	if (!ensure(AreComponentsValid()))
 	{
 		return false;
 	}
@@ -1524,7 +1535,223 @@ bool FRemoteControlPropertyHandleRotator::SetYaw(float InValue)
 	return false;
 }
 
-bool FRemoteControlPropertyHandleRotator::IsComponentsValid() const
+bool FRemoteControlPropertyHandleRotator::AreComponentsValid() const
 {
 	return RollValue.Pin().IsValid() && PitchValue.Pin().IsValid() && YawValue.Pin().IsValid();
 }
+
+FRemoteControlPropertyHandleColor::FRemoteControlPropertyHandleColor(const TSharedPtr<FRemoteControlProperty>& InRCProperty, FProperty* InProperty, FProperty* InParentProperty, const FString& InParentFieldPath, int32 InArrayIndex)
+	: FRemoteControlPropertyHandle(InRCProperty, InProperty, InParentProperty, InParentFieldPath, InArrayIndex)
+{
+	constexpr bool bRecurse = false;
+
+	ColorComponents.Add(GetChildHandle("R", bRecurse));
+	ColorComponents.Add(GetChildHandle("G", bRecurse));
+	ColorComponents.Add(GetChildHandle("B", bRecurse));
+	ColorComponents.Add(GetChildHandle("A", bRecurse));
+}
+
+bool FRemoteControlPropertyHandleColor::Supports(const FProperty* InProperty)
+{
+	if (!InProperty)
+	{
+		return false;
+	}
+
+	const FStructProperty* StructProp = CastField<FStructProperty>(InProperty);
+
+	bool bSupported = false;
+	if (StructProp && StructProp->Struct)
+	{
+		const FName StructName = StructProp->Struct->GetFName();
+		bSupported = StructName == NAME_Color;
+	}
+
+	return bSupported;
+}
+
+bool FRemoteControlPropertyHandleColor::GetValue(FColor& OutValue) const
+{
+	if (!ensure(ColorComponents.Num() == 4 && AreComponentsValid()))
+	{
+		return false;
+	}
+
+	const bool ResR = ColorComponents[0].Pin()->GetValue(OutValue.R);
+	const bool ResG = ColorComponents[1].Pin()->GetValue(OutValue.G);
+	const bool ResB = ColorComponents[2].Pin()->GetValue(OutValue.B);
+	const bool ResA = ColorComponents[3].Pin()->GetValue(OutValue.A);
+
+	return ResR && ResG && ResB && ResA;
+}
+
+bool FRemoteControlPropertyHandleColor::SetValue(FColor InValue)
+{
+	if (!ensure(ColorComponents.Num() == 4 && AreComponentsValid()))
+	{
+		return false;
+	}
+
+	const bool ResR = ColorComponents[0].Pin()->SetValue(InValue.R);
+	const bool ResG = ColorComponents[1].Pin()->SetValue(InValue.G);
+	const bool ResB = ColorComponents[2].Pin()->SetValue(InValue.B);
+	const bool ResA = ColorComponents[3].Pin()->SetValue(InValue.A);
+
+	return ResR && ResG && ResB && ResA;
+}
+
+bool FRemoteControlPropertyHandleColor::SetR(uint8 InValue)
+{
+	if (ensure(ColorComponents.IsValidIndex(0) && ColorComponents[0].Pin().IsValid()))
+	{
+		return ColorComponents[0].Pin()->SetValue(InValue);
+	}
+
+	return false;
+}
+
+bool FRemoteControlPropertyHandleColor::SetG(uint8 InValue)
+{
+	if (ensure(ColorComponents.IsValidIndex(1) && ColorComponents[1].Pin().IsValid()))
+	{
+		return ColorComponents[1].Pin()->SetValue(InValue);
+	}
+
+	return false;
+}
+
+bool FRemoteControlPropertyHandleColor::SetB(uint8 InValue)
+{
+	if (ensure(ColorComponents.IsValidIndex(2) && ColorComponents[2].Pin().IsValid()))
+	{
+		return ColorComponents[2].Pin()->SetValue(InValue);
+	}
+
+	return false;
+}
+
+bool FRemoteControlPropertyHandleColor::SetA(uint8 InValue)
+{
+	if (ensure(ColorComponents.IsValidIndex(3) && ColorComponents[3].Pin().IsValid()))
+	{
+		return ColorComponents[3].Pin()->SetValue(InValue);
+	}
+
+	return false;
+}
+
+bool FRemoteControlPropertyHandleColor::AreComponentsValid() const
+{
+	return Algo::AllOf(ColorComponents, [](const TWeakPtr<IRemoteControlPropertyHandle> ColorComponentPtr)
+	{
+		return ColorComponentPtr.Pin().IsValid();
+	});
+}
+
+FRemoteControlPropertyHandleLinearColor::FRemoteControlPropertyHandleLinearColor(const TSharedPtr<FRemoteControlProperty>& InRCProperty, FProperty* InProperty, FProperty* InParentProperty, const FString& InParentFieldPath, int32 InArrayIndex)
+	: FRemoteControlPropertyHandle(InRCProperty, InProperty, InParentProperty, InParentFieldPath, InArrayIndex)
+{
+	constexpr bool bRecurse = false;
+
+	ColorComponents.Add(GetChildHandle("R", bRecurse));
+	ColorComponents.Add(GetChildHandle("G", bRecurse));
+	ColorComponents.Add(GetChildHandle("B", bRecurse));
+	ColorComponents.Add(GetChildHandle("A", bRecurse));
+}
+
+bool FRemoteControlPropertyHandleLinearColor::Supports(const FProperty* InProperty)
+{
+	if (!InProperty)
+	{
+		return false;
+	}
+
+	const FStructProperty* StructProp = CastField<FStructProperty>(InProperty);
+
+	bool bSupported = false;
+	if (StructProp && StructProp->Struct)
+	{
+		const FName StructName = StructProp->Struct->GetFName();
+		bSupported = StructName == NAME_LinearColor;
+	}
+
+	return bSupported;
+}
+
+bool FRemoteControlPropertyHandleLinearColor::GetValue(FLinearColor& OutValue) const
+{
+	if (!ensure(ColorComponents.Num() == 4 && AreComponentsValid()))
+	{
+		return false;
+	}
+
+	const bool ResR = ColorComponents[0].Pin()->GetValue(OutValue.R);
+	const bool ResG = ColorComponents[1].Pin()->GetValue(OutValue.G);
+	const bool ResB = ColorComponents[2].Pin()->GetValue(OutValue.B);
+	const bool ResA = ColorComponents[3].Pin()->GetValue(OutValue.A);
+
+	return ResR && ResG && ResB && ResA;
+}
+
+bool FRemoteControlPropertyHandleLinearColor::SetValue(FLinearColor InValue)
+{
+	if (!ensure(ColorComponents.Num() == 4 && AreComponentsValid()))
+	{
+		return false;
+	}
+
+	const bool ResR = ColorComponents[0].Pin()->SetValue(InValue.R);
+	const bool ResG = ColorComponents[1].Pin()->SetValue(InValue.G);
+	const bool ResB = ColorComponents[2].Pin()->SetValue(InValue.B);
+	const bool ResA = ColorComponents[3].Pin()->SetValue(InValue.A);
+
+	return ResR && ResG && ResB && ResA;
+}
+
+bool FRemoteControlPropertyHandleLinearColor::SetR(float InValue)
+{
+	if (ensure(ColorComponents.IsValidIndex(0) && ColorComponents[0].Pin().IsValid()))
+	{
+		return ColorComponents[0].Pin()->SetValue(InValue);
+	}
+
+	return false;
+}
+
+bool FRemoteControlPropertyHandleLinearColor::SetG(float InValue)
+{
+	if (ensure(ColorComponents.IsValidIndex(1) && ColorComponents[1].Pin().IsValid()))
+	{
+		return ColorComponents[1].Pin()->SetValue(InValue);
+	}
+
+	return false;
+}
+
+bool FRemoteControlPropertyHandleLinearColor::SetB(float InValue)
+{
+	if (ensure(ColorComponents.IsValidIndex(2) && ColorComponents[2].Pin().IsValid()))
+	{
+		return ColorComponents[2].Pin()->SetValue(InValue);
+	}
+
+	return false;
+}
+
+bool FRemoteControlPropertyHandleLinearColor::SetA(float InValue)
+{
+	if (ensure(ColorComponents.IsValidIndex(3) && ColorComponents[3].Pin().IsValid()))
+	{
+		return ColorComponents[3].Pin()->SetValue(InValue);
+	}
+
+	return false;
+}
+
+bool FRemoteControlPropertyHandleLinearColor::AreComponentsValid() const
+{
+	return Algo::AllOf(ColorComponents, [](const TWeakPtr<IRemoteControlPropertyHandle> ColorComponentPtr)
+	{
+		return ColorComponentPtr.Pin().IsValid();
+	});
+}
\ No newline at end of file
diff --git a/Engine/Plugins/VirtualProduction/RemoteControl/Source/RemoteControl/Private/RemoteControlPropertyHandle.h b/Engine/Plugins/VirtualProduction/RemoteControl/Source/RemoteControl/Private/RemoteControlPropertyHandle.h
--- a/Engine/Plugins/VirtualProduction/RemoteControl/Source/RemoteControl/Private/RemoteControlPropertyHandle.h
+++ b/Engine/Plugins/VirtualProduction/RemoteControl/Source/RemoteControl/Private/RemoteControlPropertyHandle.h
@@ -87,6 +87,8 @@ public:
 	DECLARE_PROPERTY_ACCESSOR(FVector4)
 	DECLARE_PROPERTY_ACCESSOR(FQuat)
 	DECLARE_PROPERTY_ACCESSOR(FRotator)
+	DECLARE_PROPERTY_ACCESSOR(FColor)
+	DECLARE_PROPERTY_ACCESSOR(FLinearColor)
 #undef DECLARE_PROPERTY_ACCESSOR
 
 	virtual bool SetValue(const TCHAR* InValue) override { return false; }
@@ -381,7 +383,8 @@ public:
 	bool SetW(float InValue);
 
 	/** Check if all vector components are valid */
-	bool IsComponentsValid() const;
+	bool AreComponentsValid() const;
+	
 private:
 	/** Array of child struct vector components */
 	TArray< TWeakPtr<IRemoteControlPropertyHandle> > VectorComponents;
@@ -417,7 +420,7 @@ public:
 	bool SetYaw(float InValue);
 
 	/** Check if all rotator components are valid */
-	bool IsComponentsValid() const;
+	bool AreComponentsValid() const;
 
 private:
 	/** Roll handle of the rotator */
@@ -506,4 +509,84 @@ public:
 	//~ End IRemoteControlPropertyHandleMap Interface
 };
 
+/**
+ * The color implementation of a property handle
+ */
+class FRemoteControlPropertyHandleColor : public FRemoteControlPropertyHandle
+{
+public:
+	DECLARE_PROPERTY_ACCESSOR_CONSTRUCTOR(FRemoteControlPropertyHandleColor)
+
+	/**
+	 * Check if color handle support the given property
+	 * @param InProperty Property to check
+	 * @return true if this handle supports the given property
+	 */
+	static bool Supports(const FProperty* InProperty);
+
+	//~ Begin IRemoteControlPropertyHandle Interface
+	virtual bool GetValue(FColor& OutValue) const override;
+	virtual bool SetValue(FColor InValue) override;
+	//~ End IRemoteControlPropertyHandle Interface
+
+	/** Set R value of the color */
+	bool SetR(uint8 InValue);
+
+	/** Set G value of the color */
+	bool SetG(uint8 InValue);
+
+	/** Set B value of the color */
+	bool SetB(uint8 InValue);
+
+	/** Set A value of the color */
+	bool SetA(uint8 InValue);
+
+	/** Check if all color components are valid */
+	bool AreComponentsValid() const;
+	
+private:
+	/** Array of child struct color components */
+	TArray< TWeakPtr<IRemoteControlPropertyHandle> > ColorComponents;
+};
+
+/**
+* The linear color implementation of a property handle
+*/
+class FRemoteControlPropertyHandleLinearColor : public FRemoteControlPropertyHandle
+{
+public:
+	DECLARE_PROPERTY_ACCESSOR_CONSTRUCTOR(FRemoteControlPropertyHandleLinearColor)
+
+	/**
+	* Check if linear color handle support the given property
+	* @param InProperty Property to check
+	* @return true if this handle supports the given property
+	*/
+	static bool Supports(const FProperty* InProperty);
+
+	//~ Begin IRemoteControlPropertyHandle Interface
+	virtual bool GetValue(FLinearColor& OutValue) const override;
+	virtual bool SetValue(FLinearColor InValue) override;
+	//~ End IRemoteControlPropertyHandle Interface
+
+	/** Set R value of the linear color */
+	bool SetR(float InValue);
+
+	/** Set G value of the linear color */
+	bool SetG(float InValue);
+
+	/** Set B value of the linear color */
+	bool SetB(float InValue);
+
+	/** Set A value of the linear color */
+	bool SetA(float InValue);
+
+	/** Check if all color components are valid */
+	bool AreComponentsValid() const;
+	
+private:
+	/** Array of child struct color components */
+	TArray< TWeakPtr<IRemoteControlPropertyHandle> > ColorComponents;
+};
+
 #undef DECLARE_PROPERTY_ACCESSOR_CONSTRUCTOR
diff --git a/Engine/Plugins/VirtualProduction/RemoteControl/Source/RemoteControl/Private/Tests/RemoteControlPropertyHandleTestData.h b/Engine/Plugins/VirtualProduction/RemoteControl/Source/RemoteControl/Private/Tests/RemoteControlPropertyHandleTestData.h
--- a/Engine/Plugins/VirtualProduction/RemoteControl/Source/RemoteControl/Private/Tests/RemoteControlPropertyHandleTestData.h
+++ b/Engine/Plugins/VirtualProduction/RemoteControl/Source/RemoteControl/Private/Tests/RemoteControlPropertyHandleTestData.h
@@ -24,7 +24,7 @@ namespace ERemoteControlEnum
 }
 
 USTRUCT()
-struct FRemoteControlTestStructInnerSimle
+struct FRemoteControlTestStructInnerSimple
 {
 	GENERATED_BODY()
 
@@ -41,7 +41,7 @@ struct FRemoteControlTestStructInner
 	int8 Int8Value = 36;
 
 	UPROPERTY(EditAnywhere, Category = "RC")
-	FRemoteControlTestStructInnerSimle InnerSimle;
+	FRemoteControlTestStructInnerSimple InnerSimple;
 
 	UPROPERTY(EditAnywhere, Category = "RC")
 	int32 Int32Value = -791;
@@ -164,4 +164,10 @@ public:
 
 	UPROPERTY(EditAnywhere, Category = "RC")
 	FRotator RotatorValue = {45.1f, 96.2f, 184.5f};
+
+	UPROPERTY(EditAnywhere, Category = "RC")
+	FColor ColorValue = {40, 150, 200, 50};
+
+	UPROPERTY(EditAnywhere, Category = "RC")
+	FLinearColor LinearColorValue = {0.7f, 0.8f, 0.22f, 0.652f};
 };
diff --git a/Engine/Plugins/VirtualProduction/RemoteControl/Source/RemoteControl/Private/Tests/RemoteControlPropertyHandleTests.cpp b/Engine/Plugins/VirtualProduction/RemoteControl/Source/RemoteControl/Private/Tests/RemoteControlPropertyHandleTests.cpp
--- a/Engine/Plugins/VirtualProduction/RemoteControl/Source/RemoteControl/Private/Tests/RemoteControlPropertyHandleTests.cpp
+++ b/Engine/Plugins/VirtualProduction/RemoteControl/Source/RemoteControl/Private/Tests/RemoteControlPropertyHandleTests.cpp
@@ -9,6 +9,8 @@
 #include "RemoteControlPropertyHandleTestData.h"
 #include "RemoteControlPreset.h"
 
+PRAGMA_DISABLE_OPTIMIZATION
+
 #define PROP_NAME(Class, Name) GET_MEMBER_NAME_CHECKED(Class, Name)
 #define GET_TEST_PROP(PropName) URemoteControlAPITestObject::StaticClass()->FindPropertyByName(PROP_NAME(URemoteControlAPITestObject, PropName))
 
@@ -676,7 +678,7 @@ namespace RemoteControlAPIIntegrationTest
 		const TSharedPtr<FRemoteControlProperty> RCRotatorProp = Preset->ExposeProperty(TestObject.Get(), FRCFieldPathInfo{ RotatorProperty->GetName() }).Pin();
 		Test.TestNotNull(TEXT("The exposed property must be valid."), RCRotatorProp.Get());
 
-		// Test FText property
+		// Test FRotator property
 		const FRotator RotatorPropertyTestValue(-1, -2, -3);
 		const FRotator RotatorPropertyDefaultValue(4, 5, 6);
 		TestObject->RotatorValue = RotatorPropertyDefaultValue;
@@ -693,6 +695,62 @@ namespace RemoteControlAPIIntegrationTest
 		Test.TestTrue(TEXT("Value in UObject should be the same as a value from Property Handle."), RotatorPropertyTestValue.Equals(TestObject->RotatorValue));
 	}
 
+	void TestColorProperty(FAutomationTestBase& Test)
+	{
+		// Setup test data
+		TStrongObjectPtr<URemoteControlPreset> Preset{ NewObject<URemoteControlPreset>() };
+		TStrongObjectPtr<URemoteControlAPITestObject> TestObject{ NewObject<URemoteControlAPITestObject>() };
+
+		const FProperty* ColorProperty = GET_TEST_PROP(ColorValue);
+
+		const TSharedPtr<FRemoteControlProperty> RCColorProp = Preset->ExposeProperty(TestObject.Get(), FRCFieldPathInfo{ ColorProperty->GetName() }).Pin();
+		Test.TestNotNull(TEXT("The exposed property must be valid."), RCColorProp.Get());
+
+		// Test FColor property
+		const FColor ColorPropertyTestValue(8, 140, 20);
+		const FColor ColorPropertyDefaultValue(4, 5, 6);
+		TestObject->ColorValue = ColorPropertyDefaultValue;
+
+		// Get API property Text property handle
+		const TSharedPtr<IRemoteControlPropertyHandle> ColorHandle = RCColorProp->GetPropertyHandle();
+		Test.TestNotNull(TEXT("The exposed property must be valid."), ColorHandle.Get());
+
+		FColor ColorValueFromAPI;
+		ColorHandle->GetValue(ColorValueFromAPI);
+		Test.TestTrue(TEXT("Value in UObject should be the same as a value from Property Handle."), ColorValueFromAPI == TestObject->ColorValue);
+
+		ColorHandle->SetValue(ColorPropertyTestValue);
+		Test.TestTrue(TEXT("Value in UObject should be the same as a value from Property Handle."), ColorPropertyTestValue == TestObject->ColorValue);
+	}
+
+	void TestLinearColorProperty(FAutomationTestBase& Test)
+	{
+		// Setup test data
+		TStrongObjectPtr<URemoteControlPreset> Preset{ NewObject<URemoteControlPreset>() };
+		TStrongObjectPtr<URemoteControlAPITestObject> TestObject{ NewObject<URemoteControlAPITestObject>() };
+
+		const FProperty* LinearColorProperty = GET_TEST_PROP(LinearColorValue);
+
+		const TSharedPtr<FRemoteControlProperty> RCLinearColorProp = Preset->ExposeProperty(TestObject.Get(), FRCFieldPathInfo{ LinearColorProperty->GetName() }).Pin();
+		Test.TestNotNull(TEXT("The exposed property must be valid."), RCLinearColorProp.Get());
+
+		// Test FLinearColor property
+		const FLinearColor LinearColorPropertyTestValue(0.84f, 0.12f, 0.2f);
+		const FLinearColor LinearColorPropertyDefaultValue(0.42f, 0.5f, 0.6f);
+		TestObject->LinearColorValue = LinearColorPropertyDefaultValue;
+
+		// Get API property Text property handle
+		const TSharedPtr<IRemoteControlPropertyHandle> LinearColorHandle = RCLinearColorProp->GetPropertyHandle();
+		Test.TestNotNull(TEXT("The exposed property must be valid."), LinearColorHandle.Get());
+
+		FLinearColor LinearColorValueFromAPI;
+		LinearColorHandle->GetValue(LinearColorValueFromAPI);
+		Test.TestTrue(TEXT("Value in UObject should be the same as a value from Property Handle."), LinearColorValueFromAPI == TestObject->LinearColorValue);
+
+		LinearColorHandle->SetValue(LinearColorPropertyTestValue);
+		Test.TestTrue(TEXT("Value in UObject should be the same as a value from Property Handle."), LinearColorPropertyTestValue == TestObject->LinearColorValue);
+	}
+
 	void TestComplexPath(FAutomationTestBase& Test)
 	{
 		// Create preset and uobject
@@ -705,7 +763,7 @@ namespace RemoteControlAPIIntegrationTest
 		constexpr int32 StructOuterMapIndex = 468;
 
 		FRemoteControlTestStructInner RemoteControlTestStructInner;
-		RemoteControlTestStructInner.InnerSimle.Int32Value = Int32ValueTest;
+		RemoteControlTestStructInner.InnerSimple.Int32Value = Int32ValueTest;
 
 		FRemoteControlTestStructOuter RemoteControlTestStructOuter;
 		RemoteControlTestStructOuter.StructInnerSet.Add(RemoteControlTestStructInner);
@@ -716,7 +774,7 @@ namespace RemoteControlAPIIntegrationTest
 		const FProperty* StructOuterMapProp = GET_TEST_PROP(StructOuterMap);
 		const FGuid StructOuterMapPropId = Preset->ExposeProperty(TestObject.Get(), FRCFieldPathInfo{ StructOuterMapProp->GetName() }).Pin()->GetId();
 
-		// Test complex path hanles
+		// Test complex path handles
 		{
 			// 1. Find the Remote Control preset
 			URemoteControlPreset* ResolvedPreset = IRemoteControlModule::Get().ResolvePreset(PresetName);
@@ -769,21 +827,21 @@ namespace RemoteControlAPIIntegrationTest
 				return;
 			}
 
-			const TSharedPtr<IRemoteControlPropertyHandle> StructInnerropertyHandle = StructInnerSetPropertyHandleAsSet->FindElement(&RemoteControlTestStructInner);
-			if (!StructInnerropertyHandle.IsValid())
+			const TSharedPtr<IRemoteControlPropertyHandle> StructInnerPropertyHandle = StructInnerSetPropertyHandleAsSet->FindElement(&RemoteControlTestStructInner);
+			if (!StructInnerPropertyHandle.IsValid())
 			{
 				Test.AddError(TEXT("StructInnerropertyHandle not valid"));
 				return;
 			}
 
-			const TSharedPtr<IRemoteControlPropertyHandle> StructInnerSimleHandle = StructInnerropertyHandle->GetChildHandle(1);
-			if (!StructInnerSimleHandle.IsValid())
+			const TSharedPtr<IRemoteControlPropertyHandle> StructInnerSimpleHandle = StructInnerPropertyHandle->GetChildHandle(1);
+			if (!StructInnerSimpleHandle.IsValid())
 			{
-				Test.AddError(TEXT("StructInnerSimleHandle not valid"));
+				Test.AddError(TEXT("StructInnerSimpleHandle not valid"));
 				return;
 			}
 
-			const TSharedPtr<IRemoteControlPropertyHandle> Int32ValueHandle =  StructInnerSimleHandle->GetChildHandle(0);
+			const TSharedPtr<IRemoteControlPropertyHandle> Int32ValueHandle =  StructInnerSimpleHandle->GetChildHandle(0);
 			if (!Int32ValueHandle.IsValid())
 			{
 				Test.AddError(TEXT("Int32ValueHandle not valid"));
@@ -792,7 +850,7 @@ namespace RemoteControlAPIIntegrationTest
 
 			int32 GetInt32Value = 0;
 			Int32ValueHandle->GetValue(GetInt32Value);
-			Test.TestEqual(TEXT("Value in UObject should be the same as a value from Property Handle."), RemoteControlTestStructInner.InnerSimle.Int32Value, GetInt32Value);
+			Test.TestEqual(TEXT("Value in UObject should be the same as a value from Property Handle."), RemoteControlTestStructInner.InnerSimple.Int32Value, GetInt32Value);
 
 			int32 SetInt32Value = 0;
 			Int32ValueHandle->SetValue(SetInt32Value);
@@ -800,7 +858,7 @@ namespace RemoteControlAPIIntegrationTest
 				TestObject
 					->StructOuterMap.Find(StructOuterMapIndex)
 					->StructInnerSet.FindByHash(GetTypeHash(RemoteControlTestStructInner), RemoteControlTestStructInner)
-					->InnerSimle.Int32Value,
+					->InnerSimple.Int32Value,
 				SetInt32Value);
 		}
 	}
@@ -809,14 +867,18 @@ namespace RemoteControlAPIIntegrationTest
 	{
 		// Create preset and uobject
 		const TCHAR* PresetName = TEXT("TestComplexPath");
+
 		TStrongObjectPtr<URemoteControlPreset> Preset{ NewObject<URemoteControlPreset>(GetTransientPackage(), PresetName) };
+		check(Preset);
+		
 		TStrongObjectPtr<URemoteControlAPITestObject> TestObject{ NewObject<URemoteControlAPITestObject>() };
+		check(TestObject);
 
 		// Set a test value
 		constexpr int32 Int32ValueTest = 14923;
 
 		FRemoteControlTestStructInner RemoteControlTestStructInner;
-		RemoteControlTestStructInner.InnerSimle.Int32Value = Int32ValueTest;
+		RemoteControlTestStructInner.InnerSimple.Int32Value = Int32ValueTest;
 
 		FRemoteControlTestStructOuter RemoteControlTestStructOuter;
 		RemoteControlTestStructOuter.StructInnerSet.Add(RemoteControlTestStructInner);
@@ -833,30 +895,32 @@ namespace RemoteControlAPIIntegrationTest
 		const FName IntMapLabel = Preset->ExposeProperty(TestObject.Get(), FRCFieldPathInfo{ IntMapProp->GetName() }).Pin()->GetLabel();
 
 		const TSharedPtr<IRemoteControlPropertyHandle> StructOuterMapPropertyHandle = IRemoteControlPropertyHandle::GetPropertyHandle(PresetName, StructOuterMapPropId);
-		Test.TestTrue(TEXT("Property handle is valid"), StructOuterMapPropertyHandle.IsValid());
+		if(!Test.TestValid(TEXT("Property handle is valid"), StructOuterMapPropertyHandle))
+		{
+			return;
+		}
 
 		Test.TestFalse(TEXT("Child property handle should nullptr if the path same as a parent path"), StructOuterMapPropertyHandle->GetChildHandleByFieldPath(TEXT("StructOuterMap")).IsValid());
 		Test.TestFalse(TEXT("Child property handle should nullptr if the field path is wrong"), StructOuterMapPropertyHandle->GetChildHandleByFieldPath(TEXT("StructOuterMap_Wrong_Path")).IsValid());
 
 		const TSharedPtr<IRemoteControlPropertyHandleMap> StructOuterMapPropertyHandleAsMap = StructOuterMapPropertyHandle->AsMap();
-		if (!StructOuterMapPropertyHandleAsMap.IsValid())
+		if(!Test.TestValid(TEXT("StructOuterMapPropertyHandleAsMap is valid"), StructOuterMapPropertyHandleAsMap))
 		{
-			Test.AddError(TEXT("StructOuterMapPropertyHandleAsMap not valid"));
 			return;
 		}
 
 		const TSharedPtr<IRemoteControlPropertyHandle> StructOuterPropertyHandle_0 = StructOuterMapPropertyHandle->GetChildHandleByFieldPath(TEXT("StructOuterMap[1]"));
-		Test.TestTrue(TEXT("Property handle is valid"), StructOuterPropertyHandle_0.IsValid());
+		Test.TestValid(TEXT("Property handle is valid"), StructOuterPropertyHandle_0);
 		Test.TestNotNull(TEXT("Property handle should be valid"), CastField<FMapProperty>(StructOuterPropertyHandle_0->GetParentProperty()));
 		Test.TestTrue(TEXT("Index should be 1"), StructOuterPropertyHandle_0->GetIndexInArray() == 1);
 
-		const TSharedPtr<IRemoteControlPropertyHandle> IntPropertyHandle = StructOuterMapPropertyHandle->GetChildHandleByFieldPath(TEXT("StructOuterMap[1].StructInnerSet[0].InnerSimle.Int32Value"));
-		Test.TestTrue(TEXT("Property handle is valid"), IntPropertyHandle.IsValid());
-
+		const TSharedPtr<IRemoteControlPropertyHandle> IntPropertyHandle = StructOuterMapPropertyHandle->GetChildHandleByFieldPath(TEXT("StructOuterMap[1].StructInnerSet[0].InnerSimple.Int32Value"));
+		Test.TestValid(TEXT("Property handle is valid"), IntPropertyHandle);
+		
 		{
 			int32 GetInt32Value = 0;
 			IntPropertyHandle->GetValue(GetInt32Value);
-			Test.TestEqual(TEXT("Value in UObject should be the same as a value from Property Handle."), RemoteControlTestStructInner.InnerSimle.Int32Value, GetInt32Value);
+			Test.TestEqual(TEXT("Value in UObject should be the same as a value from Property Handle."), RemoteControlTestStructInner.InnerSimple.Int32Value, GetInt32Value);
 
 		}
 
@@ -933,6 +997,8 @@ bool FRemoteControlAPIIntegrationTest::RunTest(const FString& Parameters)
 	RemoteControlAPIIntegrationTest::TestTextProperty(*this);
 	RemoteControlAPIIntegrationTest::TestVectorProperty(*this);
 	RemoteControlAPIIntegrationTest::TestRotatorProperty(*this);
+	RemoteControlAPIIntegrationTest::TestColorProperty(*this);
+	RemoteControlAPIIntegrationTest::TestLinearColorProperty(*this);
 	RemoteControlAPIIntegrationTest::TestComplexPath(*this);
 	RemoteControlAPIIntegrationTest::TestGetPropertyHandleByFieldPath(*this);
 
@@ -941,3 +1007,5 @@ bool FRemoteControlAPIIntegrationTest::RunTest(const FString& Parameters)
 
 #undef GET_TEST_PROP
 #undef PROP_NAME
+
+PRAGMA_ENABLE_OPTIMIZATION
diff --git a/Engine/Plugins/VirtualProduction/RemoteControl/Source/RemoteControl/Public/IRemoteControlPropertyHandle.h b/Engine/Plugins/VirtualProduction/RemoteControl/Source/RemoteControl/Public/IRemoteControlPropertyHandle.h
--- a/Engine/Plugins/VirtualProduction/RemoteControl/Source/RemoteControl/Public/IRemoteControlPropertyHandle.h
+++ b/Engine/Plugins/VirtualProduction/RemoteControl/Source/RemoteControl/Public/IRemoteControlPropertyHandle.h
@@ -67,7 +67,9 @@ public:
 	virtual bool GetValue(FVector4& OutValue) const = 0;
 	virtual bool GetValue(FQuat& OutValue) const = 0;
 	virtual bool GetValue(FRotator& OutValue) const = 0;
-
+	virtual bool GetValue(FColor& OutValue) const = 0;
+	virtual bool GetValue(FLinearColor& OutValue) const = 0;
+	
 	/**
 	 * Sets the typed value of a property.
 	 * If the property does not support the value type false is returned
@@ -95,6 +97,8 @@ public:
 	virtual bool SetValue(FVector4 InValue) = 0;
 	virtual bool SetValue(FQuat InValue) = 0;
 	virtual bool SetValue(FRotator InValue) = 0;
+	virtual bool SetValue(FColor InValue) = 0;
+	virtual bool SetValue(FLinearColor InValue) = 0;
 	virtual bool SetValue(const TCHAR* InValue) = 0;
 
 	/**
