Epic Stage App connect screen using UDP multicast, and beacon receiver on engine side to reply to the beacon messages

diff --git a/Engine/Plugins/VirtualProduction/RemoteControl/Source/WebRemoteControl/Private/StageAppBeaconReceiver.cpp b/Engine/Plugins/VirtualProduction/RemoteControl/Source/WebRemoteControl/Private/StageAppBeaconReceiver.cpp
new file mode 100644
--- /dev/null
+++ b/Engine/Plugins/VirtualProduction/RemoteControl/Source/WebRemoteControl/Private/StageAppBeaconReceiver.cpp
@@ -0,0 +1,164 @@
+// Copyright Epic Games, Inc. All Rights Reserved.
+
+#include "StageAppBeaconReceiver.h"
+
+#include "Common/UdpSocketBuilder.h"
+#include "IRemoteControlModule.h"
+#include "MessageEndpoint.h"
+#include "MessageEndpointBuilder.h"
+#include "RemoteControlSettings.h"
+#include "Serialization/ArrayReader.h"
+#include "Serialization/ArrayWriter.h"
+#include "StageAppSettings.h"
+
+namespace StageAppBeaconReceiverConstants
+{
+	const uint8 ProtocolVersion = 0;
+}
+
+FStageAppBeaconReceiver::FStageAppBeaconReceiver()
+	: Guid(FGuid::NewGuid())
+{
+}
+
+void FStageAppBeaconReceiver::Startup()
+{
+	const UStageAppSettings& Settings = *GetDefault<UStageAppSettings>();
+	FIPv4Address DiscoveryAddress;
+	if (!FIPv4Address::Parse(Settings.DiscoveryEndpoint, DiscoveryAddress))
+	{
+		UE_LOG(LogRemoteControl, Error, TEXT("Failed to parse Stage App discovery endpoint address \"%s\""), *Settings.DiscoveryEndpoint);
+		return;
+	}
+
+	Socket = FUdpSocketBuilder(TEXT("StageAppBeaconResponder"))
+		.AsNonBlocking()
+		.AsReusable()
+		.BoundToPort(Settings.DiscoveryPort)
+		.JoinedToGroup(DiscoveryAddress)
+		.WithMulticastLoopback()
+		.WithReceiveBufferSize(256);
+
+	WebsocketPort = GetDefault<URemoteControlSettings>()->RemoteControlWebSocketServerPort;
+
+	if (!Socket)
+	{
+		UE_LOG(LogRemoteControl, Warning, TEXT("StageAppBeaconResponder failed to create multicast socket"));
+		return;
+	}
+
+	Thread.Reset(FRunnableThread::Create(this, TEXT("StageAppBeaconResponderThread"), 0, TPri_Normal, FPlatformAffinity::GetPoolThreadMask()));
+}
+
+void FStageAppBeaconReceiver::Shutdown()
+{
+	if (Thread)
+	{
+		Thread->Kill();
+		Thread = nullptr;
+	}
+
+	if (Socket)
+	{
+		ISocketSubsystem::Get()->DestroySocket(Socket);
+	}
+}
+
+bool FStageAppBeaconReceiver::Init()
+{
+	return true;
+}
+
+uint32 FStageAppBeaconReceiver::Run()
+{
+	const UStageAppSettings& Settings = *GetDefault<UStageAppSettings>();
+	const FTimespan WaitTime = FTimespan::FromSeconds(Settings.DiscoverySocketWaitTime);
+
+	while (!bStopping)
+	{
+		if (!Socket)
+		{
+			return 0;
+		}
+
+		Socket->Wait(ESocketWaitConditions::WaitForRead, WaitTime);
+		ReceiveBeaconMessages();
+	}
+
+	return 0;
+}
+
+void FStageAppBeaconReceiver::Stop()
+{
+	bStopping = true;
+}
+
+void FStageAppBeaconReceiver::ReceiveBeaconMessages()
+{
+	if (!Socket)
+	{
+		return;
+	}
+
+	uint32 PendingDataSize;
+	if (Socket->HasPendingData(PendingDataSize) && PendingDataSize > 0)
+	{
+		FArrayReader MessageData = FArrayReader(true);
+		MessageData.SetNumUninitialized(PendingDataSize);
+
+		int32 NumRead;
+		TSharedRef<FInternetAddr> Source = ISocketSubsystem::Get()->CreateInternetAddr();
+		if (Socket->RecvFrom(MessageData.GetData(), MessageData.Num(), NumRead, *Source) && NumRead > 0)
+		{
+			HandleBeaconMessage(MessageData, Source);
+		}
+	}
+}
+
+void FStageAppBeaconReceiver::HandleBeaconMessage(FArrayReader& MessageData, TSharedRef<FInternetAddr> Source)
+{
+	if (MessageData.Num() != 1)
+	{
+		return;
+	}
+
+	// We don't do anything with this yet, but could use it to ignore beacons from apps with incompatible protocols
+	uint8 AppProtocolVersion;
+	MessageData << AppProtocolVersion;
+
+	const FString FriendlyName = GetFriendlyName();
+
+	FArrayWriter Writer(true);
+
+#if !PLATFORM_LITTLE_ENDIAN
+	// App expects little-endian byte order
+	Writer.SetByteSwapping(true);
+#endif
+
+	{
+		Writer << (uint8&)StageAppBeaconReceiverConstants::ProtocolVersion;
+		Writer << Guid;
+		Writer << (uint32&)WebsocketPort;
+		Writer << (FString&)FriendlyName;
+	}
+
+	int32 Sent;
+	Socket->SendTo(Writer.GetData(), Writer.Num(), Sent, *Source);
+}
+
+FString FStageAppBeaconReceiver::GetFriendlyName() const
+{
+	FString FriendlyName;
+
+	if (FParse::Value(FCommandLine::Get(), TEXT("-StageFriendlyName="), FriendlyName))
+	{
+		return FriendlyName;
+	}
+
+	if (FParse::Value(FCommandLine::Get(), TEXT("-CONCERTDISPLAYNAME="), FriendlyName))
+	{
+		return FriendlyName;
+	}
+
+	return FApp::GetSessionOwner();
+}
diff --git a/Engine/Plugins/VirtualProduction/RemoteControl/Source/WebRemoteControl/Private/StageAppBeaconReceiver.h b/Engine/Plugins/VirtualProduction/RemoteControl/Source/WebRemoteControl/Private/StageAppBeaconReceiver.h
new file mode 100644
--- /dev/null
+++ b/Engine/Plugins/VirtualProduction/RemoteControl/Source/WebRemoteControl/Private/StageAppBeaconReceiver.h
@@ -0,0 +1,87 @@
+// Copyright Epic Games, Inc. All Rights Reserved.
+
+#pragma once
+
+#include "CoreMinimal.h"
+#include "StageAppBeaconReceiver.generated.h"
+
+class FArrayReader;
+class FInternetAddr;
+class FSocket;
+
+USTRUCT()
+struct FRemoteControl_StageAppBeaconMessage
+{
+	GENERATED_BODY()
+};
+
+USTRUCT()
+struct FRemoteControl_StageAppConnectionInfoMessage
+{
+	GENERATED_BODY()
+
+	/** Version of the stage app beacon protocol to use. */
+	uint8 Protocol;
+
+	/** ID of the sender. */
+	FGuid SenderId;
+
+	/** The websocket port the app should connect to. */
+	uint32 WebsocketPort;
+
+	/** The websocket address the app should connect to. */
+	FString WebsocketAddress;
+
+	/** The user-friendly name of the connection to show to the user. */
+	FString FriendlyName;
+};
+
+/**
+ * Receives beacon messages from the Unreal Stage App and replies with connection information.
+ * This allows the app to detect compatible Unreal instances on the local network and list them for the user.
+ */
+class FStageAppBeaconReceiver
+	: public FRunnable
+{
+public:
+	/** Creates and initializes a new beacon receiver. */
+	FStageAppBeaconReceiver();
+
+	/** Open a socket and start a thread listening for beacon messages. */
+	void Startup();
+
+	/** Close the socket and kill the listening thread. */
+	void Shutdown();
+
+protected:
+	//~ Begin FRunnable implementation
+	virtual bool Init() override;
+	virtual uint32 Run() override;
+	virtual void Stop() override;
+	//~ End FRunnable implementation
+
+private:
+	/** Called regularly to receive beacon messages. */
+	void ReceiveBeaconMessages();
+
+	/** Handle a message received on the beacon endpoint. */
+	void HandleBeaconMessage(FArrayReader& MessageData, TSharedRef<FInternetAddr> Source);
+
+	/** Get the name to report to apps searching for the engine. */
+	FString GetFriendlyName() const;
+
+	/** Identifier for this engine instance. Used by the app to differentiate instances even if they're on the same machine. */
+	FGuid Guid;
+
+	/** The websocket port to reply with. */
+	uint32 WebsocketPort;
+
+	/** Socket used to listen for and reply to beacon messages from the app. */
+	FSocket* Socket;
+
+	/** Flag indicating that the thread is stopping. */
+	TAtomic<bool> bStopping;
+
+	/** The thread to receive beacon messages on. */
+	TUniquePtr<FRunnableThread> Thread = nullptr;
+};
\ No newline at end of file
diff --git a/Engine/Plugins/VirtualProduction/RemoteControl/Source/WebRemoteControl/Private/WebRemoteControl.cpp b/Engine/Plugins/VirtualProduction/RemoteControl/Source/WebRemoteControl/Private/WebRemoteControl.cpp
--- a/Engine/Plugins/VirtualProduction/RemoteControl/Source/WebRemoteControl/Private/WebRemoteControl.cpp
+++ b/Engine/Plugins/VirtualProduction/RemoteControl/Source/WebRemoteControl/Private/WebRemoteControl.cpp
@@ -277,6 +277,8 @@ void FWebRemoteControlModule::StartupModule()
 	{
 		StartWebSocketServer();
 	}
+
+	StageAppBeaconReceiver.Startup();
 }
 
 void FWebRemoteControlModule::ShutdownModule()
@@ -291,12 +293,14 @@ void FWebRemoteControlModule::ShutdownModule()
 	{
 		return;
 	}
-	
+
+	StageAppBeaconReceiver.Shutdown();
 	EditorRoutes.UnregisterRoutes(this);
 	WebSocketHandler->UnregisterRoutes(this);
 	StopHttpServer();
 	StopWebSocketServer();
 	UnregisterConsoleCommands();
+
 #if WITH_EDITOR
 	UnregisterSettings();
 #endif
diff --git a/Engine/Plugins/VirtualProduction/RemoteControl/Source/WebRemoteControl/Private/WebRemoteControl.h b/Engine/Plugins/VirtualProduction/RemoteControl/Source/WebRemoteControl/Private/WebRemoteControl.h
--- a/Engine/Plugins/VirtualProduction/RemoteControl/Source/WebRemoteControl/Private/WebRemoteControl.h
+++ b/Engine/Plugins/VirtualProduction/RemoteControl/Source/WebRemoteControl/Private/WebRemoteControl.h
@@ -12,6 +12,7 @@
 #include "RemoteControlRequest.h"
 #include "RemoteControlRoute.h"
 #include "RemoteControlWebSocketServer.h"
+#include "StageAppBeaconReceiver.h"
 #include "WebRemoteControlEditorRoutes.h"
 
 struct FHttpServerRequest;
@@ -196,4 +197,7 @@ private:
 	FSimpleMulticastDelegate OnHttpServerStoppedDelegate;
 	FOnWebServerStarted OnWebSocketServerStartedDelegate;
 	FSimpleMulticastDelegate OnWebSocketServerStoppedDelegate;
+
+	/** Receiver that responds to beacon messages from the stage app. */
+	FStageAppBeaconReceiver StageAppBeaconReceiver;
 };
\ No newline at end of file
diff --git a/Engine/Plugins/VirtualProduction/RemoteControl/Source/WebRemoteControl/Public/StageAppSettings.h b/Engine/Plugins/VirtualProduction/RemoteControl/Source/WebRemoteControl/Public/StageAppSettings.h
new file mode 100644
--- /dev/null
+++ b/Engine/Plugins/VirtualProduction/RemoteControl/Source/WebRemoteControl/Public/StageAppSettings.h
@@ -0,0 +1,57 @@
+// Copyright Epic Games, Inc. All Rights Reserved.
+
+#pragma once
+
+#include "CoreMinimal.h"
+#include "Engine/DeveloperSettings.h"
+#include "StageAppSettings.generated.h"
+
+/**
+ * Global settings for Epic Stage App integration.
+ */
+UCLASS(config = Engine)
+class WEBREMOTECONTROL_API UStageAppSettings : public UDeveloperSettings
+{
+	GENERATED_BODY()
+
+public:
+	//~ Begin UDeveloperSettings implementation
+	virtual FName GetContainerName() const
+	{
+		return "Project";
+	}
+
+	virtual FName GetCategoryName() const
+	{
+		return "Plugins";
+	}
+
+	virtual FName GetSectionName() const
+	{
+		return "Epic Stage App";
+	}
+
+	virtual FText GetSectionText() const
+	{
+		return NSLOCTEXT("StageAppSettings", "StageAppSettingsSection", "Epic Stage App");
+	}
+	//~ End UDeveloperSettings implementation
+
+	/**
+	 * The IP endpoint to listen to for app discovery messages.
+	 */
+	UPROPERTY(config, EditAnywhere, Category = "Epic Stage App")
+	FString DiscoveryEndpoint = "230.0.0.2";
+
+	/**
+	 * The port to listen to for app discovery messages.
+	 */
+	UPROPERTY(config, EditAnywhere, Category = "Epic Stage App")
+	int32 DiscoveryPort = 6667;
+
+	/**
+	 * How many seconds to wait for app discovery messages.
+	 */
+	UPROPERTY(config, EditAnywhere, Category = "Epic Stage App")
+	float DiscoverySocketWaitTime = 10.f;
+};
diff --git a/Engine/Plugins/VirtualProduction/RemoteControl/Source/WebRemoteControl/WebRemoteControl.Build.cs b/Engine/Plugins/VirtualProduction/RemoteControl/Source/WebRemoteControl/WebRemoteControl.Build.cs
--- a/Engine/Plugins/VirtualProduction/RemoteControl/Source/WebRemoteControl/WebRemoteControl.Build.cs
+++ b/Engine/Plugins/VirtualProduction/RemoteControl/Source/WebRemoteControl/WebRemoteControl.Build.cs
@@ -17,6 +17,7 @@ public class WebRemoteControl : ModuleRules
         PrivateDependencyModuleNames.AddRange(
 			new string[] {
 				"AssetRegistry",
+				"Networking",
 				"RemoteControl",
 				"RemoteControlCommon",
 				"Serialization",
