Fix compilation warnings/errors in non-unity build for StageAppBeaconReceiver

diff --git a/Engine/Plugins/VirtualProduction/RemoteControl/Source/WebRemoteControl/Private/StageAppBeaconReceiver.cpp b/Engine/Plugins/VirtualProduction/RemoteControl/Source/WebRemoteControl/Private/StageAppBeaconReceiver.cpp
--- a/Engine/Plugins/VirtualProduction/RemoteControl/Source/WebRemoteControl/Private/StageAppBeaconReceiver.cpp
+++ b/Engine/Plugins/VirtualProduction/RemoteControl/Source/WebRemoteControl/Private/StageAppBeaconReceiver.cpp
@@ -6,6 +6,8 @@
 #include "IRemoteControlModule.h"
 #include "MessageEndpoint.h"
 #include "MessageEndpointBuilder.h"
+#include "Misc/App.h"
+#include "Misc/CommandLine.h"
 #include "RemoteControlSettings.h"
 #include "Serialization/ArrayReader.h"
 #include "Serialization/ArrayWriter.h"
diff --git a/Engine/Plugins/VirtualProduction/RemoteControl/Source/WebRemoteControl/Private/StageAppBeaconReceiver.h b/Engine/Plugins/VirtualProduction/RemoteControl/Source/WebRemoteControl/Private/StageAppBeaconReceiver.h
--- a/Engine/Plugins/VirtualProduction/RemoteControl/Source/WebRemoteControl/Private/StageAppBeaconReceiver.h
+++ b/Engine/Plugins/VirtualProduction/RemoteControl/Source/WebRemoteControl/Private/StageAppBeaconReceiver.h
@@ -3,6 +3,8 @@
 #pragma once
 
 #include "CoreMinimal.h"
+#include "HAL/Runnable.h"
+#include "HAL/RunnableThread.h"
 #include "StageAppBeaconReceiver.generated.h"
 
 class FArrayReader;
@@ -47,6 +49,9 @@ public:
 	/** Creates and initializes a new beacon receiver. */
 	FStageAppBeaconReceiver();
 
+	/** Destroy the beacon receiver. */
+	virtual ~FStageAppBeaconReceiver() {};
+
 	/** Open a socket and start a thread listening for beacon messages. */
 	void Startup();
 
