Expose WebRemoteControl websocket API to other modules

diff --git a/Engine/Plugins/VirtualProduction/RemoteControl/Source/WebRemoteControl/Private/RemoteControlRoute.h b/Engine/Plugins/VirtualProduction/RemoteControl/Source/WebRemoteControl/Private/RemoteControlRoute.h
--- a/Engine/Plugins/VirtualProduction/RemoteControl/Source/WebRemoteControl/Private/RemoteControlRoute.h
+++ b/Engine/Plugins/VirtualProduction/RemoteControl/Source/WebRemoteControl/Private/RemoteControlRoute.h
@@ -6,19 +6,11 @@
 #include "HttpServerResponse.h"
 #include "HttpServerRequest.h"
 #include "HttpPath.h"
+#include "RemoteControlWebsocketRoute.h"
 
 #include "RemoteControlRoute.generated.h"
 
-struct FRemoteControlWebSocketMessage
-{
-	FString MessageName;
-	int32 MessageId = -1;
-	FGuid ClientId;
-	TArrayView<uint8> RequestPayload;
-};
-
 DECLARE_DELEGATE_RetVal_TwoParams(bool, FRequestHandlerDelegate, const FHttpServerRequest&, const FHttpResultCallback&);
-DECLARE_DELEGATE_OneParam(FWebSocketMessageDelegate, const FRemoteControlWebSocketMessage& /** Message */);
 
 struct FRemoteControlRoute
 {
@@ -42,25 +34,6 @@ struct FRemoteControlRoute
 	friend bool operator==(const FRemoteControlRoute& LHS, const FRemoteControlRoute& RHS) { return LHS.Path == RHS.Path && LHS.Verb == RHS.Verb; }
 };
 
-struct FRemoteControlWebsocketRoute
-{
-	FRemoteControlWebsocketRoute(const FString& InRouteDescription, const FString& InMessageName, const FWebSocketMessageDelegate& InDelegate)
-		: RouteDescription(InRouteDescription)
-		, MessageName(InMessageName)
-		, Delegate(InDelegate)
-	{}
-	
-	/** A description of how the route should be used. */
-	FString RouteDescription;
-	/**  The message handled by this route. */
-	FString MessageName;
-	/** The handler called when the route is accessed. */
-	FWebSocketMessageDelegate Delegate;
-
-	friend uint32 GetTypeHash(const FRemoteControlWebsocketRoute& Route) { return GetTypeHash(Route.MessageName); }
-	friend bool operator==(const FRemoteControlWebsocketRoute& LHS, const FRemoteControlWebsocketRoute& RHS) { return LHS.MessageName == RHS.MessageName; }
-};
-
 UENUM()
 enum class ERemoteControlHttpVerbs : uint16
 {
diff --git a/Engine/Plugins/VirtualProduction/RemoteControl/Source/WebRemoteControl/Private/RemoteControlWebSocketServer.cpp b/Engine/Plugins/VirtualProduction/RemoteControl/Source/WebRemoteControl/Private/RemoteControlWebSocketServer.cpp
--- a/Engine/Plugins/VirtualProduction/RemoteControl/Source/WebRemoteControl/Private/RemoteControlWebSocketServer.cpp
+++ b/Engine/Plugins/VirtualProduction/RemoteControl/Source/WebRemoteControl/Private/RemoteControlWebSocketServer.cpp
@@ -2,13 +2,13 @@
 
 #include "RemoteControlWebSocketServer.h"
 #include "Containers/Ticker.h"
+#include "WebRemoteControlInternalUtils.h"
 #include "IPAddress.h"
 #include "IRemoteControlModule.h"
 #include "IWebSocketNetworkingModule.h"
 #include "RemoteControlRequest.h"
 #include "Sockets.h"
 #include "SocketSubsystem.h"
-#include "WebRemoteControlUtils.h"
 #include "WebSocketNetworkingDelegates.h"
 
 #define LOCTEXT_NAMESPACE "RCWebSocketServer"
@@ -21,7 +21,7 @@ namespace RemoteControlWebSocketServer
 	TOptional<FRemoteControlWebSocketMessage> ParseWebsocketMessage(TArrayView<uint8> InPayload)
 	{
 		FRCWebSocketRequest Request;
-		bool bSuccess = WebRemoteControlUtils::DeserializeRequestPayload(InPayload, nullptr, Request);
+		bool bSuccess = WebRemoteControlInternalUtils::DeserializeRequestPayload(InPayload, nullptr, Request);
 
 		FString ErrorText;
 		if (Request.MessageName.IsEmpty())
@@ -179,6 +179,7 @@ void FRCWebSocketServer::OnWebSocketClientConnected(INetworkingWebSocket* Socket
 		CloseCallback.BindRaw(this, &FRCWebSocketServer::OnSocketClose, Socket);
 		Socket->SetSocketClosedCallBack(CloseCallback);
 
+		OnConnectionOpened().Broadcast(Connection.Id);
 		Connections.Add(MoveTemp(Connection));
 	}
 }
diff --git a/Engine/Plugins/VirtualProduction/RemoteControl/Source/WebRemoteControl/Private/RemoteControlWebSocketServer.h b/Engine/Plugins/VirtualProduction/RemoteControl/Source/WebRemoteControl/Private/RemoteControlWebSocketServer.h
--- a/Engine/Plugins/VirtualProduction/RemoteControl/Source/WebRemoteControl/Private/RemoteControlWebSocketServer.h
+++ b/Engine/Plugins/VirtualProduction/RemoteControl/Source/WebRemoteControl/Private/RemoteControlWebSocketServer.h
@@ -5,6 +5,7 @@
 #include "WebSocketNetDriver.h"
 #include "UObject/StrongObjectPtr.h"
 #include "INetworkingWebSocket.h"
+#include "IWebRemoteControlModule.h"
 #include "IWebSocketServer.h"
 #include "RemoteControlRoute.h"
 #include "Containers/Ticker.h"
@@ -81,10 +82,12 @@ public:
 
 	/** Returns whether the server is currently listening for messages. */
 	bool IsRunning() const;
+
+	/** Callback when a socket is opened */
+	FOnWebSocketConnectionOpened& OnConnectionOpened() { return OnConnectionOpenedDelegate; }
 	
 	/** Callback when a socket is closed */
-	DECLARE_MULTICAST_DELEGATE_OneParam(FOnConnectionClosed, FGuid /*ClientId*/);
-	FOnConnectionClosed& OnConnectionClosed() { return OnConnectionClosedDelegate; }
+	FOnWebSocketConnectionClosed& OnConnectionClosed() { return OnConnectionClosedDelegate; }
 
 private:
 	bool Tick(float DeltaTime);
@@ -154,6 +157,9 @@ private:
 	/** Holds the router responsible for dispatching messages received by this server. */
 	TSharedPtr<FWebsocketMessageRouter> Router;
 
+	/** Delegate triggered when a connection is opened */
+	FOnWebSocketConnectionOpened OnConnectionOpenedDelegate;
+
 	/** Delegate triggered when a connection is closed */
-	FOnConnectionClosed OnConnectionClosedDelegate;
+	FOnWebSocketConnectionClosed OnConnectionClosedDelegate;
 };
diff --git a/Engine/Plugins/VirtualProduction/RemoteControl/Source/WebRemoteControl/Private/WebRemoteControl.cpp b/Engine/Plugins/VirtualProduction/RemoteControl/Source/WebRemoteControl/Private/WebRemoteControl.cpp
--- a/Engine/Plugins/VirtualProduction/RemoteControl/Source/WebRemoteControl/Private/WebRemoteControl.cpp
+++ b/Engine/Plugins/VirtualProduction/RemoteControl/Source/WebRemoteControl/Private/WebRemoteControl.cpp
@@ -10,7 +10,8 @@
 #include "RemoteControlRoute.h"
 #include "RemoteControlSettings.h"
 #include "RemoteControlPreset.h"
-#include "WebRemoteControlUtils.h"
+#include "RemoteControlWebsocketRoute.h"
+#include "WebRemoteControlInternalUtils.h"
 #include "WebRemoteControlExternalLogger.h"
 #include "WebSocketMessageHandler.h"
 
@@ -372,6 +373,11 @@ void FWebRemoteControlModule::UnregisterWebsocketRoute(const FRemoteControlWebso
 	WebSocketRouter->UnbindRoute(Route.MessageName);
 }
 
+void FWebRemoteControlModule::SendWebsocketMessage(const FGuid& InTargetClientId, const TArray<uint8>& InUTF8Payload)
+{
+	WebSocketServer.Send(InTargetClientId, InUTF8Payload);
+}
+
 void FWebRemoteControlModule::StartHttpServer()
 {
 	if (!HttpRouter)
@@ -701,7 +707,7 @@ void FWebRemoteControlModule::UnregisterConsoleCommands()
 
 bool FWebRemoteControlModule::HandleInfoRoute(const FHttpServerRequest& Request, const FHttpResultCallback& OnComplete)
 {
-	TUniquePtr<FHttpServerResponse> Response = WebRemoteControlUtils::CreateHttpResponse(EHttpServerResponseCodes::Ok);
+	TUniquePtr<FHttpServerResponse> Response = WebRemoteControlInternalUtils::CreateHttpResponse(EHttpServerResponseCodes::Ok);
 
 	
 	bool bInPackaged = false;
@@ -716,16 +722,16 @@ bool FWebRemoteControlModule::HandleInfoRoute(const FHttpServerRequest& Request,
 #endif
 
 	FAPIInfoResponse RCResponse{RegisteredHttpRoutes.Array(), bInPackaged, ActivePreset};
-	WebRemoteControlUtils::SerializeResponse(MoveTemp(RCResponse), Response->Body);
+	WebRemoteControlUtils::SerializeMessage(MoveTemp(RCResponse), Response->Body);
 	OnComplete(MoveTemp(Response));
 	return true;
 }
 
 bool FWebRemoteControlModule::HandleBatchRequest(const FHttpServerRequest& Request, const FHttpResultCallback& OnComplete)
 {
-	TUniquePtr<FHttpServerResponse> Response = WebRemoteControlUtils::CreateHttpResponse();
+	TUniquePtr<FHttpServerResponse> Response = WebRemoteControlInternalUtils::CreateHttpResponse();
 
-	if (!WebRemoteControlUtils::ValidateContentType(Request, TEXT("application/json"), OnComplete))
+	if (!WebRemoteControlInternalUtils::ValidateContentType(Request, TEXT("application/json"), OnComplete))
 	{
 		return true;
 	}
@@ -733,7 +739,7 @@ bool FWebRemoteControlModule::HandleBatchRequest(const FHttpServerRequest& Reque
 	Response->Code = EHttpServerResponseCodes::Ok;
 
 	FRCBatchRequest BatchRequest;
-	if (!WebRemoteControlUtils::DeserializeRequest(Request, &OnComplete, BatchRequest))
+	if (!WebRemoteControlInternalUtils::DeserializeRequest(Request, &OnComplete, BatchRequest))
 	{
 		return true;
 	}
@@ -762,16 +768,16 @@ bool FWebRemoteControlModule::HandleBatchRequest(const FHttpServerRequest& Reque
 
 bool FWebRemoteControlModule::HandleOptionsRoute(const FHttpServerRequest& Request, const FHttpResultCallback& OnComplete)
 {
-	OnComplete(WebRemoteControlUtils::CreateHttpResponse(EHttpServerResponseCodes::Ok));
+	OnComplete(WebRemoteControlInternalUtils::CreateHttpResponse(EHttpServerResponseCodes::Ok));
 	return true;
 }
 
 bool FWebRemoteControlModule::HandleObjectCallRoute(const FHttpServerRequest& Request, const FHttpResultCallback& OnComplete)
 {
-	TUniquePtr<FHttpServerResponse> Response = WebRemoteControlUtils::CreateHttpResponse();
+	TUniquePtr<FHttpServerResponse> Response = WebRemoteControlInternalUtils::CreateHttpResponse();
 
 	FString ErrorText;
-	if (!WebRemoteControlUtils::ValidateContentType(Request, TEXT("application/json"), OnComplete))
+	if (!WebRemoteControlInternalUtils::ValidateContentType(Request, TEXT("application/json"), OnComplete))
 	{
 		return true;
 	}
@@ -785,7 +791,7 @@ bool FWebRemoteControlModule::HandleObjectCallRoute(const FHttpServerRequest& Re
 	// if we haven't resolved the object or function, return not found
 	if (!Call.IsValid())
 	{
-		WebRemoteControlUtils::CreateUTF8ErrorMessage(TEXT("The object or function was not found."), Response->Body);
+		WebRemoteControlInternalUtils::CreateUTF8ErrorMessage(TEXT("The object or function was not found."), Response->Body);
 		Response->Code = EHttpServerResponseCodes::NotFound;
 		OnComplete(MoveTemp(Response));
 		return true;
@@ -811,9 +817,9 @@ bool FWebRemoteControlModule::HandleObjectCallRoute(const FHttpServerRequest& Re
 
 bool FWebRemoteControlModule::HandleObjectPropertyRoute(const FHttpServerRequest& Request, const FHttpResultCallback& OnComplete)
 {
-	TUniquePtr<FHttpServerResponse> Response = WebRemoteControlUtils::CreateHttpResponse();
+	TUniquePtr<FHttpServerResponse> Response = WebRemoteControlInternalUtils::CreateHttpResponse();
 
-	if (!WebRemoteControlUtils::ValidateContentType(Request, TEXT("application/json"), OnComplete))
+	if (!WebRemoteControlInternalUtils::ValidateContentType(Request, TEXT("application/json"), OnComplete))
 	{
 		return true;
 	}
@@ -832,7 +838,7 @@ bool FWebRemoteControlModule::HandleObjectPropertyRoute(const FHttpServerRequest
 	if (!ObjectRef.IsValid())
 	{
 		Response->Code = EHttpServerResponseCodes::NotFound;
-		WebRemoteControlUtils::CreateUTF8ErrorMessage(TEXT("Unable to find the object."), Response->Body);
+		WebRemoteControlInternalUtils::CreateUTF8ErrorMessage(TEXT("Unable to find the object."), Response->Body);
 		OnComplete(MoveTemp(Response));
 		return true;
 	}
@@ -888,7 +894,7 @@ bool FWebRemoteControlModule::HandleObjectPropertyRoute(const FHttpServerRequest
 
 bool FWebRemoteControlModule::HandlePresetCallFunctionRoute(const FHttpServerRequest& Request, const FHttpResultCallback& OnComplete)
 {
-	TUniquePtr<FHttpServerResponse> Response = WebRemoteControlUtils::CreateHttpResponse();
+	TUniquePtr<FHttpServerResponse> Response = WebRemoteControlInternalUtils::CreateHttpResponse();
 
 	FResolvePresetFieldArgs Args;
 	Args.PresetName = Request.PathParams.FindChecked(TEXT("preset"));
@@ -898,7 +904,7 @@ bool FWebRemoteControlModule::HandlePresetCallFunctionRoute(const FHttpServerReq
 	if (Preset == nullptr)
 	{
 		Response->Code = EHttpServerResponseCodes::NotFound;
-		WebRemoteControlUtils::CreateUTF8ErrorMessage(TEXT("Unable to resolve the preset."), Response->Body);
+		WebRemoteControlInternalUtils::CreateUTF8ErrorMessage(TEXT("Unable to resolve the preset."), Response->Body);
 		OnComplete(MoveTemp(Response));
 		return true;
 	}
@@ -908,13 +914,13 @@ bool FWebRemoteControlModule::HandlePresetCallFunctionRoute(const FHttpServerReq
 	if (!RCFunction || !RCFunction->GetFunction() || !RCFunction->FunctionArguments || !RCFunction->FunctionArguments->IsValid())
 	{
 		Response->Code = EHttpServerResponseCodes::NotFound;
-		WebRemoteControlUtils::CreateUTF8ErrorMessage(TEXT("Unable to resolve the preset field."), Response->Body);
+		WebRemoteControlInternalUtils::CreateUTF8ErrorMessage(TEXT("Unable to resolve the preset field."), Response->Body);
 		OnComplete(MoveTemp(Response));
 		return true;
 	}
 
 	FRCPresetCallRequest CallRequest;
-	if (!WebRemoteControlUtils::DeserializeRequest(Request, &OnComplete, CallRequest))
+	if (!WebRemoteControlInternalUtils::DeserializeRequest(Request, &OnComplete, CallRequest))
 	{
 		return true;
 	}
@@ -1017,7 +1023,7 @@ bool FWebRemoteControlModule::HandlePresetCallFunctionRoute(const FHttpServerReq
 	}
 	else
 	{
-		WebRemoteControlUtils::CreateUTF8ErrorMessage(FString::Printf(TEXT("Error while trying to call function %s."), *Args.FieldLabel), Response->Body);
+		WebRemoteControlInternalUtils::CreateUTF8ErrorMessage(FString::Printf(TEXT("Error while trying to call function %s."), *Args.FieldLabel), Response->Body);
 	}
 
 	OnComplete(MoveTemp(Response));
@@ -1026,15 +1032,15 @@ bool FWebRemoteControlModule::HandlePresetCallFunctionRoute(const FHttpServerReq
 
 bool FWebRemoteControlModule::HandlePresetSetPropertyRoute(const FHttpServerRequest& Request, const FHttpResultCallback& OnComplete)
 {
-	TUniquePtr<FHttpServerResponse> Response = WebRemoteControlUtils::CreateHttpResponse();
+	TUniquePtr<FHttpServerResponse> Response = WebRemoteControlInternalUtils::CreateHttpResponse();
 
-	if (!WebRemoteControlUtils::ValidateContentType(Request, TEXT("application/json"), OnComplete))
+	if (!WebRemoteControlInternalUtils::ValidateContentType(Request, TEXT("application/json"), OnComplete))
 	{
 		return true;
 	}
 
 	FRCPresetSetPropertyRequest SetPropertyRequest;
-	if (!WebRemoteControlUtils::DeserializeRequest(Request, &OnComplete, SetPropertyRequest))
+	if (!WebRemoteControlInternalUtils::DeserializeRequest(Request, &OnComplete, SetPropertyRequest))
 	{
 		return true;
 	}
@@ -1047,7 +1053,7 @@ bool FWebRemoteControlModule::HandlePresetSetPropertyRoute(const FHttpServerRequ
 	if (Preset == nullptr)
 	{
 		Response->Code = EHttpServerResponseCodes::NotFound;
-		WebRemoteControlUtils::CreateUTF8ErrorMessage(TEXT("Unable to resolve the preset."), Response->Body);
+		WebRemoteControlInternalUtils::CreateUTF8ErrorMessage(TEXT("Unable to resolve the preset."), Response->Body);
 		OnComplete(MoveTemp(Response));
 		return true;
 	}
@@ -1057,7 +1063,7 @@ bool FWebRemoteControlModule::HandlePresetSetPropertyRoute(const FHttpServerRequ
 	if (!RemoteControlProperty.IsValid())
 	{
 		Response->Code = EHttpServerResponseCodes::NotFound;
-		WebRemoteControlUtils::CreateUTF8ErrorMessage(TEXT("Unable to resolve the preset field."), Response->Body);
+		WebRemoteControlInternalUtils::CreateUTF8ErrorMessage(TEXT("Unable to resolve the preset field."), Response->Body);
 		OnComplete(MoveTemp(Response));
 		return true;
 	}
@@ -1113,7 +1119,7 @@ bool FWebRemoteControlModule::HandlePresetSetPropertyRoute(const FHttpServerRequ
 	}
 	else
 	{
-		WebRemoteControlUtils::CreateUTF8ErrorMessage(FString::Printf(TEXT("Error while trying to modify property %s."), *Args.FieldLabel), Response->Body);
+		WebRemoteControlInternalUtils::CreateUTF8ErrorMessage(FString::Printf(TEXT("Error while trying to modify property %s."), *Args.FieldLabel), Response->Body);
 	}
 
 	OnComplete(MoveTemp(Response));
@@ -1122,7 +1128,7 @@ bool FWebRemoteControlModule::HandlePresetSetPropertyRoute(const FHttpServerRequ
 
 bool FWebRemoteControlModule::HandlePresetGetPropertyRoute(const FHttpServerRequest& Request, const FHttpResultCallback& OnComplete)
 {
-	TUniquePtr<FHttpServerResponse> Response = WebRemoteControlUtils::CreateHttpResponse();
+	TUniquePtr<FHttpServerResponse> Response = WebRemoteControlInternalUtils::CreateHttpResponse();
 
 	FResolvePresetFieldArgs Args;
 	Args.PresetName = Request.PathParams.FindChecked(TEXT("preset"));
@@ -1132,7 +1138,7 @@ bool FWebRemoteControlModule::HandlePresetGetPropertyRoute(const FHttpServerRequ
 	if (Preset == nullptr)
 	{
 		Response->Code = EHttpServerResponseCodes::NotFound;
-		WebRemoteControlUtils::CreateUTF8ErrorMessage(TEXT("Unable to resolve the preset."), Response->Body);
+		WebRemoteControlInternalUtils::CreateUTF8ErrorMessage(TEXT("Unable to resolve the preset."), Response->Body);
 		OnComplete(MoveTemp(Response));
 		return true;
 	}
@@ -1142,7 +1148,7 @@ bool FWebRemoteControlModule::HandlePresetGetPropertyRoute(const FHttpServerRequ
 	if (!RemoteControlProperty)
 	{
 		Response->Code = EHttpServerResponseCodes::NotFound;
-		WebRemoteControlUtils::CreateUTF8ErrorMessage(TEXT("Unable to resolve the exposed entity."), Response->Body);
+		WebRemoteControlInternalUtils::CreateUTF8ErrorMessage(TEXT("Unable to resolve the exposed entity."), Response->Body);
 		OnComplete(MoveTemp(Response));
 		return true;
 	}
@@ -1150,7 +1156,7 @@ bool FWebRemoteControlModule::HandlePresetGetPropertyRoute(const FHttpServerRequ
 	if (!RemoteControlProperty->IsBound())
 	{
 		Response->Code = EHttpServerResponseCodes::NotFound;
-		WebRemoteControlUtils::CreateUTF8ErrorMessage(TEXT("Exposed entity was found but not bound to any object."), Response->Body);
+		WebRemoteControlInternalUtils::CreateUTF8ErrorMessage(TEXT("Exposed entity was found but not bound to any object."), Response->Body);
 		OnComplete(MoveTemp(Response));
 		return true;
 	}
@@ -1166,14 +1172,14 @@ bool FWebRemoteControlModule::HandlePresetGetPropertyRoute(const FHttpServerRequ
 	{
 		FStructOnScope PropertyValueOnScope = WebRemoteControlStructUtils::CreatePropertyValueOnScope(RemoteControlProperty, ObjectRef);
 		FStructOnScope FinalStruct = WebRemoteControlStructUtils::CreateGetPropertyOnScope(RemoteControlProperty, ObjectRef, MoveTemp(PropertyValueOnScope));
-		WebRemoteControlUtils::SerializeStructOnScope(FinalStruct, Writer);
+		WebRemoteControlInternalUtils::SerializeStructOnScope(FinalStruct, Writer);
 
 		Response->Code = EHttpServerResponseCodes::Ok;
 		WebRemoteControlUtils::ConvertToUTF8(WorkingBuffer, Response->Body);
 	}
 	else
 	{
-		WebRemoteControlUtils::CreateUTF8ErrorMessage(FString::Printf(TEXT("Error while trying to read property %s."), *Args.FieldLabel), Response->Body);
+		WebRemoteControlInternalUtils::CreateUTF8ErrorMessage(FString::Printf(TEXT("Error while trying to read property %s."), *Args.FieldLabel), Response->Body);
 	}
 
 	OnComplete(MoveTemp(Response));
@@ -1182,7 +1188,7 @@ bool FWebRemoteControlModule::HandlePresetGetPropertyRoute(const FHttpServerRequ
 
 bool FWebRemoteControlModule::HandlePresetGetExposedActorPropertyRoute(const FHttpServerRequest& Request, const FHttpResultCallback& OnComplete)
 {
-	TUniquePtr<FHttpServerResponse> Response = WebRemoteControlUtils::CreateHttpResponse();
+	TUniquePtr<FHttpServerResponse> Response = WebRemoteControlInternalUtils::CreateHttpResponse();
 
 	FString PresetName = Request.PathParams.FindChecked(TEXT("preset"));
 	FString ActorRCLabel = Request.PathParams.FindChecked(TEXT("actor"));
@@ -1194,7 +1200,7 @@ bool FWebRemoteControlModule::HandlePresetGetExposedActorPropertyRoute(const FHt
 	if (Preset == nullptr)
 	{
 		Response->Code = EHttpServerResponseCodes::NotFound;
-		WebRemoteControlUtils::CreateUTF8ErrorMessage(TEXT("Unable to resolve the preset."), Response->Body);
+		WebRemoteControlInternalUtils::CreateUTF8ErrorMessage(TEXT("Unable to resolve the preset."), Response->Body);
 		OnComplete(MoveTemp(Response));
 		return true;
 	}
@@ -1211,7 +1217,7 @@ bool FWebRemoteControlModule::HandlePresetGetExposedActorPropertyRoute(const FHt
 			if (IRemoteControlModule::Get().ResolveObjectProperty(ERCAccess::READ_ACCESS, Actor, FieldPath, ObjectRef))
 			{
 				FStructOnScope ActorPropertyOnScope = WebRemoteControlStructUtils::CreateActorPropertyOnScope(ObjectRef);
-				WebRemoteControlUtils::SerializeStructOnScope(ActorPropertyOnScope, Writer);
+				WebRemoteControlInternalUtils::SerializeStructOnScope(ActorPropertyOnScope, Writer);
 			
 				Response->Code = EHttpServerResponseCodes::Ok;
 				WebRemoteControlUtils::ConvertToUTF8(WorkingBuffer, Response->Body);
@@ -1219,7 +1225,7 @@ bool FWebRemoteControlModule::HandlePresetGetExposedActorPropertyRoute(const FHt
 			else
 			{
 				Response->Code = EHttpServerResponseCodes::NotFound;
-				WebRemoteControlUtils::CreateUTF8ErrorMessage(FString::Printf(TEXT("Error while trying to read property %s on actor %s."), *PropertyName, *ActorRCLabel), Response->Body);
+				WebRemoteControlInternalUtils::CreateUTF8ErrorMessage(FString::Printf(TEXT("Error while trying to read property %s on actor %s."), *PropertyName, *ActorRCLabel), Response->Body);
 			}
 		}
 	}
@@ -1230,7 +1236,7 @@ bool FWebRemoteControlModule::HandlePresetGetExposedActorPropertyRoute(const FHt
 
 bool FWebRemoteControlModule::HandlePresetGetExposedActorPropertiesRoute(const FHttpServerRequest& Request, const FHttpResultCallback& OnComplete)
 {
-	TUniquePtr<FHttpServerResponse> Response = WebRemoteControlUtils::CreateHttpResponse();
+	TUniquePtr<FHttpServerResponse> Response = WebRemoteControlInternalUtils::CreateHttpResponse();
 
 	FString PresetName = Request.PathParams.FindChecked(TEXT("preset"));
 	FString ActorRCLabel = Request.PathParams.FindChecked(TEXT("actor"));
@@ -1240,7 +1246,7 @@ bool FWebRemoteControlModule::HandlePresetGetExposedActorPropertiesRoute(const F
 	if (Preset == nullptr)
 	{
 		Response->Code = EHttpServerResponseCodes::NotFound;
-		WebRemoteControlUtils::CreateUTF8ErrorMessage(TEXT("Unable to resolve the preset."), Response->Body);
+		WebRemoteControlInternalUtils::CreateUTF8ErrorMessage(TEXT("Unable to resolve the preset."), Response->Body);
 		OnComplete(MoveTemp(Response));
 		return true;
 	}
@@ -1262,19 +1268,19 @@ bool FWebRemoteControlModule::HandlePresetGetExposedActorPropertiesRoute(const F
 			else
 			{
 				Response->Code = EHttpServerResponseCodes::NotFound;
-				WebRemoteControlUtils::CreateUTF8ErrorMessage(TEXT("Unable to serialize exposed actor " + ActorRCLabel), Response->Body);
+				WebRemoteControlInternalUtils::CreateUTF8ErrorMessage(TEXT("Unable to serialize exposed actor " + ActorRCLabel), Response->Body);
 			}
 		}
 		else
 		{
 			Response->Code = EHttpServerResponseCodes::NotFound;
-			WebRemoteControlUtils::CreateUTF8ErrorMessage(TEXT("Unable to resolve object path " + RCActor->Path.ToString()), Response->Body);
+			WebRemoteControlInternalUtils::CreateUTF8ErrorMessage(TEXT("Unable to resolve object path " + RCActor->Path.ToString()), Response->Body);
 		}
 	}
 	else
 	{
 		Response->Code = EHttpServerResponseCodes::NotFound;
-		WebRemoteControlUtils::CreateUTF8ErrorMessage(TEXT("Unable to resolve the exposed actor " + ActorRCLabel), Response->Body);
+		WebRemoteControlInternalUtils::CreateUTF8ErrorMessage(TEXT("Unable to resolve the exposed actor " + ActorRCLabel), Response->Body);
 	}
 
 	OnComplete(MoveTemp(Response));
@@ -1283,10 +1289,10 @@ bool FWebRemoteControlModule::HandlePresetGetExposedActorPropertiesRoute(const F
 
 bool FWebRemoteControlModule::HandlePresetSetExposedActorPropertyRoute(const FHttpServerRequest& Request, const FHttpResultCallback& OnComplete)
 {
-	TUniquePtr<FHttpServerResponse> Response = WebRemoteControlUtils::CreateHttpResponse();
+	TUniquePtr<FHttpServerResponse> Response = WebRemoteControlInternalUtils::CreateHttpResponse();
 
 	FRCPresetSetPropertyRequest SetPropertyRequest;
-	if (!WebRemoteControlUtils::DeserializeRequest(Request, &OnComplete, SetPropertyRequest))
+	if (!WebRemoteControlInternalUtils::DeserializeRequest(Request, &OnComplete, SetPropertyRequest))
 	{
 		return true;
 	}
@@ -1300,7 +1306,7 @@ bool FWebRemoteControlModule::HandlePresetSetExposedActorPropertyRoute(const FHt
 	if (Preset == nullptr)
 	{
 		Response->Code = EHttpServerResponseCodes::NotFound;
-		WebRemoteControlUtils::CreateUTF8ErrorMessage(TEXT("Unable to resolve the preset."), Response->Body);
+		WebRemoteControlInternalUtils::CreateUTF8ErrorMessage(TEXT("Unable to resolve the preset."), Response->Body);
 		OnComplete(MoveTemp(Response));
 		return true;
 	}
@@ -1344,13 +1350,13 @@ bool FWebRemoteControlModule::HandlePresetSetExposedActorPropertyRoute(const FHt
 		else
 		{
 			Response->Code = EHttpServerResponseCodes::NotFound;
-			WebRemoteControlUtils::CreateUTF8ErrorMessage(TEXT("Unable to resolve object path " + RCActor->Path.ToString()), Response->Body);
+			WebRemoteControlInternalUtils::CreateUTF8ErrorMessage(TEXT("Unable to resolve object path " + RCActor->Path.ToString()), Response->Body);
 		}
 	}
 	else
 	{
 		Response->Code = EHttpServerResponseCodes::NotFound;
-		WebRemoteControlUtils::CreateUTF8ErrorMessage(TEXT("Unable to resolve the exposed actor " + ActorRCLabel), Response->Body);
+		WebRemoteControlInternalUtils::CreateUTF8ErrorMessage(TEXT("Unable to resolve the exposed actor " + ActorRCLabel), Response->Body);
 	}
 
 	if (bSuccess)
@@ -1359,7 +1365,7 @@ bool FWebRemoteControlModule::HandlePresetSetExposedActorPropertyRoute(const FHt
 	}
 	else
 	{
-		WebRemoteControlUtils::CreateUTF8ErrorMessage(FString::Printf(TEXT("Error while trying to modify property %s on actor %s."), *PropertyName, *ActorRCLabel), Response->Body);
+		WebRemoteControlInternalUtils::CreateUTF8ErrorMessage(FString::Printf(TEXT("Error while trying to modify property %s on actor %s."), *PropertyName, *ActorRCLabel), Response->Body);
 	}
 
 	OnComplete(MoveTemp(Response));
@@ -1368,18 +1374,18 @@ bool FWebRemoteControlModule::HandlePresetSetExposedActorPropertyRoute(const FHt
 
 bool FWebRemoteControlModule::HandleGetPresetRoute(const FHttpServerRequest& Request, const FHttpResultCallback& OnComplete)
 {
-	TUniquePtr<FHttpServerResponse> Response = WebRemoteControlUtils::CreateHttpResponse();
+	TUniquePtr<FHttpServerResponse> Response = WebRemoteControlInternalUtils::CreateHttpResponse();
 
 	FString PresetName = Request.PathParams.FindChecked(TEXT("preset"));
 
 	if (URemoteControlPreset* Preset = WebRemoteControl::GetPreset(PresetName))
 	{
-		WebRemoteControlUtils::SerializeResponse(FGetPresetResponse{ Preset }, Response->Body);
+		WebRemoteControlUtils::SerializeMessage(FGetPresetResponse{ Preset }, Response->Body);
 		Response->Code = EHttpServerResponseCodes::Ok;
 	}
 	else
 	{
-		WebRemoteControlUtils::CreateUTF8ErrorMessage(FString::Printf(TEXT("Preset %s could not be found."), *PresetName), Response->Body);
+		WebRemoteControlInternalUtils::CreateUTF8ErrorMessage(FString::Printf(TEXT("Preset %s could not be found."), *PresetName), Response->Body);
 		Response->Code = EHttpServerResponseCodes::NotFound;
 	}
 
@@ -1389,12 +1395,12 @@ bool FWebRemoteControlModule::HandleGetPresetRoute(const FHttpServerRequest& Req
 
 bool FWebRemoteControlModule::HandleGetPresetsRoute(const FHttpServerRequest& Request, const FHttpResultCallback& OnComplete)
 {
-	TUniquePtr<FHttpServerResponse> Response = WebRemoteControlUtils::CreateHttpResponse(EHttpServerResponseCodes::Ok);
+	TUniquePtr<FHttpServerResponse> Response = WebRemoteControlInternalUtils::CreateHttpResponse(EHttpServerResponseCodes::Ok);
 
 	TArray<FAssetData> PresetAssets;
 	IRemoteControlModule::Get().GetPresetAssets(PresetAssets);
 
-	WebRemoteControlUtils::SerializeResponse(FListPresetsResponse{ PresetAssets }, Response->Body);
+	WebRemoteControlUtils::SerializeMessage(FListPresetsResponse{ PresetAssets }, Response->Body);
 
 	OnComplete(MoveTemp(Response));
 	return true;
@@ -1402,10 +1408,10 @@ bool FWebRemoteControlModule::HandleGetPresetsRoute(const FHttpServerRequest& Re
 
 bool FWebRemoteControlModule::HandleDescribeObjectRoute(const FHttpServerRequest& Request, const FHttpResultCallback& OnComplete)
 {
-	TUniquePtr<FHttpServerResponse> Response = WebRemoteControlUtils::CreateHttpResponse();
+	TUniquePtr<FHttpServerResponse> Response = WebRemoteControlInternalUtils::CreateHttpResponse();
 
 	FDescribeObjectRequest DescribeRequest;
-	if (!WebRemoteControlUtils::DeserializeRequest(Request, &OnComplete, DescribeRequest))
+	if (!WebRemoteControlInternalUtils::DeserializeRequest(Request, &OnComplete, DescribeRequest))
 	{
 		return true;
 	}
@@ -1414,12 +1420,12 @@ bool FWebRemoteControlModule::HandleDescribeObjectRoute(const FHttpServerRequest
 	FString ErrorText;
 	if (IRemoteControlModule::Get().ResolveObject(ERCAccess::READ_ACCESS, DescribeRequest.ObjectPath, TEXT(""), Ref, &ErrorText) && Ref.Object.IsValid())
 	{
-		WebRemoteControlUtils::SerializeResponse(FDescribeObjectResponse{ Ref.Object.Get() }, Response->Body);
+		WebRemoteControlUtils::SerializeMessage(FDescribeObjectResponse{ Ref.Object.Get() }, Response->Body);
 		Response->Code = EHttpServerResponseCodes::Ok;
 	}
 	else
 	{
-		WebRemoteControlUtils::CreateUTF8ErrorMessage(*ErrorText, Response->Body);
+		WebRemoteControlInternalUtils::CreateUTF8ErrorMessage(*ErrorText, Response->Body);
 		Response->Code = EHttpServerResponseCodes::NotFound;
 	}
 
@@ -1429,17 +1435,17 @@ bool FWebRemoteControlModule::HandleDescribeObjectRoute(const FHttpServerRequest
 
 bool FWebRemoteControlModule::HandleSearchActorRoute(const FHttpServerRequest& Request, const FHttpResultCallback& OnComplete)
 {
-	TUniquePtr<FHttpServerResponse> Response = WebRemoteControlUtils::CreateHttpResponse(EHttpServerResponseCodes::NotSupported);
-	WebRemoteControlUtils::CreateUTF8ErrorMessage(TEXT("Route not implemented."), Response->Body);
+	TUniquePtr<FHttpServerResponse> Response = WebRemoteControlInternalUtils::CreateHttpResponse(EHttpServerResponseCodes::NotSupported);
+	WebRemoteControlInternalUtils::CreateUTF8ErrorMessage(TEXT("Route not implemented."), Response->Body);
 	OnComplete(MoveTemp(Response));
 	return true;
 }
 
 bool FWebRemoteControlModule::HandleSearchAssetRoute(const FHttpServerRequest& Request, const FHttpResultCallback& OnComplete)
 {
-	TUniquePtr<FHttpServerResponse> Response = WebRemoteControlUtils::CreateHttpResponse();
+	TUniquePtr<FHttpServerResponse> Response = WebRemoteControlInternalUtils::CreateHttpResponse();
 	FSearchAssetRequest SearchAssetRequest;
-	if (!WebRemoteControlUtils::DeserializeRequest(Request, &OnComplete, SearchAssetRequest))
+	if (!WebRemoteControlInternalUtils::DeserializeRequest(Request, &OnComplete, SearchAssetRequest))
 	{
 		return true;
 	}
@@ -1481,7 +1487,7 @@ bool FWebRemoteControlModule::HandleSearchAssetRoute(const FHttpServerRequest& R
 		}
 	}
 
-	WebRemoteControlUtils::SerializeResponse(FSearchAssetResponse{ MoveTemp(FilteredAssets) }, Response->Body);
+	WebRemoteControlUtils::SerializeMessage(FSearchAssetResponse{ MoveTemp(FilteredAssets) }, Response->Body);
 	Response->Code = EHttpServerResponseCodes::Ok;
 
 	OnComplete(MoveTemp(Response));
@@ -1490,18 +1496,18 @@ bool FWebRemoteControlModule::HandleSearchAssetRoute(const FHttpServerRequest& R
 
 bool FWebRemoteControlModule::HandleGetMetadataRoute(const FHttpServerRequest& Request, const FHttpResultCallback& OnComplete)
 {
-	TUniquePtr<FHttpServerResponse> Response = WebRemoteControlUtils::CreateHttpResponse();
+	TUniquePtr<FHttpServerResponse> Response = WebRemoteControlInternalUtils::CreateHttpResponse();
 
 	FString PresetName = Request.PathParams.FindChecked(TEXT("preset"));
 
 	if (URemoteControlPreset* Preset = WebRemoteControl::GetPreset(PresetName))
 	{
-		WebRemoteControlUtils::SerializeResponse(FGetMetadataResponse{Preset->Metadata}, Response->Body);
+		WebRemoteControlUtils::SerializeMessage(FGetMetadataResponse{Preset->Metadata}, Response->Body);
 		Response->Code = EHttpServerResponseCodes::Ok;
 	}
 	else
 	{
-		WebRemoteControlUtils::CreateUTF8ErrorMessage(FString::Printf(TEXT("Preset %s could not be found."), *PresetName), Response->Body);
+		WebRemoteControlInternalUtils::CreateUTF8ErrorMessage(FString::Printf(TEXT("Preset %s could not be found."), *PresetName), Response->Body);
 		Response->Code = EHttpServerResponseCodes::NotFound;
 	}
 
@@ -1511,7 +1517,7 @@ bool FWebRemoteControlModule::HandleGetMetadataRoute(const FHttpServerRequest& R
 
 bool FWebRemoteControlModule::HandleMetadataFieldOperationsRoute(const FHttpServerRequest& Request, const FHttpResultCallback& OnComplete)
 {
-	TUniquePtr<FHttpServerResponse> Response = WebRemoteControlUtils::CreateHttpResponse();
+	TUniquePtr<FHttpServerResponse> Response = WebRemoteControlInternalUtils::CreateHttpResponse();
 
 	FString PresetName = Request.PathParams.FindChecked(TEXT("preset"));
 	FString MetadataField = Request.PathParams.FindChecked(TEXT("metadatafield"));
@@ -1522,18 +1528,18 @@ bool FWebRemoteControlModule::HandleMetadataFieldOperationsRoute(const FHttpServ
 		{
 			if (FString* MetadataValue = Preset->Metadata.Find(MetadataField))
 			{
-				WebRemoteControlUtils::SerializeResponse(FGetMetadataFieldResponse{ *MetadataValue }, Response->Body);
+				WebRemoteControlUtils::SerializeMessage(FGetMetadataFieldResponse{ *MetadataValue }, Response->Body);
 			}
 			else
 			{
-				WebRemoteControlUtils::CreateUTF8ErrorMessage(FString::Printf(TEXT("Metadata field %s could not be found."), *MetadataField), Response->Body);
+				WebRemoteControlInternalUtils::CreateUTF8ErrorMessage(FString::Printf(TEXT("Metadata field %s could not be found."), *MetadataField), Response->Body);
 				Response->Code = EHttpServerResponseCodes::NotFound;
 			}
 		}
 		else if (Request.Verb == EHttpServerRequestVerbs::VERB_PUT)
 		{
 			FSetPresetMetadataRequest SetMetadataRequest;
-			if (!WebRemoteControlUtils::DeserializeRequest(Request, &OnComplete, SetMetadataRequest))
+			if (!WebRemoteControlInternalUtils::DeserializeRequest(Request, &OnComplete, SetMetadataRequest))
 			{
 				return true;
 			}
@@ -1552,7 +1558,7 @@ bool FWebRemoteControlModule::HandleMetadataFieldOperationsRoute(const FHttpServ
 	}
 	else
 	{
-		WebRemoteControlUtils::CreateUTF8ErrorMessage(FString::Printf(TEXT("Preset %s could not be found."), *PresetName), Response->Body);
+		WebRemoteControlInternalUtils::CreateUTF8ErrorMessage(FString::Printf(TEXT("Preset %s could not be found."), *PresetName), Response->Body);
 		Response->Code = EHttpServerResponseCodes::NotFound;
 	}
 
@@ -1562,15 +1568,15 @@ bool FWebRemoteControlModule::HandleMetadataFieldOperationsRoute(const FHttpServ
 
 bool FWebRemoteControlModule::HandleSearchObjectRoute(const FHttpServerRequest& Request, const FHttpResultCallback& OnComplete)
 {
-	TUniquePtr<FHttpServerResponse> Response = WebRemoteControlUtils::CreateHttpResponse(EHttpServerResponseCodes::NotSupported);
-	WebRemoteControlUtils::CreateUTF8ErrorMessage(TEXT("Route not implemented."), Response->Body);
+	TUniquePtr<FHttpServerResponse> Response = WebRemoteControlInternalUtils::CreateHttpResponse(EHttpServerResponseCodes::NotSupported);
+	WebRemoteControlInternalUtils::CreateUTF8ErrorMessage(TEXT("Route not implemented."), Response->Body);
 	OnComplete(MoveTemp(Response));
 	return true;
 }
 
 bool FWebRemoteControlModule::HandleEntityMetadataOperationsRoute(const FHttpServerRequest& Request, const FHttpResultCallback& OnComplete)
 {
-	TUniquePtr<FHttpServerResponse> Response = WebRemoteControlUtils::CreateHttpResponse();
+	TUniquePtr<FHttpServerResponse> Response = WebRemoteControlInternalUtils::CreateHttpResponse();
 
 	FString PresetName = Request.PathParams.FindChecked(TEXT("preset"));
 	FString Label = Request.PathParams.FindChecked(TEXT("label"));
@@ -1580,7 +1586,7 @@ bool FWebRemoteControlModule::HandleEntityMetadataOperationsRoute(const FHttpSer
 	if (Preset == nullptr)
 	{
 		Response->Code = EHttpServerResponseCodes::NotFound;
-		WebRemoteControlUtils::CreateUTF8ErrorMessage(TEXT("Unable to resolve the preset."), Response->Body);
+		WebRemoteControlInternalUtils::CreateUTF8ErrorMessage(TEXT("Unable to resolve the preset."), Response->Body);
 		OnComplete(MoveTemp(Response));
 		return true;
 	}
@@ -1590,7 +1596,7 @@ bool FWebRemoteControlModule::HandleEntityMetadataOperationsRoute(const FHttpSer
 	if (!Entity.IsValid())
 	{
 		Response->Code = EHttpServerResponseCodes::NotFound;
-		WebRemoteControlUtils::CreateUTF8ErrorMessage(TEXT("Unable to resolve the preset entity."), Response->Body);
+		WebRemoteControlInternalUtils::CreateUTF8ErrorMessage(TEXT("Unable to resolve the preset entity."), Response->Body);
 		OnComplete(MoveTemp(Response));
 		return true;
 	}
@@ -1599,23 +1605,23 @@ bool FWebRemoteControlModule::HandleEntityMetadataOperationsRoute(const FHttpSer
 	{
 		if (const FString* MetadataValue = Entity->GetMetadata().Find(*Key))
 		{
-			WebRemoteControlUtils::SerializeResponse(FGetMetadataFieldResponse{ *MetadataValue }, Response->Body);
+			WebRemoteControlUtils::SerializeMessage(FGetMetadataFieldResponse{ *MetadataValue }, Response->Body);
 		}
 		else
 		{
-			WebRemoteControlUtils::CreateUTF8ErrorMessage(FString::Printf(TEXT("Metadata entry %s could not be found."), *Key), Response->Body);
+			WebRemoteControlInternalUtils::CreateUTF8ErrorMessage(FString::Printf(TEXT("Metadata entry %s could not be found."), *Key), Response->Body);
 			Response->Code = EHttpServerResponseCodes::NotFound;
 		}
 	}
 	else if (Request.Verb == EHttpServerRequestVerbs::VERB_PUT)
 	{
-		if (!WebRemoteControlUtils::ValidateContentType(Request, TEXT("application/json"), OnComplete))
+		if (!WebRemoteControlInternalUtils::ValidateContentType(Request, TEXT("application/json"), OnComplete))
 		{
 			return true;
 		}
 		
 		FSetEntityMetadataRequest SetMetadataRequest;
-		if (!WebRemoteControlUtils::DeserializeRequest(Request, &OnComplete, SetMetadataRequest))
+		if (!WebRemoteControlInternalUtils::DeserializeRequest(Request, &OnComplete, SetMetadataRequest))
 		{
 			return true;
 		}
@@ -1643,12 +1649,12 @@ bool FWebRemoteControlModule::HandleEntityMetadataOperationsRoute(const FHttpSer
 
 bool FWebRemoteControlModule::HandleEntitySetLabelRoute(const FHttpServerRequest& Request, const FHttpResultCallback& OnComplete)
 {
-	TUniquePtr<FHttpServerResponse> Response = WebRemoteControlUtils::CreateHttpResponse();
+	TUniquePtr<FHttpServerResponse> Response = WebRemoteControlInternalUtils::CreateHttpResponse();
 
 	FString PresetName = Request.PathParams.FindChecked(TEXT("preset"));
 	FString Label = Request.PathParams.FindChecked(TEXT("label"));
 	
-	if (!WebRemoteControlUtils::ValidateContentType(Request, TEXT("application/json"), OnComplete))
+	if (!WebRemoteControlInternalUtils::ValidateContentType(Request, TEXT("application/json"), OnComplete))
 	{
 		return true;
 	}
@@ -1657,7 +1663,7 @@ bool FWebRemoteControlModule::HandleEntitySetLabelRoute(const FHttpServerRequest
 	if (Preset == nullptr)
 	{
 		Response->Code = EHttpServerResponseCodes::NotFound;
-		WebRemoteControlUtils::CreateUTF8ErrorMessage(TEXT("Unable to resolve the preset."), Response->Body);
+		WebRemoteControlInternalUtils::CreateUTF8ErrorMessage(TEXT("Unable to resolve the preset."), Response->Body);
 		OnComplete(MoveTemp(Response));
 		return true;
 	}
@@ -1667,13 +1673,13 @@ bool FWebRemoteControlModule::HandleEntitySetLabelRoute(const FHttpServerRequest
 	if (!Entity.IsValid())
 	{
 		Response->Code = EHttpServerResponseCodes::NotFound;
-		WebRemoteControlUtils::CreateUTF8ErrorMessage(TEXT("Unable to resolve the preset entity."), Response->Body);
+		WebRemoteControlInternalUtils::CreateUTF8ErrorMessage(TEXT("Unable to resolve the preset entity."), Response->Body);
 		OnComplete(MoveTemp(Response));
 		return true;
 	}
 
 	FSetEntityLabelRequest SetEntityLabelRequest;
-	if (!WebRemoteControlUtils::DeserializeRequest(Request, &OnComplete, SetEntityLabelRequest))
+	if (!WebRemoteControlInternalUtils::DeserializeRequest(Request, &OnComplete, SetEntityLabelRequest))
 	{
 		return true;
 	}
@@ -1685,7 +1691,7 @@ bool FWebRemoteControlModule::HandleEntitySetLabelRoute(const FHttpServerRequest
 	
 	FName AssignedLabel = Entity->Rename(*SetEntityLabelRequest.NewLabel);
 	
-	WebRemoteControlUtils::SerializeResponse(FSetEntityLabelResponse{ AssignedLabel.ToString() }, Response->Body);
+	WebRemoteControlUtils::SerializeMessage(FSetEntityLabelResponse{ AssignedLabel.ToString() }, Response->Body);
 	Response->Code = EHttpServerResponseCodes::Ok;
 
 	OnComplete(MoveTemp(Response));
@@ -1701,13 +1707,13 @@ void FWebRemoteControlModule::HandleWebSocketHttpMessage(const FRemoteControlWeb
 	//Early failure is http server not started
 	if (HttpRouter == nullptr)
 	{
-		WebRemoteControlUtils::CreateUTF8ErrorMessage(TEXT("HTTP server not started."), UTF8Response);
+		WebRemoteControlInternalUtils::CreateUTF8ErrorMessage(TEXT("HTTP server not started."), UTF8Response);
 		WebSocketServer.Send(WebSocketMessage.ClientId, MoveTemp(UTF8Response));
 		return;
 	}
 
 	FRCRequestWrapper Wrapper;
-	if (!WebRemoteControlUtils::DeserializeWrappedRequestPayload(WebSocketMessage.RequestPayload, nullptr, Wrapper))
+	if (!WebRemoteControlInternalUtils::DeserializeWrappedRequestPayload(WebSocketMessage.RequestPayload, nullptr, Wrapper))
 	{
 		return;
 	}
@@ -1733,8 +1739,8 @@ void FWebRemoteControlModule::InvokeWrappedRequest(const FRCRequestWrapper& Wrap
 
 	if (!HttpRouter->Query(UnwrappedRequest, ResponseLambda))
 	{
-		TUniquePtr<FHttpServerResponse> InnerRouteResponse = WebRemoteControlUtils::CreateHttpResponse(EHttpServerResponseCodes::NotFound);
-		WebRemoteControlUtils::CreateUTF8ErrorMessage(FString::Printf(TEXT("Route %s could not be found."), *Wrapper.URL), InnerRouteResponse->Body);
+		TUniquePtr<FHttpServerResponse> InnerRouteResponse = WebRemoteControlInternalUtils::CreateHttpResponse(EHttpServerResponseCodes::NotFound);
+		WebRemoteControlInternalUtils::CreateUTF8ErrorMessage(FString::Printf(TEXT("Route %s could not be found."), *Wrapper.URL), InnerRouteResponse->Body);
 		ResponseLambda(MoveTemp(InnerRouteResponse));
 	}
 }
diff --git a/Engine/Plugins/VirtualProduction/RemoteControl/Source/WebRemoteControl/Private/WebRemoteControl.h b/Engine/Plugins/VirtualProduction/RemoteControl/Source/WebRemoteControl/Private/WebRemoteControl.h
--- a/Engine/Plugins/VirtualProduction/RemoteControl/Source/WebRemoteControl/Private/WebRemoteControl.h
+++ b/Engine/Plugins/VirtualProduction/RemoteControl/Source/WebRemoteControl/Private/WebRemoteControl.h
@@ -11,6 +11,7 @@
 #include "HttpServerRequest.h"
 #include "RemoteControlRequest.h"
 #include "RemoteControlRoute.h"
+#include "RemoteControlWebsocketRoute.h"
 #include "RemoteControlWebSocketServer.h"
 #include "StageAppBeaconReceiver.h"
 #include "WebRemoteControlEditorRoutes.h"
@@ -47,7 +48,12 @@ public:
 	virtual FSimpleMulticastDelegate& OnHttpServerStopped() override { return OnHttpServerStoppedDelegate; }
 	virtual FOnWebServerStarted& OnWebSocketServerStarted() override { return OnWebSocketServerStartedDelegate; }
 	virtual FSimpleMulticastDelegate& OnWebSocketServerStopped() override { return OnWebSocketServerStoppedDelegate; }
-	virtual void SetExternalRemoteWebSocketLoggerConnection(TSharedPtr<class INetworkingWebSocket> WebSocketLoggerConnection);
+	virtual FOnWebSocketConnectionClosed& OnWebSocketConnectionOpened() override { return WebSocketServer.OnConnectionOpened(); }
+	virtual FOnWebSocketConnectionClosed& OnWebSocketConnectionClosed() override { return WebSocketServer.OnConnectionClosed(); }
+	virtual void RegisterWebsocketRoute(const FRemoteControlWebsocketRoute& Route) override;
+	virtual void UnregisterWebsocketRoute(const FRemoteControlWebsocketRoute& Route) override;
+	virtual void SendWebsocketMessage(const FGuid& InTargetClientId, const TArray<uint8>& InUTF8Payload) override;
+	virtual void SetExternalRemoteWebSocketLoggerConnection(TSharedPtr<class INetworkingWebSocket> WebSocketLoggerConnection) override;
 	//~ End IWebRemoteControlModule Interface
 
 	/**
@@ -62,18 +68,6 @@ public:
 	 */
 	void UnregisterRoute(const FRemoteControlRoute& Route);
 
-	/**
-	 * Register a websocket route.
-	 * @param Route the route to register.
-	 */
-	void RegisterWebsocketRoute(const FRemoteControlWebsocketRoute& Route);
-
-	/**
-	 * Unregister a websocket route.
-	 * @param Route the route to unregister.
-	 */
-	void UnregisterWebsocketRoute(const FRemoteControlWebsocketRoute& Route);
-
 	/**
 	 * Start the web control server
 	 */
diff --git a/Engine/Plugins/VirtualProduction/RemoteControl/Source/WebRemoteControl/Private/WebRemoteControlEditorRoutes.cpp b/Engine/Plugins/VirtualProduction/RemoteControl/Source/WebRemoteControl/Private/WebRemoteControlEditorRoutes.cpp
--- a/Engine/Plugins/VirtualProduction/RemoteControl/Source/WebRemoteControl/Private/WebRemoteControlEditorRoutes.cpp
+++ b/Engine/Plugins/VirtualProduction/RemoteControl/Source/WebRemoteControl/Private/WebRemoteControlEditorRoutes.cpp
@@ -3,7 +3,7 @@
 #include "WebRemoteControlEditorRoutes.h"
 
 #if WITH_EDITOR
-#include "WebRemoteControlUtils.h"
+#include "WebRemoteControlInternalUtils.h"
 #include "WebRemoteControl.h"
 
 // Http
@@ -68,27 +68,27 @@ void FWebRemoteControlEditorRoutes::UnregisterRoutes(FWebRemoteControlModule* We
 
 bool FWebRemoteControlEditorRoutes::HandleObjectEventRoute(const FHttpServerRequest& Request, const FHttpResultCallback& OnComplete)
 {
-	if (!WebRemoteControlUtils::ValidateContentType(Request, TEXT("application/json"), OnComplete))
+	if (!WebRemoteControlInternalUtils::ValidateContentType(Request, TEXT("application/json"), OnComplete))
 	{
 		return true;
 	}
 
 	FRemoteControlObjectEventHookRequest EventRequest;
-	if (!WebRemoteControlUtils::DeserializeRequest(Request, &OnComplete, EventRequest))
+	if (!WebRemoteControlInternalUtils::DeserializeRequest(Request, &OnComplete, EventRequest))
 	{
 		return true;
 	}
 
 	// Queue the request and complete the event when it triggers
-	AddPendingEvent(MoveTemp(EventRequest), WebRemoteControlUtils::CreateHttpResponse(), OnComplete);
+	AddPendingEvent(MoveTemp(EventRequest), WebRemoteControlInternalUtils::CreateHttpResponse(), OnComplete);
 	return true;
 };
 
 bool FWebRemoteControlEditorRoutes::HandleGetThumbnailRoute(const FHttpServerRequest& Request, const FHttpResultCallback& OnComplete)
 {
-	TUniquePtr<FHttpServerResponse> Response = WebRemoteControlUtils::CreateHttpResponse();
+	TUniquePtr<FHttpServerResponse> Response = WebRemoteControlInternalUtils::CreateHttpResponse();
 	FGetObjectThumbnailRequest GetThumbnailRequest;
-	if (!WebRemoteControlUtils::DeserializeRequest(Request, &OnComplete, GetThumbnailRequest))
+	if (!WebRemoteControlInternalUtils::DeserializeRequest(Request, &OnComplete, GetThumbnailRequest))
 	{
 		return true;
 	}
@@ -137,7 +137,7 @@ bool FWebRemoteControlEditorRoutes::HandleGetThumbnailRoute(const FHttpServerReq
 		{
 			if (FObjectThumbnail* Thumbnail = ThumbnailMap.Find(ObjectFullName))
 			{
-				WebRemoteControlUtils::AddContentTypeHeaders(Response.Get(), TEXT("image/png"));
+				WebRemoteControlInternalUtils::AddContentTypeHeaders(Response.Get(), TEXT("image/png"));
 				IImageWrapperModule& ImageWrapperModule = FModuleManager::Get().LoadModuleChecked<IImageWrapperModule>(TEXT("ImageWrapper"));
 				EImageFormat Format = ImageWrapperModule.DetectImageFormat((void*)Thumbnail->AccessCompressedImageData().GetData(), Thumbnail->AccessCompressedImageData().Num());
 				if (Format == EImageFormat::PNG)
@@ -161,7 +161,7 @@ bool FWebRemoteControlEditorRoutes::HandleGetThumbnailRoute(const FHttpServerReq
 				FName ResourceName = ThumbnailBrush->GetResourceName();
 				if (FFileHelper::LoadFileToArray(Response->Body, *ResourceName.ToString()))
 				{
-					WebRemoteControlUtils::AddContentTypeHeaders(Response.Get(), TEXT("image/png"));
+					WebRemoteControlInternalUtils::AddContentTypeHeaders(Response.Get(), TEXT("image/png"));
 					Response->Code = EHttpServerResponseCodes::Ok;
 				}
 			}
@@ -170,7 +170,7 @@ bool FWebRemoteControlEditorRoutes::HandleGetThumbnailRoute(const FHttpServerReq
 
 	if (!Response->Body.Num())
 	{
-		WebRemoteControlUtils::CreateUTF8ErrorMessage(FString::Printf(TEXT("Could not load thumbnail for object %s"), *GetThumbnailRequest.ObjectPath), Response->Body);
+		WebRemoteControlInternalUtils::CreateUTF8ErrorMessage(FString::Printf(TEXT("Could not load thumbnail for object %s"), *GetThumbnailRequest.ObjectPath), Response->Body);
 		Response->Code = EHttpServerResponseCodes::NotFound;
 	}
 
@@ -193,7 +193,7 @@ void FWebRemoteControlEditorRoutes::AddPendingEvent(FRemoteControlObjectEventHoo
 	}
 	else
 	{
-		WebRemoteControlUtils::CreateUTF8ErrorMessage(ErrorText, InResponse->Body);
+		WebRemoteControlInternalUtils::CreateUTF8ErrorMessage(ErrorText, InResponse->Body);
 		InResponse->Code = EHttpServerResponseCodes::BadRequest;
 		OnComplete(MoveTemp(InResponse));
 	}
diff --git a/Engine/Plugins/VirtualProduction/RemoteControl/Source/WebRemoteControl/Private/WebRemoteControlExternalLogger.cpp b/Engine/Plugins/VirtualProduction/RemoteControl/Source/WebRemoteControl/Private/WebRemoteControlExternalLogger.cpp
--- a/Engine/Plugins/VirtualProduction/RemoteControl/Source/WebRemoteControl/Private/WebRemoteControlExternalLogger.cpp
+++ b/Engine/Plugins/VirtualProduction/RemoteControl/Source/WebRemoteControl/Private/WebRemoteControlExternalLogger.cpp
@@ -2,7 +2,7 @@
 
 #include "WebRemoteControlExternalLogger.h"
 #include "WebSocketNetworkingDelegates.h"
-#include "WebRemoteControlUtils.h"
+#include "WebRemoteControlInternalUtils.h"
 #include "INetworkingWebSocket.h"
 
 
diff --git a/Engine/Plugins/VirtualProduction/RemoteControl/Source/WebRemoteControl/Private/WebRemoteControlInternalUtils.cpp b/Engine/Plugins/VirtualProduction/RemoteControl/Source/WebRemoteControl/Private/WebRemoteControlInternalUtils.cpp
new file mode 100644
--- /dev/null
+++ b/Engine/Plugins/VirtualProduction/RemoteControl/Source/WebRemoteControl/Private/WebRemoteControlInternalUtils.cpp
@@ -0,0 +1,433 @@
+// Copyright Epic Games, Inc. All Rights Reserved.
+
+#include "WebRemoteControlInternalUtils.h"
+#include "HttpServerRequest.h"
+#include "Misc/Base64.h"
+#include "Serialization/JsonReader.h"
+#include "UObject/StructOnScope.h"
+
+
+namespace RemotePayloadSerializer
+{
+FName NAME_Get = "GET";
+FName NAME_Put = "PUT";
+FName NAME_Post = "POST";
+FName NAME_Patch = "PATCH";
+FName NAME_Delete = "DELETE";
+FName NAME_Options = "OPTIONS";
+
+typedef UCS2CHAR PayloadCharType;
+const FName ReturnValuePropName(TEXT("ReturnValue"));
+
+void ReplaceFirstOccurence(TConstArrayView<uint8> InPayload, const FString& From, const FString& To, TArray<uint8>& OutModifiedPayload)
+{
+	if (TCHAR* PropValuePtr = FCString::Stristr((TCHAR*)InPayload.GetData(), *From))
+	{
+		OutModifiedPayload.Reserve(InPayload.Num() + (To.Len() - (From.Len()) * sizeof(TCHAR)));
+
+		int32 StartIndex = UE_PTRDIFF_TO_INT32(PropValuePtr - (TCHAR*)InPayload.GetData()) * sizeof(TCHAR);
+		int32 EndIndex = StartIndex + From.Len() * sizeof(TCHAR);
+		TConstArrayView<uint8> RestOfPayload = InPayload.Slice(EndIndex, InPayload.Num() - EndIndex);
+		OutModifiedPayload.Append(InPayload.GetData(), StartIndex);
+		OutModifiedPayload.Append((uint8*)*To, To.Len() * sizeof(TCHAR));
+		OutModifiedPayload.Append(RestOfPayload.GetData(), RestOfPayload.Num());
+	}
+}
+
+EHttpServerRequestVerbs ParseHttpVerb(FName InVerb)
+{
+	if (InVerb == NAME_Get)
+	{
+		return EHttpServerRequestVerbs::VERB_GET;
+	}
+	else if (InVerb == NAME_Post)
+	{
+		return EHttpServerRequestVerbs::VERB_POST;
+	}
+	else if (InVerb == NAME_Put)
+	{
+		return EHttpServerRequestVerbs::VERB_PUT;
+	}
+	else if (InVerb == NAME_Patch)
+	{
+		return EHttpServerRequestVerbs::VERB_PATCH;
+	}
+	else if (InVerb == NAME_Delete)
+	{
+		return EHttpServerRequestVerbs::VERB_DELETE;
+	}
+	else if (InVerb == NAME_Options)
+	{
+		return EHttpServerRequestVerbs::VERB_OPTIONS;
+	}
+	else
+	{
+		return EHttpServerRequestVerbs::VERB_NONE;
+	}
+}
+
+TSharedRef<FHttpServerRequest> UnwrapHttpRequest(const FRCRequestWrapper& Wrapper, const FHttpServerRequest* TemplateRequest)
+{
+	TSharedRef<FHttpServerRequest> WrappedHttpRequest = MakeShared<FHttpServerRequest>();
+	if (TemplateRequest)
+	{
+		WrappedHttpRequest->HttpVersion = TemplateRequest->HttpVersion;
+		WrappedHttpRequest->Headers = TemplateRequest->Headers;
+	}
+	else
+	{
+		// Assign Request defaults.
+		WrappedHttpRequest->HttpVersion = HttpVersion::EHttpServerHttpVersion::HTTP_VERSION_1_1;
+		WrappedHttpRequest->Headers.Add(TEXT("Content-Type"), { TEXT("application/json") });
+	}
+
+	WebRemoteControlInternalUtils::AddWrappedRequestHeader(*WrappedHttpRequest);
+	WrappedHttpRequest->RelativePath = Wrapper.URL;
+	WrappedHttpRequest->Verb = ParseHttpVerb(Wrapper.Verb);
+	WrappedHttpRequest->Body = Wrapper.TCHARBody;
+
+	return WrappedHttpRequest;
+}
+
+void SerializeWrappedCallResponse(int32 RequestId, TUniquePtr<FHttpServerResponse> Response, FMemoryWriter& Writer)
+{
+	FRCJsonStructSerializerBackend Backend(Writer, FRCJsonStructSerializerBackend::DefaultSerializerFlags);
+	TSharedPtr<TJsonWriter<ANSICHAR>> JsonWriter = TJsonWriter<ANSICHAR>::Create(&Writer);
+	TArray<FString>* ContentTypeHeaders = Response->Headers.Find(TEXT("Content-Type"));
+	const bool bIsBinaryData = ContentTypeHeaders && ContentTypeHeaders->Contains(TEXT("image/png"));
+
+	JsonWriter->WriteObjectStart();
+	JsonWriter->WriteValue(TEXT("RequestId"), RequestId);
+	JsonWriter->WriteValue(TEXT("ResponseCode"), static_cast<int32>(Response->Code));
+	JsonWriter->WriteIdentifierPrefix(TEXT("ResponseBody"));
+
+	if (Response->Body.Num())
+	{
+		if (!bIsBinaryData)
+		{
+			Writer.Serialize((void*)Response->Body.GetData(), Response->Body.Num());
+		}
+		else
+		{
+			FString Base64String = FString::Printf(TEXT("\"%s\""), *FBase64::Encode(Response->Body));
+			TArray<uint8> WorkingBuffer;
+			WebRemoteControlUtils::ConvertToUTF8(Base64String, WorkingBuffer);
+			Writer.Serialize((void*)WorkingBuffer.GetData(), WorkingBuffer.Num());
+		}
+	}
+	else
+	{
+		JsonWriter->WriteNull();
+	}
+
+	JsonWriter->WriteObjectEnd();
+}
+
+bool DeserializeCall(const FHttpServerRequest& InRequest, FRCCall& OutCall, const FHttpResultCallback& InCompleteCallback)
+{
+	// Create Json reader to read the payload, the payload will already be validated as being in TCHAR
+	FRCCallRequest CallRequest;
+	if (!WebRemoteControlInternalUtils::DeserializeRequest(InRequest, &InCompleteCallback, CallRequest))
+	{
+		return false;
+	}
+
+	FString ErrorText;
+	bool bSuccess = true;
+	if (IRemoteControlModule::Get().ResolveCall(CallRequest.ObjectPath, CallRequest.FunctionName, OutCall.CallRef, &ErrorText))
+	{
+		// Initialize the param struct with default parameters
+		OutCall.bGenerateTransaction = CallRequest.GenerateTransaction;
+		OutCall.ParamStruct = FStructOnScope(OutCall.CallRef.Function.Get());
+
+		// If some parameters were provided, deserialize them
+		const FBlockDelimiters& ParametersDelimiters = CallRequest.GetStructParameters().FindChecked(FRCCallRequest::ParametersLabel());
+		if (ParametersDelimiters.BlockStart > 0)
+		{
+			FMemoryReader Reader(CallRequest.TCHARBody);
+			Reader.Seek(ParametersDelimiters.BlockStart);
+			Reader.SetLimitSize(ParametersDelimiters.BlockEnd);
+
+			FRCJsonStructDeserializerBackend Backend(Reader);
+			if (!FStructDeserializer::Deserialize((void*)OutCall.ParamStruct.GetStructMemory(), *const_cast<UStruct*>(OutCall.ParamStruct.GetStruct()), Backend, FStructDeserializerPolicies()))
+			{
+				ErrorText = TEXT("Parameters object improperly formatted.");
+				bSuccess = false;
+			}
+		}
+	}
+	else
+	{
+		bSuccess = false;
+	}
+
+	if (!bSuccess)
+	{
+		UE_LOG(LogRemoteControl, Error, TEXT("Web Remote Call deserialization error: %s"), *ErrorText);
+		TUniquePtr<FHttpServerResponse> Response = WebRemoteControlInternalUtils::CreateHttpResponse();
+		WebRemoteControlInternalUtils::CreateUTF8ErrorMessage(ErrorText, Response->Body);
+		InCompleteCallback(MoveTemp(Response));
+	}
+
+	return bSuccess;
+}
+
+bool SerializeCall(const FRCCall& InCall, TArray<uint8>& OutPayload, bool bOnlyReturn)
+{
+	FMemoryWriter Writer(OutPayload);
+	TSharedPtr<TJsonWriter<PayloadCharType>> JsonWriter;
+
+	if (!bOnlyReturn)
+	{
+		// Create Json writer to write the payload, if we do not just write the return value back.
+		JsonWriter = TJsonWriter<PayloadCharType>::Create(&Writer);
+
+		JsonWriter->WriteObjectStart();
+		JsonWriter->WriteValue(TEXT("ObjectPath"), InCall.CallRef.Object->GetPathName());
+		JsonWriter->WriteValue(TEXT("FunctionName"), InCall.CallRef.Function->GetFName().ToString());
+		JsonWriter->WriteIdentifierPrefix(TEXT("Parameters"));
+	}
+
+	// write the param struct
+	FRCJsonStructSerializerBackend Backend(Writer, FRCJsonStructSerializerBackend::DefaultSerializerFlags);
+	FStructSerializerPolicies Policies;
+
+	if (bOnlyReturn)
+	{
+		Policies.PropertyFilter = [](const FProperty* CurrentProp, const FProperty* ParentProperty) -> bool
+		{
+			return CurrentProp->HasAnyPropertyFlags(CPF_ReturnParm | CPF_OutParm) || ParentProperty != nullptr;
+		};
+	}
+
+	FStructSerializer::Serialize((void*)InCall.ParamStruct.GetStructMemory(), *const_cast<UStruct*>(InCall.ParamStruct.GetStruct()), Backend, Policies);
+
+	if (!bOnlyReturn)
+	{
+		JsonWriter->WriteObjectEnd();
+	}
+
+	return true;
+}
+
+bool DeserializeObjectRef(const FHttpServerRequest& InRequest, FRCObjectReference& OutObjectRef, FRCObjectRequest& OutDeserializedRequest, const FHttpResultCallback& InCompleteCallback)
+{
+	if (!WebRemoteControlInternalUtils::DeserializeRequest(InRequest, &InCompleteCallback, OutDeserializedRequest))
+	{
+		return false;
+	}
+
+	FString ErrorText;
+
+	// If we properly identified the object path, property name and access type as well as identified the starting / end position for property value
+	// resolve the object reference
+	IRemoteControlModule::Get().ResolveObject(OutDeserializedRequest.GetAccessValue(), OutDeserializedRequest.ObjectPath, OutDeserializedRequest.PropertyName, OutObjectRef, &ErrorText);
+
+	if (!ErrorText.IsEmpty())
+	{
+		UE_LOG(LogRemoteControl, Error, TEXT("Web Remote Object Access error: %s"), *ErrorText);
+		TUniquePtr<FHttpServerResponse> Response = WebRemoteControlInternalUtils::CreateHttpResponse();
+		WebRemoteControlInternalUtils::CreateUTF8ErrorMessage(ErrorText, Response->Body);
+		InCompleteCallback(MoveTemp(Response));
+		return false;
+	}
+
+	return true;
+}
+}
+
+TUniquePtr<FHttpServerResponse> WebRemoteControlInternalUtils::CreateHttpResponse(EHttpServerResponseCodes InResponseCode)
+{
+	TUniquePtr<FHttpServerResponse> Response = MakeUnique<FHttpServerResponse>();
+	AddCORSHeaders(Response.Get());
+	AddContentTypeHeaders(Response.Get(), TEXT("application/json"));
+	Response->Code = InResponseCode;
+	return Response;
+}
+
+void WebRemoteControlInternalUtils::CreateUTF8ErrorMessage(const FString& InMessage, TArray<uint8>& OutUTF8Message)
+{
+	WebRemoteControlUtils::ConvertToUTF8(FString::Printf(TEXT("{ \"errorMessage\": \"%s\" }"), *InMessage), OutUTF8Message);
+}
+
+bool WebRemoteControlInternalUtils::GetStructParametersDelimiters(TConstArrayView<uint8> InTCHARPayload, TMap<FString, FBlockDelimiters>& InOutStructParameters, FString* OutErrorText)
+{
+	typedef UCS2CHAR PayloadCharType;
+	FMemoryReaderView Reader(InTCHARPayload);
+	TSharedRef<TJsonReader<PayloadCharType>> JsonReader = TJsonReader<PayloadCharType>::Create(&Reader);
+
+	EJsonNotation Notation;
+
+	FString ErrorText;
+	// The payload should be an object
+	JsonReader->ReadNext(Notation);
+	if (Notation != EJsonNotation::ObjectStart)
+	{
+		ErrorText = TEXT("Expected json object.");
+	}
+
+	// Mark the start/end of the param object in the payload
+	while (JsonReader->ReadNext(Notation) && ErrorText.IsEmpty())
+	{
+		switch (Notation)
+		{
+			// this should mean we reached the parameters field, record the start and ending offset
+		case EJsonNotation::ObjectStart:
+		case EJsonNotation::ArrayStart:
+			if (FBlockDelimiters* Delimiters = InOutStructParameters.Find(JsonReader->GetIdentifier()))
+			{
+				Delimiters->BlockStart = Reader.Tell() - sizeof(PayloadCharType);
+
+				auto Skip = [&JsonReader, &Notation]() { return Notation == EJsonNotation::ObjectStart ? JsonReader->SkipObject() : JsonReader->SkipArray(); };
+				if (Skip())
+				{
+					Delimiters->BlockEnd = Reader.Tell();
+				}
+				else
+				{
+					ErrorText = FString::Printf(TEXT("%s object improperly formatted."), *JsonReader->GetIdentifier());
+				}
+			}
+			break;
+		case EJsonNotation::Error:
+			ErrorText = JsonReader->GetErrorMessage();
+			break;
+		default:
+			// Ignore any other fields
+			break;
+		}
+	}
+
+	if (!ErrorText.IsEmpty())
+	{
+		UE_LOG(LogRemoteControl, Error, TEXT("Web Remote Control deserialization error: %s"), *ErrorText);
+		if (OutErrorText)
+		{
+			*OutErrorText = MoveTemp(ErrorText);
+		}
+		return false;
+	}
+
+	return true;
+}
+
+bool WebRemoteControlInternalUtils::GetBatchRequestStructDelimiters(TConstArrayView<uint8> InTCHARPayload, TMap<int32, FBlockDelimiters>& OutStructParameters, FString* OutErrorText)
+{
+	typedef UCS2CHAR PayloadCharType;
+	FMemoryReaderView Reader(InTCHARPayload);
+	TSharedRef<TJsonReader<PayloadCharType>> JsonReader = TJsonReader<PayloadCharType>::Create(&Reader);
+
+	EJsonNotation Notation;
+
+	FString ErrorText;
+	// The payload should be an object
+	JsonReader->ReadNext(Notation);
+	if (Notation != EJsonNotation::ObjectStart)
+	{
+		ErrorText = TEXT("Expected json object.");
+	}
+
+	int32 LastEncounteredRequestId = -1;
+
+	// Mark the start/end of the param object in the payload
+	while (JsonReader->ReadNext(Notation) && ErrorText.IsEmpty())
+	{
+		switch (Notation)
+		{
+
+		case EJsonNotation::Number:
+			if (JsonReader->GetIdentifier() == TEXT("RequestId"))
+			{
+				LastEncounteredRequestId = JsonReader->GetValueAsNumber();
+			}
+			break;
+		case EJsonNotation::ObjectStart:
+			if (JsonReader->GetIdentifier() == TEXT("Body"))
+			{
+				FBlockDelimiters& Delimiters = OutStructParameters.Add(LastEncounteredRequestId, FBlockDelimiters{});
+				Delimiters.BlockStart = Reader.Tell() - sizeof(PayloadCharType);
+				if (JsonReader->SkipObject())
+				{
+					Delimiters.BlockEnd = Reader.Tell();
+				}
+				else
+				{
+					ErrorText = FString::Printf(TEXT("%s object improperly formatted."), *JsonReader->GetIdentifier());
+				}
+			}
+			break;
+		case EJsonNotation::Error:
+			ErrorText = JsonReader->GetErrorMessage();
+			break;
+		default:
+			// Ignore any other fields
+			break;
+		}
+	}
+
+	if (!ErrorText.IsEmpty())
+	{
+		UE_LOG(LogRemoteControl, Error, TEXT("Web Remote Control deserialization error: %s"), *ErrorText);
+		if (OutErrorText)
+		{
+			*OutErrorText = MoveTemp(ErrorText);
+		}
+		return false;
+	}
+
+	return true;
+}
+
+bool WebRemoteControlInternalUtils::ValidateContentType(const FHttpServerRequest& InRequest, FString InContentType, const FHttpResultCallback& InCompleteCallback)
+{
+	FString ErrorText;
+	if (!IsRequestContentType(InRequest, MoveTemp(InContentType), &ErrorText))
+	{
+		TUniquePtr<FHttpServerResponse> Response = CreateHttpResponse();
+		CreateUTF8ErrorMessage(ErrorText, Response->Body);
+		InCompleteCallback(MoveTemp(Response));
+		return false;
+	}
+	return true;
+}
+
+void WebRemoteControlInternalUtils::AddContentTypeHeaders(FHttpServerResponse* InOutResponse, FString InContentType)
+{
+	InOutResponse->Headers.Add(TEXT("content-type"), { MoveTemp(InContentType) });
+}
+
+void WebRemoteControlInternalUtils::AddWrappedRequestHeader(FHttpServerRequest& Request)
+{
+	Request.Headers.Add(WrappedRequestHeader, {});
+}
+
+bool WebRemoteControlInternalUtils::IsWrappedRequest(const FHttpServerRequest& Request)
+{
+	return Request.Headers.Contains(WrappedRequestHeader);
+}
+
+void WebRemoteControlInternalUtils::AddCORSHeaders(FHttpServerResponse* InOutResponse)
+{
+	check(InOutResponse != nullptr);
+	InOutResponse->Headers.Add(TEXT("Access-Control-Allow-Origin"), { TEXT("*") });
+	InOutResponse->Headers.Add(TEXT("Access-Control-Allow-Methods"), { TEXT("PUT, POST, GET, OPTIONS") });
+	InOutResponse->Headers.Add(TEXT("Access-Control-Allow-Headers"), { TEXT("Origin, X-Requested-With, Content-Type, Accept") });
+	InOutResponse->Headers.Add(TEXT("Access-Control-Max-Age"), { TEXT("600") });
+}
+
+bool WebRemoteControlInternalUtils::IsRequestContentType(const FHttpServerRequest& InRequest, const FString& InContentType, FString* OutErrorText)
+{
+	if (const TArray<FString>* ContentTypeHeaders = InRequest.Headers.Find(TEXT("Content-Type")))
+	{
+		if (ContentTypeHeaders->Num() > 0 && (*ContentTypeHeaders)[0] == InContentType)
+		{
+			return true;
+		}
+	}
+
+	if (OutErrorText)
+	{
+		*OutErrorText = FString::Printf(TEXT("Request content type must be %s"), *InContentType);
+	}
+	return false;
+}
diff --git a/Engine/Plugins/VirtualProduction/RemoteControl/Source/WebRemoteControl/Private/WebRemoteControlInternalUtils.h b/Engine/Plugins/VirtualProduction/RemoteControl/Source/WebRemoteControl/Private/WebRemoteControlInternalUtils.h
new file mode 100644
--- /dev/null
+++ b/Engine/Plugins/VirtualProduction/RemoteControl/Source/WebRemoteControl/Private/WebRemoteControlInternalUtils.h
@@ -0,0 +1,285 @@
+// Copyright Epic Games, Inc. All Rights Reserved.
+
+#pragma once
+
+#include "CoreMinimal.h"
+#include "IStructSerializerBackend.h"
+#include "Serialization/RCJsonStructSerializerBackend.h"
+#include "Serialization/RCJsonStructDeserializerBackend.h"
+#include "HttpServerResponse.h"
+#include "HttpServerRequest.h"
+#include "Serialization/MemoryReader.h"
+#include "StructDeserializer.h"
+#include "StructSerializer.h"
+#include "RemoteControlRequest.h"
+#include "Templates/UnrealTypeTraits.h"
+#include "WebRemoteControlUtils.h"
+
+
+namespace RemotePayloadSerializer
+{
+	/**
+	 * Replaces the first occurrence of a string in a TCHAR binary payload.
+	 */
+	void ReplaceFirstOccurence(TConstArrayView<uint8> InPayload, const FString& From, const FString& To, TArray<uint8>& OutModifiedPayload);
+
+	/**
+	 * Converts a string verb to the enum representation.
+	 */
+	EHttpServerRequestVerbs ParseHttpVerb(FName InVerb);
+
+	/**
+	 * Unwrap a request wrapper, while copying headers and http version from a template request if available.
+	 */
+	TSharedRef<FHttpServerRequest> UnwrapHttpRequest(const FRCRequestWrapper& Wrapper, const FHttpServerRequest* TemplateRequest = nullptr);
+
+	/**
+	 * @note This will serialize in ANSI directly.
+	 */
+	void SerializeWrappedCallResponse(int32 RequestId, TUniquePtr<FHttpServerResponse> Response, FMemoryWriter& Writer);
+
+	bool DeserializeCall(const FHttpServerRequest& InRequest, FRCCall& OutCall, const FHttpResultCallback& InCompleteCallback);
+
+	bool SerializeCall(const FRCCall& InCall, TArray<uint8>& OutPayload, bool bOnlyReturn = false);
+
+	bool DeserializeObjectRef(const FHttpServerRequest& InRequest, FRCObjectReference& OutObjectRef, FRCObjectRequest& OutDeserializedRequest, const FHttpResultCallback& InCompleteCallback);
+}
+
+
+namespace WebRemoteControlInternalUtils
+{
+	static const TCHAR* WrappedRequestHeader = TEXT("UE-Wrapped-Request");
+
+	/**
+	 * Construct a default http response with CORS headers.
+	 * @param InResponseCode The response's code. (Defaults to a bad request)
+	 * @return The constructed server response.
+	 */
+	TUniquePtr<FHttpServerResponse> CreateHttpResponse(EHttpServerResponseCodes InResponseCode = EHttpServerResponseCodes::BadRequest);
+
+	/**
+	 * Create a json structure containing an error message.
+	 * @param InMessage The desired error message.
+	 * @param OutUTF8Message The error message wrapped in a json struct, in UTF-8 format.
+	 * Example output:
+	 *	{
+	 *	  errorMessage: "Request content type must be application/json" 
+	 *	}
+	 */
+	void CreateUTF8ErrorMessage(const FString& InMessage, TArray<uint8>& OutUTF8Message);
+
+	/**
+	 * Deserialize a json structure to find the start and end of every struct parameter in the request.
+	 * @param InTCHARPayload The json payload to deserialize.
+	 * @param InOutStructParameters A map of struct parameter names to the start and end of the struct in the payload.
+	 * @param OutErrorText If set, the string pointer will be populated with an error message on error.
+	 * @return Whether the deserialization was successful.
+	 */
+	bool GetStructParametersDelimiters(TConstArrayView<uint8> InTCHARPayload, TMap<FString, FBlockDelimiters>& InOutStructParameters, FString* OutErrorText = nullptr);
+
+	/**
+	 * Deserialize a request into a UStruct.
+	 * @param InTCHARPayload The json payload to deserialize.
+	 * @param InCompleteCallback The callback to call error.
+	 * @param The structure to serialize using the request's content.
+	 * @return Whether the deserialization was successful.
+	 *
+	 * @note InCompleteCallback will be called with an appropriate http response if the deserialization fails.
+	 */
+	template <typename RequestType>
+	UE_NODISCARD bool DeserializeRequestPayload(TConstArrayView<uint8> InTCHARPayload, const FHttpResultCallback* InCompleteCallback, RequestType& OutDeserializedRequest)
+	{
+		FMemoryReaderView Reader(InTCHARPayload);
+		FJsonStructDeserializerBackend DeserializerBackend(Reader);
+		if (!FStructDeserializer::Deserialize(&OutDeserializedRequest, *RequestType::StaticStruct(), DeserializerBackend, FStructDeserializerPolicies()))
+		{
+			if (InCompleteCallback)
+			{
+				TUniquePtr<FHttpServerResponse> Response = CreateHttpResponse();
+				CreateUTF8ErrorMessage(TEXT("Unable to deserialize request."), Response->Body);
+				(*InCompleteCallback)(MoveTemp(Response));
+			}
+			return false;
+		}
+
+		if (!GetStructParametersDelimiters(InTCHARPayload, OutDeserializedRequest.GetStructParameters(), nullptr))
+		{
+			if (InCompleteCallback)
+			{
+				TUniquePtr<FHttpServerResponse> Response = CreateHttpResponse();
+				CreateUTF8ErrorMessage(TEXT("Unable to deserialize request."), Response->Body);
+				(*InCompleteCallback)(MoveTemp(Response));
+			}
+			return false;
+		}
+
+		return true;
+	}
+
+	/**
+	 * Deserialize a wrapped request into a wrapper struct.
+	 * @param InTCHARPayload The json payload to deserialize.
+	 * @param InCompleteCallback The callback to call error.
+	 * @param The wrapper structure to populate with the request's content.
+	 * @return Whether the deserialization was successful.
+	 */
+	UE_NODISCARD inline bool DeserializeWrappedRequestPayload(TConstArrayView<uint8> InTCHARPayload, const FHttpResultCallback* InCompleteCallback, FRCRequestWrapper& Wrapper)
+	{
+		if (!DeserializeRequestPayload(InTCHARPayload, InCompleteCallback, Wrapper))
+		{
+			return false;
+		}
+
+		FBlockDelimiters& BodyDelimiters = Wrapper.GetParameterDelimiters(FRCRequestWrapper::BodyLabel());
+		if (BodyDelimiters.BlockStart != BodyDelimiters.BlockEnd)
+		{
+			Wrapper.TCHARBody = InTCHARPayload.Slice(BodyDelimiters.BlockStart, BodyDelimiters.BlockEnd - BodyDelimiters.BlockStart);
+		}
+
+		return true;
+	}
+
+	/**
+	 * Get the struct delimiters for all the batched requests.
+	 * @param InTCHARPayload The json payload to deserialize.
+	 * @param OutStructParameters A mapping of Request Id to their respective struct delimiters.
+	 * @param OutErrorText If set, the string pointer will be populated with an error message on error.
+	 * @return Whether the delimiters were able to be found.
+	 */
+	UE_NODISCARD bool GetBatchRequestStructDelimiters(TConstArrayView<uint8> InTCHARPayload, TMap<int32, FBlockDelimiters>& OutStructParameters, FString* OutErrorText = nullptr);
+	
+	/**
+	 * Specialization of DeserializeRequestPayload that handles Batch requests.
+	 * This will populate the TCHARBody of all the wrapped requests.
+	 * @param InTCHARPayload The json payload to deserialize.
+	 * @param InCompleteCallback The callback to call error.
+	 * @param The structure to serialize using the request's content.
+	 * @return Whether the deserialization was successful.
+	 *
+	 * @note InCompleteCallback will be called with an appropriate http response if the deserialization fails.
+	 */
+	template <>
+	UE_NODISCARD inline bool DeserializeRequestPayload(TConstArrayView<uint8> InTCHARPayload, const FHttpResultCallback* InCompleteCallback, FRCBatchRequest& OutDeserializedRequest)
+	{
+		FMemoryReaderView Reader(InTCHARPayload);
+		FJsonStructDeserializerBackend DeserializerBackend(Reader);
+		
+		if (!FStructDeserializer::Deserialize(&OutDeserializedRequest, *FRCBatchRequest::StaticStruct(), DeserializerBackend, FStructDeserializerPolicies()))
+		{
+			if (InCompleteCallback)
+			{
+				TUniquePtr<FHttpServerResponse> Response = CreateHttpResponse();
+				CreateUTF8ErrorMessage(TEXT("Unable to deserialize request."), Response->Body);
+				(*InCompleteCallback)(MoveTemp(Response));
+			}
+			return false;
+		}
+
+		TMap<int32, FBlockDelimiters> Delimiters;
+		if (!GetBatchRequestStructDelimiters(InTCHARPayload, Delimiters, nullptr))
+		{
+			if (InCompleteCallback)
+			{
+				TUniquePtr<FHttpServerResponse> Response = CreateHttpResponse();
+				CreateUTF8ErrorMessage(TEXT("Unable to deserialize request."), Response->Body);
+				(*InCompleteCallback)(MoveTemp(Response));
+			}
+			return false;
+		}
+
+		for (FRCRequestWrapper& Wrapper : OutDeserializedRequest.Requests)
+		{
+			if (FBlockDelimiters* BodyDelimiters = Delimiters.Find(Wrapper.RequestId))
+			{
+				Wrapper.GetParameterDelimiters(FRCRequestWrapper::BodyLabel()) = MoveTemp(*BodyDelimiters);
+				Wrapper.TCHARBody = InTCHARPayload.Slice(BodyDelimiters->BlockStart, BodyDelimiters->BlockEnd - BodyDelimiters->BlockStart);
+			}
+		}
+
+		return true;
+	}
+
+	/**
+	 * Adds a header indicating that the request is a wrapped request that originated from the engine itself.
+	 */
+	void AddWrappedRequestHeader(FHttpServerRequest& Request);
+
+	/**
+	 * Get whether the request is a wrapped request.
+	 */
+	bool IsWrappedRequest(const FHttpServerRequest& Request);
+
+	/**
+	 * Deserialize a request into a UStruct.
+	 * @param InRequest The incoming http request.
+	 * @param InCompleteCallback The callback to call error.
+	 * @param The structure to serialize using the request's content.
+	 * @return Whether the deserialization was successful. 
+	 * 
+	 * @note InCompleteCallback will be called with an appropriate http response if the deserialization fails.
+	 */
+	template <typename RequestType>
+	UE_NODISCARD bool DeserializeRequest(const FHttpServerRequest& InRequest, const FHttpResultCallback* InCompleteCallback, RequestType& OutDeserializedRequest)
+	{
+		static_assert(TIsDerivedFrom<RequestType, FRCRequest>::IsDerived, "Argument OutDeserializedRequest must derive from FRCRequest");
+		
+		if (IsWrappedRequest(InRequest))
+		{
+			// If the request is wrapped, the body should already be encoded in UCS2.
+			OutDeserializedRequest.TCHARBody = InRequest.Body;
+		}
+		
+		if (!OutDeserializedRequest.TCHARBody.Num())
+		{
+			WebRemoteControlUtils::ConvertToTCHAR(InRequest.Body, OutDeserializedRequest.TCHARBody);
+		}
+
+		return DeserializeRequestPayload(OutDeserializedRequest.TCHARBody, InCompleteCallback, OutDeserializedRequest);
+	}
+
+	/**
+	 * Validate a content-type.
+	 * @param InRequest The incoming http request.
+	 * @param InContentType The callback to call error.
+	 * @param InCompleteCallback The callback to call error.
+	 * @return Whether the content type was valid or not.
+	 * 
+	 * @note InCompleteCallback will be called with an appropriate http response if the content type is not valid.
+	 */
+	UE_NODISCARD bool ValidateContentType(const FHttpServerRequest& InRequest, FString InContentType, const FHttpResultCallback& InCompleteCallback);
+
+	/**
+	 * Add the desired content type to the http response headers.
+	 * @param InResponse The response to add the content type to.
+	 * @param InContentType The content type header to add.
+	 */
+	void AddContentTypeHeaders(FHttpServerResponse* InOutResponse, FString InContentType);
+	
+	/**
+	* Add CORS headers to a http response.
+	* @param InOutResponse The http response to add the CORS headers to.
+	*/
+	void AddCORSHeaders(FHttpServerResponse* InOutResponse);
+
+	/**
+	 * Validate a request's content type.
+	 * @param InRequest The request to validate the content type on.
+	 * @param InContentType The target content type.
+	 * @param OutErrorText If set, the string pointer will be populated with an error message on error.
+	 * @return Whether or not the content type matches the target content type.
+	 */
+	bool IsRequestContentType(const FHttpServerRequest& InRequest, const FString& InContentType, FString* OutErrorText);
+
+	/**
+	 * Serialize a struct on scope.
+	 * @param Struct the struct on scope to serialize.
+	 * @param Writer the memory archive to write to.
+	 */
+	template <typename SerializerBackendType = FRCJsonStructSerializerBackend>
+	void SerializeStructOnScope(const FStructOnScope& Struct, FMemoryWriter& Writer)
+	{
+		static_assert(TIsDerivedFrom<SerializerBackendType, IStructSerializerBackend>::IsDerived, "SerializerBackendType must inherit from IStructSerializerBackend.");
+		SerializerBackendType SerializerBackend(Writer);
+		FStructSerializer::Serialize(Struct.GetStructMemory(), *(UScriptStruct*)Struct.GetStruct(), SerializerBackend, FStructSerializerPolicies());
+	}
+}
diff --git a/Engine/Plugins/VirtualProduction/RemoteControl/Source/WebRemoteControl/Private/WebRemoteControlUtils.cpp b/Engine/Plugins/VirtualProduction/RemoteControl/Source/WebRemoteControl/Private/WebRemoteControlUtils.cpp
--- a/Engine/Plugins/VirtualProduction/RemoteControl/Source/WebRemoteControl/Private/WebRemoteControlUtils.cpp
+++ b/Engine/Plugins/VirtualProduction/RemoteControl/Source/WebRemoteControl/Private/WebRemoteControlUtils.cpp
@@ -1,240 +1,8 @@
 // Copyright Epic Games, Inc. All Rights Reserved.
 
 #include "WebRemoteControlUtils.h"
-#include "HttpServerRequest.h"
-#include "Misc/Base64.h"
-#include "Serialization/JsonReader.h"
-#include "UObject/StructOnScope.h"
-
-
-namespace RemotePayloadSerializer
-{
-	FName NAME_Get = "GET";
-	FName NAME_Put = "PUT";
-	FName NAME_Post = "POST";
-	FName NAME_Patch = "PATCH";
-	FName NAME_Delete = "DELETE";
-	FName NAME_Options = "OPTIONS";
-
-	typedef UCS2CHAR PayloadCharType;
-	const FName ReturnValuePropName(TEXT("ReturnValue"));
-
-	void ReplaceFirstOccurence(TConstArrayView<uint8> InPayload, const FString& From, const FString& To, TArray<uint8>& OutModifiedPayload)
-	{
-		if (TCHAR* PropValuePtr = FCString::Stristr((TCHAR*)InPayload.GetData(), *From))
-		{
-			OutModifiedPayload.Reserve(InPayload.Num() + (To.Len() - (From.Len()) * sizeof(TCHAR)));
-
-			int32 StartIndex = UE_PTRDIFF_TO_INT32(PropValuePtr - (TCHAR*)InPayload.GetData()) * sizeof(TCHAR);
-			int32 EndIndex = StartIndex + From.Len() * sizeof(TCHAR);
-			TConstArrayView<uint8> RestOfPayload = InPayload.Slice(EndIndex, InPayload.Num() - EndIndex);
-			OutModifiedPayload.Append(InPayload.GetData(), StartIndex);
-			OutModifiedPayload.Append((uint8*)*To, To.Len() * sizeof(TCHAR));
-			OutModifiedPayload.Append(RestOfPayload.GetData(), RestOfPayload.Num());
-		}
-	}
-
-	EHttpServerRequestVerbs ParseHttpVerb(FName InVerb)
-	{
-		if (InVerb == NAME_Get)
-		{
-			return EHttpServerRequestVerbs::VERB_GET;
-		}
-		else if (InVerb == NAME_Post)
-		{
-			return EHttpServerRequestVerbs::VERB_POST;
-		}
-		else if (InVerb == NAME_Put)
-		{
-			return EHttpServerRequestVerbs::VERB_PUT;
-		}
-		else if (InVerb == NAME_Patch)
-		{
-			return EHttpServerRequestVerbs::VERB_PATCH;
-		}
-		else if (InVerb == NAME_Delete)
-		{
-			return EHttpServerRequestVerbs::VERB_DELETE;
-		}
-		else if (InVerb == NAME_Options)
-		{
-			return EHttpServerRequestVerbs::VERB_OPTIONS;
-		}
-		else
-		{
-			return EHttpServerRequestVerbs::VERB_NONE;
-		}
-	}
-
-	TSharedRef<FHttpServerRequest> UnwrapHttpRequest(const FRCRequestWrapper& Wrapper, const FHttpServerRequest* TemplateRequest)
-	{
-		TSharedRef<FHttpServerRequest> WrappedHttpRequest = MakeShared<FHttpServerRequest>();
-		if (TemplateRequest)
-		{
-			WrappedHttpRequest->HttpVersion = TemplateRequest->HttpVersion;
-			WrappedHttpRequest->Headers = TemplateRequest->Headers;
-		}
-		else
-		{
-			// Assign Request defaults.
-			WrappedHttpRequest->HttpVersion = HttpVersion::EHttpServerHttpVersion::HTTP_VERSION_1_1;
-			WrappedHttpRequest->Headers.Add(TEXT("Content-Type"), { TEXT("application/json") });
-		}
-
-		WebRemoteControlUtils::AddWrappedRequestHeader(*WrappedHttpRequest);
-		WrappedHttpRequest->RelativePath = Wrapper.URL;
-		WrappedHttpRequest->Verb = ParseHttpVerb(Wrapper.Verb);
-		WrappedHttpRequest->Body = Wrapper.TCHARBody;
-
-		return WrappedHttpRequest;
-	}
-
-	void SerializeWrappedCallResponse(int32 RequestId, TUniquePtr<FHttpServerResponse> Response, FMemoryWriter& Writer)
-	{
-		FRCJsonStructSerializerBackend Backend(Writer, FRCJsonStructSerializerBackend::DefaultSerializerFlags);
-		TSharedPtr<TJsonWriter<ANSICHAR>> JsonWriter = TJsonWriter<ANSICHAR>::Create(&Writer);
-		TArray<FString>* ContentTypeHeaders = Response->Headers.Find(TEXT("Content-Type"));
-		const bool bIsBinaryData = ContentTypeHeaders && ContentTypeHeaders->Contains(TEXT("image/png"));
-
-		JsonWriter->WriteObjectStart();
-		JsonWriter->WriteValue(TEXT("RequestId"), RequestId);
-		JsonWriter->WriteValue(TEXT("ResponseCode"), static_cast<int32>(Response->Code));
-		JsonWriter->WriteIdentifierPrefix(TEXT("ResponseBody"));
-	
-		if (Response->Body.Num())
-		{
-			if (!bIsBinaryData)
-			{
-				Writer.Serialize((void*)Response->Body.GetData(), Response->Body.Num());
-			}
-			else
-			{
-				FString Base64String = FString::Printf(TEXT("\"%s\""), *FBase64::Encode(Response->Body));
-				TArray<uint8> WorkingBuffer;
-				WebRemoteControlUtils::ConvertToUTF8(Base64String, WorkingBuffer);
-				Writer.Serialize((void*)WorkingBuffer.GetData(), WorkingBuffer.Num());
-			}
-		}
-		else
-		{
-			JsonWriter->WriteNull();
-		}
-
-		JsonWriter->WriteObjectEnd();
-	}
-
-	bool DeserializeCall(const FHttpServerRequest& InRequest, FRCCall& OutCall, const FHttpResultCallback& InCompleteCallback)
-	{
-		// Create Json reader to read the payload, the payload will already be validated as being in TCHAR
-		FRCCallRequest CallRequest;
-		if (!WebRemoteControlUtils::DeserializeRequest(InRequest, &InCompleteCallback, CallRequest))
-		{
-			return false;
-		}
-
-		FString ErrorText;
-		bool bSuccess = true;
-		if (IRemoteControlModule::Get().ResolveCall(CallRequest.ObjectPath, CallRequest.FunctionName, OutCall.CallRef, &ErrorText))
-		{
-			// Initialize the param struct with default parameters
-			OutCall.bGenerateTransaction = CallRequest.GenerateTransaction;
-			OutCall.ParamStruct = FStructOnScope(OutCall.CallRef.Function.Get());
-
-			// If some parameters were provided, deserialize them
-			const FBlockDelimiters& ParametersDelimiters = CallRequest.GetStructParameters().FindChecked(FRCCallRequest::ParametersLabel());
-			if (ParametersDelimiters.BlockStart > 0)
-			{
-				FMemoryReader Reader(CallRequest.TCHARBody);
-				Reader.Seek(ParametersDelimiters.BlockStart);
-				Reader.SetLimitSize(ParametersDelimiters.BlockEnd);
-
-				FRCJsonStructDeserializerBackend Backend(Reader);
-				if (!FStructDeserializer::Deserialize((void*)OutCall.ParamStruct.GetStructMemory(), *const_cast<UStruct*>(OutCall.ParamStruct.GetStruct()), Backend, FStructDeserializerPolicies()))
-				{
-					ErrorText = TEXT("Parameters object improperly formatted.");
-					bSuccess = false;
-				}
-			}
-		}
-		else
-		{
-			bSuccess = false;
-		}
-
-		if (!bSuccess)
-		{
-			UE_LOG(LogRemoteControl, Error, TEXT("Web Remote Call deserialization error: %s"), *ErrorText);
-			TUniquePtr<FHttpServerResponse> Response = WebRemoteControlUtils::CreateHttpResponse();
-			WebRemoteControlUtils::CreateUTF8ErrorMessage(ErrorText, Response->Body);
-			InCompleteCallback(MoveTemp(Response));
-		}
-
-		return bSuccess;
-	}
-
-	bool SerializeCall(const FRCCall& InCall, TArray<uint8>& OutPayload, bool bOnlyReturn)
-	{
-		FMemoryWriter Writer(OutPayload);
-		TSharedPtr<TJsonWriter<PayloadCharType>> JsonWriter;
-
-		if (!bOnlyReturn)
-		{
-			// Create Json writer to write the payload, if we do not just write the return value back.
-			JsonWriter = TJsonWriter<PayloadCharType>::Create(&Writer);
-
-			JsonWriter->WriteObjectStart();
-			JsonWriter->WriteValue(TEXT("ObjectPath"), InCall.CallRef.Object->GetPathName());
-			JsonWriter->WriteValue(TEXT("FunctionName"), InCall.CallRef.Function->GetFName().ToString());
-			JsonWriter->WriteIdentifierPrefix(TEXT("Parameters"));
-		}
-
-		// write the param struct
-		FRCJsonStructSerializerBackend Backend(Writer, FRCJsonStructSerializerBackend::DefaultSerializerFlags);
-		FStructSerializerPolicies Policies;
-
-		if (bOnlyReturn)
-		{
-			Policies.PropertyFilter = [](const FProperty* CurrentProp, const FProperty* ParentProperty) -> bool
-			{
-				return CurrentProp->HasAnyPropertyFlags(CPF_ReturnParm | CPF_OutParm) || ParentProperty != nullptr;
-			};
-		}
-
-		FStructSerializer::Serialize((void*)InCall.ParamStruct.GetStructMemory(), *const_cast<UStruct*>(InCall.ParamStruct.GetStruct()), Backend, Policies);
-
-		if (!bOnlyReturn)
-		{
-			JsonWriter->WriteObjectEnd();
-		}
-
-		return true;
-	}
-
-	bool DeserializeObjectRef(const FHttpServerRequest& InRequest, FRCObjectReference& OutObjectRef, FRCObjectRequest& OutDeserializedRequest, const FHttpResultCallback& InCompleteCallback)
-	{
-		if (!WebRemoteControlUtils::DeserializeRequest(InRequest, &InCompleteCallback, OutDeserializedRequest))
-		{
-			return false;
-		}
-
-		FString ErrorText;
-
-		// If we properly identified the object path, property name and access type as well as identified the starting / end position for property value
-		// resolve the object reference
-		IRemoteControlModule::Get().ResolveObject(OutDeserializedRequest.GetAccessValue(), OutDeserializedRequest.ObjectPath, OutDeserializedRequest.PropertyName, OutObjectRef, &ErrorText);
-
-		if (!ErrorText.IsEmpty())
-		{
-			UE_LOG(LogRemoteControl, Error, TEXT("Web Remote Object Access error: %s"), *ErrorText);
-			TUniquePtr<FHttpServerResponse> Response = WebRemoteControlUtils::CreateHttpResponse();
-			WebRemoteControlUtils::CreateUTF8ErrorMessage(ErrorText, Response->Body);
-			InCompleteCallback(MoveTemp(Response));
-			return false;
-		}
-
-		return true;
-	}
-}
+#include "Serialization/RCJsonStructSerializerBackend.h"
+#include "Serialization/RCJsonStructDeserializerBackend.h"
 
 void WebRemoteControlUtils::ConvertToTCHAR(TConstArrayView<uint8> InUTF8Payload, TArray<uint8>& OutTCHARPayload)
 {
@@ -257,198 +25,12 @@ void WebRemoteControlUtils::ConvertToUTF8(const FString& InString, TArray<uint8>
 	FTCHARToUTF8_Convert::Convert((ANSICHAR*)(OutUTF8Payload.GetData() + StartIndex), (OutUTF8Payload.Num() - StartIndex) / sizeof(ANSICHAR), *InString, InString.Len());
 }
 
-TUniquePtr<FHttpServerResponse> WebRemoteControlUtils::CreateHttpResponse(EHttpServerResponseCodes InResponseCode)
+TSharedRef<IStructSerializerBackend> WebRemoteControlUtils::CreateJsonSerializerBackend(FMemoryWriter& Writer)
 {
-	TUniquePtr<FHttpServerResponse> Response = MakeUnique<FHttpServerResponse>();
-	AddCORSHeaders(Response.Get());
-	AddContentTypeHeaders(Response.Get(), TEXT("application/json"));
-	Response->Code = InResponseCode;
-	return Response;
+	return MakeShared<FRCJsonStructSerializerBackend>(Writer, FRCJsonStructSerializerBackend::DefaultSerializerFlags);
 }
 
-void WebRemoteControlUtils::CreateUTF8ErrorMessage(const FString& InMessage, TArray<uint8>& OutUTF8Message)
+TSharedRef<IStructDeserializerBackend> WebRemoteControlUtils::CreateJsonDeserializerBackend(FMemoryReaderView& Reader)
 {
-	ConvertToUTF8(FString::Printf(TEXT("{ \"errorMessage\": \"%s\" }"), *InMessage), OutUTF8Message);
-}
-
-bool WebRemoteControlUtils::GetStructParametersDelimiters(TConstArrayView<uint8> InTCHARPayload, TMap<FString, FBlockDelimiters>& InOutStructParameters, FString* OutErrorText)
-{
-	typedef UCS2CHAR PayloadCharType;
-	FMemoryReaderView Reader(InTCHARPayload);
-	TSharedRef<TJsonReader<PayloadCharType>> JsonReader = TJsonReader<PayloadCharType>::Create(&Reader);
-
-	EJsonNotation Notation;
-
-	FString ErrorText;
-	// The payload should be an object
-	JsonReader->ReadNext(Notation);
-	if (Notation != EJsonNotation::ObjectStart)
-	{
-		ErrorText = TEXT("Expected json object.");
-	}
-
-	// Mark the start/end of the param object in the payload
-	while (JsonReader->ReadNext(Notation) && ErrorText.IsEmpty())
-	{
-		switch (Notation)
-		{
-			// this should mean we reached the parameters field, record the start and ending offset
-		case EJsonNotation::ObjectStart:
-		case EJsonNotation::ArrayStart:
-			if (FBlockDelimiters* Delimiters = InOutStructParameters.Find(JsonReader->GetIdentifier()))
-			{
-				Delimiters->BlockStart = Reader.Tell() - sizeof(PayloadCharType);
-				
-				auto Skip = [&JsonReader, &Notation] () { return Notation == EJsonNotation::ObjectStart ? JsonReader->SkipObject() : JsonReader->SkipArray(); };
-				if (Skip())
-				{
-					Delimiters->BlockEnd = Reader.Tell();
-				}
-				else
-				{
-					ErrorText = FString::Printf(TEXT("%s object improperly formatted."), *JsonReader->GetIdentifier());
-				}
-			}
-			break;
-		case EJsonNotation::Error:
-			ErrorText = JsonReader->GetErrorMessage();
-			break;
-		default:
-			// Ignore any other fields
-			break;
-		}
-	}
-
-	if (!ErrorText.IsEmpty())
-	{
-		UE_LOG(LogRemoteControl, Error, TEXT("Web Remote Control deserialization error: %s"), *ErrorText);
-		if (OutErrorText)
-		{
-			*OutErrorText = MoveTemp(ErrorText);
-		}
-		return false;
-	}
-
-	return true;
-}
-
-bool WebRemoteControlUtils::GetBatchRequestStructDelimiters(TConstArrayView<uint8> InTCHARPayload, TMap<int32, FBlockDelimiters>& OutStructParameters, FString* OutErrorText)
-{
-	typedef UCS2CHAR PayloadCharType;
-	FMemoryReaderView Reader(InTCHARPayload);
-	TSharedRef<TJsonReader<PayloadCharType>> JsonReader = TJsonReader<PayloadCharType>::Create(&Reader);
-
-	EJsonNotation Notation;
-
-	FString ErrorText;
-	// The payload should be an object
-	JsonReader->ReadNext(Notation);
-	if (Notation != EJsonNotation::ObjectStart)
-	{
-		ErrorText = TEXT("Expected json object.");
-	}
-	
-	int32 LastEncounteredRequestId = -1;
-
-	// Mark the start/end of the param object in the payload
-	while (JsonReader->ReadNext(Notation) && ErrorText.IsEmpty())
-	{
-		switch (Notation)
-		{
-
-		case EJsonNotation::Number:
-			if (JsonReader->GetIdentifier() == TEXT("RequestId"))
-			{
-				LastEncounteredRequestId = JsonReader->GetValueAsNumber();
-			}
-			break;
-		case EJsonNotation::ObjectStart:
-			if (JsonReader->GetIdentifier() == TEXT("Body"))
-			{
-				FBlockDelimiters& Delimiters = OutStructParameters.Add(LastEncounteredRequestId, FBlockDelimiters{});
-				Delimiters.BlockStart = Reader.Tell() - sizeof(PayloadCharType);
-				if (JsonReader->SkipObject())
-				{
-					Delimiters.BlockEnd = Reader.Tell();
-				}
-				else
-				{
-					ErrorText = FString::Printf(TEXT("%s object improperly formatted."), *JsonReader->GetIdentifier());
-				}
-			}
-			break;
-		case EJsonNotation::Error:
-			ErrorText = JsonReader->GetErrorMessage();
-			break;
-		default:
-			// Ignore any other fields
-			break;
-		}
-	}
-
-	if (!ErrorText.IsEmpty())
-	{
-		UE_LOG(LogRemoteControl, Error, TEXT("Web Remote Control deserialization error: %s"), *ErrorText);
-		if (OutErrorText)
-		{
-			*OutErrorText = MoveTemp(ErrorText);
-		}
-		return false;
-	}
-
-	return true;
-}
-
-bool WebRemoteControlUtils::ValidateContentType(const FHttpServerRequest& InRequest, FString InContentType, const FHttpResultCallback& InCompleteCallback)
-{
-	FString ErrorText;
-	if (!IsRequestContentType(InRequest, MoveTemp(InContentType), &ErrorText))
-	{
-		TUniquePtr<FHttpServerResponse> Response = CreateHttpResponse();
-		CreateUTF8ErrorMessage(ErrorText, Response->Body);
-		InCompleteCallback(MoveTemp(Response));
-		return false;
-	}
-	return true;
-}
-
-void WebRemoteControlUtils::AddContentTypeHeaders(FHttpServerResponse* InOutResponse, FString InContentType)
-{
-	InOutResponse->Headers.Add(TEXT("content-type"), { MoveTemp(InContentType) });
-}
-
-void WebRemoteControlUtils::AddWrappedRequestHeader(FHttpServerRequest& Request)
-{
-	Request.Headers.Add(WrappedRequestHeader, {});
-}
-
-bool WebRemoteControlUtils::IsWrappedRequest(const FHttpServerRequest& Request)
-{
-	return Request.Headers.Contains(WrappedRequestHeader);
-}
-
-void WebRemoteControlUtils::AddCORSHeaders(FHttpServerResponse* InOutResponse)
-{
-	check(InOutResponse != nullptr);
-	InOutResponse->Headers.Add(TEXT("Access-Control-Allow-Origin"), { TEXT("*") });
-	InOutResponse->Headers.Add(TEXT("Access-Control-Allow-Methods"), { TEXT("PUT, POST, GET, OPTIONS") });
-	InOutResponse->Headers.Add(TEXT("Access-Control-Allow-Headers"), { TEXT("Origin, X-Requested-With, Content-Type, Accept") });
-	InOutResponse->Headers.Add(TEXT("Access-Control-Max-Age"), { TEXT("600") });
-}
-
-bool WebRemoteControlUtils::IsRequestContentType(const FHttpServerRequest& InRequest, const FString& InContentType, FString* OutErrorText)
-{
-	if (const TArray<FString>* ContentTypeHeaders = InRequest.Headers.Find(TEXT("Content-Type")))
-	{
-		if (ContentTypeHeaders->Num() > 0 && (*ContentTypeHeaders)[0] == InContentType)
-		{
-			return true;
-		}
-	}
-
-	if (OutErrorText)
-	{
-		*OutErrorText = FString::Printf(TEXT("Request content type must be %s"), *InContentType);
-	}
-	return false;
+	return MakeShared<FRCJsonStructDeserializerBackend>(Reader);
 }
diff --git a/Engine/Plugins/VirtualProduction/RemoteControl/Source/WebRemoteControl/Private/WebRemoteControlUtils.h b/Engine/Plugins/VirtualProduction/RemoteControl/Source/WebRemoteControl/Private/WebRemoteControlUtils.h
deleted file mode 100644
diff --git a/Engine/Plugins/VirtualProduction/RemoteControl/Source/WebRemoteControl/Private/WebSocketMessageHandler.cpp b/Engine/Plugins/VirtualProduction/RemoteControl/Source/WebRemoteControl/Private/WebSocketMessageHandler.cpp
--- a/Engine/Plugins/VirtualProduction/RemoteControl/Source/WebRemoteControl/Private/WebSocketMessageHandler.cpp
+++ b/Engine/Plugins/VirtualProduction/RemoteControl/Source/WebRemoteControl/Private/WebSocketMessageHandler.cpp
@@ -13,8 +13,9 @@
 #include "RemoteControlResponse.h"
 #include "RemoteControlRoute.h"
 #include "RemoteControlReflectionUtils.h"
+#include "RemoteControlWebsocketRoute.h"
 #include "WebRemoteControl.h"
-#include "WebRemoteControlUtils.h"
+#include "WebRemoteControlInternalUtils.h"
 
 static TAutoConsoleVariable<int32> CVarWebRemoteControlFramesBetweenPropertyNotifications(TEXT("WebControl.FramesBetweenPropertyNotifications"), 5, TEXT("The number of frames between sending batches of property notifications."));
 
@@ -286,7 +287,7 @@ void FWebSocketMessageHandler::NotifyPropertyChangedRemotely(const FGuid& Origin
 void FWebSocketMessageHandler::HandleWebSocketPresetRegister(const FRemoteControlWebSocketMessage& WebSocketMessage)
 {
 	FRCWebSocketPresetRegisterBody Body;
-	if (!WebRemoteControlUtils::DeserializeRequestPayload(WebSocketMessage.RequestPayload, nullptr, Body))
+	if (!WebRemoteControlInternalUtils::DeserializeRequestPayload(WebSocketMessage.RequestPayload, nullptr, Body))
 	{
 		return;
 	}
@@ -338,7 +339,7 @@ void FWebSocketMessageHandler::HandleWebSocketPresetRegister(const FRemoteContro
 void FWebSocketMessageHandler::HandleWebSocketPresetUnregister(const FRemoteControlWebSocketMessage& WebSocketMessage)
 {
 	FRCWebSocketPresetRegisterBody Body;
-	if (!WebRemoteControlUtils::DeserializeRequestPayload(WebSocketMessage.RequestPayload, nullptr, Body))
+	if (!WebRemoteControlInternalUtils::DeserializeRequestPayload(WebSocketMessage.RequestPayload, nullptr, Body))
 	{
 		return;
 	}
@@ -604,7 +605,7 @@ void FWebSocketMessageHandler::OnEntitiesModified(URemoteControlPreset* Owner, c
 	}
 	
 	TArray<uint8> Payload;
-	WebRemoteControlUtils::SerializeResponse(FRCPresetEntitiesModifiedEvent{Owner, ModifiedEntities.Array()}, Payload);
+	WebRemoteControlUtils::SerializeMessage(FRCPresetEntitiesModifiedEvent{Owner, ModifiedEntities.Array()}, Payload);
 	BroadcastToListeners(Owner->GetPresetId(), Payload);
 }
 
@@ -695,7 +696,7 @@ void FWebSocketMessageHandler::ProcessAddedProperties()
 		}
 
 		TArray<uint8> Payload;
-		WebRemoteControlUtils::SerializeResponse(FRCPresetFieldsAddedEvent{ Preset->GetFName(), Preset->GetPresetId(), AddedPropertiesDescription }, Payload);
+		WebRemoteControlUtils::SerializeMessage(FRCPresetFieldsAddedEvent{ Preset->GetFName(), Preset->GetPresetId(), AddedPropertiesDescription }, Payload);
 		BroadcastToListeners(Entry.Key, Payload);
 	}
 
@@ -720,7 +721,7 @@ void FWebSocketMessageHandler::ProcessRemovedProperties()
 		ensure(Entry.Value.Key.Num() == Entry.Value.Value.Num());
 		
 		TArray<uint8> Payload;
-		WebRemoteControlUtils::SerializeResponse(FRCPresetFieldsRemovedEvent{ Preset->GetFName(), Preset->GetPresetId(), Entry.Value.Value, Entry.Value.Key }, Payload);
+		WebRemoteControlUtils::SerializeMessage(FRCPresetFieldsRemovedEvent{ Preset->GetFName(), Preset->GetPresetId(), Entry.Value.Value, Entry.Value.Key }, Payload);
 		BroadcastToListeners(Entry.Key, Payload);
 	}
 	
@@ -743,7 +744,7 @@ void FWebSocketMessageHandler::ProcessRenamedFields()
 		}
 
 		TArray<uint8> Payload;
-		WebRemoteControlUtils::SerializeResponse(FRCPresetFieldsRenamedEvent{Preset->GetFName(), Preset->GetPresetId(), Entry.Value}, Payload);
+		WebRemoteControlUtils::SerializeMessage(FRCPresetFieldsRenamedEvent{Preset->GetFName(), Preset->GetPresetId(), Entry.Value}, Payload);
 		BroadcastToListeners(Entry.Key, Payload);
 	}
 
@@ -759,7 +760,7 @@ void FWebSocketMessageHandler::ProcessModifiedMetadata()
 			if (URemoteControlPreset* Preset = IRemoteControlModule::Get().ResolvePreset(Entry))
 			{
 				TArray<uint8> Payload;
-				WebRemoteControlUtils::SerializeResponse(FRCPresetMetadataModified{ Preset }, Payload);
+				WebRemoteControlUtils::SerializeMessage(FRCPresetMetadataModified{ Preset }, Payload);
 				BroadcastToListeners(Entry, Payload);
 			}
 		}
@@ -777,7 +778,7 @@ void FWebSocketMessageHandler::ProcessModifiedPresetLayouts()
 			if (URemoteControlPreset* Preset = IRemoteControlModule::Get().ResolvePreset(Entry))
 			{
 				TArray<uint8> Payload;
-				WebRemoteControlUtils::SerializeResponse(FRCPresetLayoutModified{ Preset }, Payload);
+				WebRemoteControlUtils::SerializeMessage(FRCPresetLayoutModified{ Preset }, Payload);
 				BroadcastToListeners(Entry, Payload);
 			}
 		}
@@ -826,7 +827,7 @@ bool FWebSocketMessageHandler::WritePropertyChangeEventPayload(URemoteControlPre
 		FStructOnScope FieldsChangedEventOnScope = WebSocketMessageHandlerStructUtils::CreatePresetFieldsChangedStructOnScope(InPreset, PropValuesOnScope);
 
 		FMemoryWriter Writer(OutBuffer);
-		WebRemoteControlUtils::SerializeStructOnScope(FieldsChangedEventOnScope, Writer);
+		WebRemoteControlInternalUtils::SerializeStructOnScope(FieldsChangedEventOnScope, Writer);
 	}
 
 	return bHasProperty;
@@ -864,7 +865,7 @@ bool FWebSocketMessageHandler::WriteActorPropertyChangePayload(URemoteControlPre
 	}
 
 	FStructOnScope ActorsModifedOnScope = WebSocketMessageHandlerStructUtils::CreateModifiedActorsStructOnScope(InPreset, ModifiedActorsOnScope);
-	WebRemoteControlUtils::SerializeStructOnScope(ActorsModifedOnScope, InWriter);
+	WebRemoteControlInternalUtils::SerializeStructOnScope(ActorsModifedOnScope, InWriter);
 
 	return bHasProperty;
 }
diff --git a/Engine/Plugins/VirtualProduction/RemoteControl/Source/WebRemoteControl/Public/IWebRemoteControlModule.h b/Engine/Plugins/VirtualProduction/RemoteControl/Source/WebRemoteControl/Public/IWebRemoteControlModule.h
--- a/Engine/Plugins/VirtualProduction/RemoteControl/Source/WebRemoteControl/Public/IWebRemoteControlModule.h
+++ b/Engine/Plugins/VirtualProduction/RemoteControl/Source/WebRemoteControl/Public/IWebRemoteControlModule.h
@@ -7,6 +7,10 @@
 #include "Templates/SharedPointer.h"
 
 DECLARE_MULTICAST_DELEGATE_OneParam(FOnWebServerStarted, uint32 /*Port*/);
+DECLARE_MULTICAST_DELEGATE_OneParam(FOnWebSocketConnectionOpened, FGuid /*ClientId*/);
+DECLARE_MULTICAST_DELEGATE_OneParam(FOnWebSocketConnectionClosed, FGuid /*ClientId*/);
+
+struct FRemoteControlWebsocketRoute;
 
 /**
  * A Remote Control module that allows exposing objects and properties from the editor.
@@ -50,6 +54,35 @@ public:
 	 */
 	virtual FSimpleMulticastDelegate& OnWebSocketServerStopped() = 0;
 
+	/**
+	 * Event triggered when a connection to the websocket server is opened.
+	 */
+	virtual FOnWebSocketConnectionOpened& OnWebSocketConnectionOpened() = 0;
+	
+	/**
+	 * Event triggered when a connection to the websocket server is closed.
+	 */
+	virtual FOnWebSocketConnectionClosed& OnWebSocketConnectionClosed() = 0;
+	
+	/**
+	 * Register a websocket route.
+	 * @param The route to register.
+	 */
+	virtual void RegisterWebsocketRoute(const FRemoteControlWebsocketRoute& Route) = 0;
+	
+	/**
+	 * Unregister a websocket route.
+	 * @param The route to unregister.
+	 */
+	virtual void UnregisterWebsocketRoute(const FRemoteControlWebsocketRoute& Route) = 0;
+	
+	/**
+	 * Send a message through the websocket server.
+	 * @param InTargetClientId ID of the client to send the message to.
+	 * @param InUTF8Payload Payload to send with the message.
+	 */
+	virtual void SendWebsocketMessage(const FGuid& InTargetClientId, const TArray<uint8>&InUTF8Payload) = 0;
+
 	/**
 	 * Set Remote Web Socket Logger
 	 */
diff --git a/Engine/Plugins/VirtualProduction/RemoteControl/Source/WebRemoteControl/Public/RemoteControlWebsocketRoute.h b/Engine/Plugins/VirtualProduction/RemoteControl/Source/WebRemoteControl/Public/RemoteControlWebsocketRoute.h
new file mode 100644
--- /dev/null
+++ b/Engine/Plugins/VirtualProduction/RemoteControl/Source/WebRemoteControl/Public/RemoteControlWebsocketRoute.h
@@ -0,0 +1,32 @@
+// Copyright Epic Games, Inc. All Rights Reserved.
+
+#pragma once
+
+struct FRemoteControlWebSocketMessage
+{
+	FString MessageName;
+	int32 MessageId = -1;
+	FGuid ClientId;
+	TArrayView<uint8> RequestPayload;
+};
+
+DECLARE_DELEGATE_OneParam(FWebSocketMessageDelegate, const FRemoteControlWebSocketMessage& /** Message */);
+
+struct FRemoteControlWebsocketRoute
+{
+	FRemoteControlWebsocketRoute(const FString& InRouteDescription, const FString& InMessageName, const FWebSocketMessageDelegate& InDelegate)
+		: RouteDescription(InRouteDescription)
+		, MessageName(InMessageName)
+		, Delegate(InDelegate)
+	{}
+
+	/** A description of how the route should be used. */
+	FString RouteDescription;
+	/**  The message handled by this route. */
+	FString MessageName;
+	/** The handler called when the route is accessed. */
+	FWebSocketMessageDelegate Delegate;
+
+	friend uint32 GetTypeHash(const FRemoteControlWebsocketRoute& Route) { return GetTypeHash(Route.MessageName); }
+	friend bool operator==(const FRemoteControlWebsocketRoute& LHS, const FRemoteControlWebsocketRoute& RHS) { return LHS.MessageName == RHS.MessageName; }
+};
diff --git a/Engine/Plugins/VirtualProduction/RemoteControl/Source/WebRemoteControl/Public/WebRemoteControlUtils.h b/Engine/Plugins/VirtualProduction/RemoteControl/Source/WebRemoteControl/Public/WebRemoteControlUtils.h
new file mode 100644
--- /dev/null
+++ b/Engine/Plugins/VirtualProduction/RemoteControl/Source/WebRemoteControl/Public/WebRemoteControlUtils.h
@@ -0,0 +1,80 @@
+// Copyright Epic Games, Inc. All Rights Reserved.
+
+#pragma once
+
+#include "Backends/JsonStructDeserializerBackend.h"
+#include "CoreMinimal.h"
+#include "Serialization/MemoryReader.h"
+#include "Serialization/MemoryWriter.h"
+#include "StructDeserializer.h"
+#include "StructSerializer.h"
+
+struct FBlockDelimiters;
+
+namespace WebRemoteControlUtils
+{
+	/**
+	 * Convert a UTF-8 payload to a TCHAR payload.
+	 * @param InUTF8Payload The UTF-8 payload in binary format.
+	 * @param OutTCHARPayload The converted TCHAR output in binary format.
+	 */
+	WEBREMOTECONTROL_API void ConvertToTCHAR(TConstArrayView<uint8> InUTF8Payload, TArray<uint8>& OutTCHARPayload);
+
+	/**
+	 * Convert a TCHAR payload to UTF-8.
+	 * @param InTCHARPayload The TCHAR payload in binary format.
+	 * @param OutUTF8Payload The converted UTF-8 output in binary format.
+	 */
+	WEBREMOTECONTROL_API void ConvertToUTF8(TConstArrayView<uint8> InTCHARPayload, TArray<uint8>& OutUTF8Payload);
+
+	/**
+	 * Convert a FString to UTF-8.
+	 * @param InString The string to be converted.
+	 * @param OutUTF8Payload the converted UTF-8 output in binary format.
+	 */
+	WEBREMOTECONTROL_API void ConvertToUTF8(const FString& InString, TArray<uint8>& OutUTF8Payload);
+
+	/**
+	 * Create a JSON serializer backend for use with other WebRemoteControl functions.
+	 */
+	WEBREMOTECONTROL_API TSharedRef<IStructSerializerBackend> CreateJsonSerializerBackend(FMemoryWriter& Writer);
+
+	/**
+	 * Create a JSON deserializer backend for use with other WebRemoteControl functions.
+	 */
+	WEBREMOTECONTROL_API TSharedRef<IStructDeserializerBackend> CreateJsonDeserializerBackend(FMemoryReaderView& Reader);
+
+	/**
+	 * Serialize a message object into a UTF-8 Payload.
+	 * @param InMessageObject the object to serialize.
+	 * @param OutResponsePayload the resulting UTF-8 payload.
+	 */
+	template <typename MessageType>
+	void SerializeMessage(const MessageType& InMessageObject, TArray<uint8>& OutMessagePayload)
+	{
+		TArray<uint8> WorkingBuffer;
+		FMemoryWriter Writer(WorkingBuffer);
+		TSharedRef<IStructSerializerBackend> SerializerBackend = CreateJsonSerializerBackend(Writer);
+		FStructSerializer::Serialize(InMessageObject, SerializerBackend.Get(), FStructSerializerPolicies());
+		ConvertToUTF8(WorkingBuffer, OutMessagePayload);
+	}
+
+	/**
+	 * Deserialize a message payload into a UStruct.
+	 * @param InTCHARPayload The json payload to deserialize.
+	 * @param The structure to serialize using the message's content.
+	 * @return Whether the deserialization was successful.
+	 */
+	template <typename MessageType>
+	UE_NODISCARD bool DeserializeMessage(TConstArrayView<uint8> InTCHARPayload, MessageType& OutDeserializedMessage)
+	{
+		FMemoryReaderView Reader(InTCHARPayload);
+		TSharedRef<IStructDeserializerBackend> DeserializerBackend = CreateJsonDeserializerBackend(Reader);
+		if (!FStructDeserializer::Deserialize(&OutDeserializedMessage, *MessageType::StaticStruct(), DeserializerBackend.Get(), FStructDeserializerPolicies()))
+		{
+			return false;
+		}
+
+		return true;
+	}
+}
diff --git a/Engine/Plugins/VirtualProduction/RemoteControl/Source/WebRemoteControl/WebRemoteControl.Build.cs b/Engine/Plugins/VirtualProduction/RemoteControl/Source/WebRemoteControl/WebRemoteControl.Build.cs
--- a/Engine/Plugins/VirtualProduction/RemoteControl/Source/WebRemoteControl/WebRemoteControl.Build.cs
+++ b/Engine/Plugins/VirtualProduction/RemoteControl/Source/WebRemoteControl/WebRemoteControl.Build.cs
@@ -10,7 +10,8 @@ public class WebRemoteControl : ModuleRules
 			new string[] {
 				"Core",
 				"CoreUObject",
-				"HTTPServer"
+				"HTTPServer",
+				"Serialization"
 			}
 		);
 
@@ -20,7 +21,6 @@ public class WebRemoteControl : ModuleRules
 				"Networking",
 				"RemoteControl",
 				"RemoteControlCommon",
-				"Serialization",
 				"Sockets",
 				"WebSocketNetworking"
 			}
