Tentative fix for Mac build errors about FRCJsonStructSerializerBackend::DefaultSerializerFlags

diff --git a/Engine/Plugins/VirtualProduction/RemoteControl/Source/WebRemoteControl/Private/WebRemoteControlUtils.cpp b/Engine/Plugins/VirtualProduction/RemoteControl/Source/WebRemoteControl/Private/WebRemoteControlUtils.cpp
--- a/Engine/Plugins/VirtualProduction/RemoteControl/Source/WebRemoteControl/Private/WebRemoteControlUtils.cpp
+++ b/Engine/Plugins/VirtualProduction/RemoteControl/Source/WebRemoteControl/Private/WebRemoteControlUtils.cpp
@@ -27,7 +27,7 @@ void WebRemoteControlUtils::ConvertToUTF8(const FString& InString, TArray<uint8>
 
 TSharedRef<IStructSerializerBackend> WebRemoteControlUtils::CreateJsonSerializerBackend(FMemoryWriter& Writer)
 {
-	return MakeShared<FRCJsonStructSerializerBackend>(Writer, FRCJsonStructSerializerBackend::DefaultSerializerFlags);
+	return MakeShared<FRCJsonStructSerializerBackend>(Writer);
 }
 
 TSharedRef<IStructDeserializerBackend> WebRemoteControlUtils::CreateJsonDeserializerBackend(FMemoryReaderView& Reader)
