Fix remote control not updating viewport in real time

diff --git a/Engine/Plugins/VirtualProduction/RemoteControl/Source/RemoteControl/Private/RemoteControlModule.cpp b/Engine/Plugins/VirtualProduction/RemoteControl/Source/RemoteControl/Private/RemoteControlModule.cpp
--- a/Engine/Plugins/VirtualProduction/RemoteControl/Source/RemoteControl/Private/RemoteControlModule.cpp
+++ b/Engine/Plugins/VirtualProduction/RemoteControl/Source/RemoteControl/Private/RemoteControlModule.cpp
@@ -617,6 +617,8 @@ bool FRemoteControlModule::InvokeCall(FRCCall& InCall, ERCPayloadType InPayloadT
 					if (UActorComponent* Component = Cast<UActorComponent>(InCall.CallRef.Object.Get()))
 					{
 						SnapshotTransactionBuffer(Component->GetOwner());
+						Component->MarkRenderStateDirty();
+						Component->UpdateComponentToWorld();
 					}
 				}
 			}
@@ -929,7 +931,6 @@ bool FRemoteControlModule::SetObjectProperties(const FRCObjectReference& ObjectA
 					MutableObjectReference.Object->Modify();
 				}
 			}
-
 		}
 		else 
 		{
@@ -1003,6 +1004,15 @@ bool FRemoteControlModule::SetObjectProperties(const FRCObjectReference& ObjectA
 				if (OngoingModification->bHasStartedTransaction)
 				{
 					SnapshotTransactionBuffer(MutableObjectReference.Object.Get());
+					FPropertyChangedEvent PropertyEvent(MutableObjectReference.PropertyPathInfo.ToPropertyChangedEvent());
+					PropertyEvent.ChangeType = EPropertyChangeType::Interactive;
+					MutableObjectReference.Object->PostEditChangeProperty(PropertyEvent);
+
+					if (UActorComponent* Component = Cast<UActorComponent>(MutableObjectReference.Object.Get()))
+					{
+						Component->MarkRenderStateDirty();
+						Component->UpdateComponentToWorld();
+					}
 				}
 
 				// Update the world lighting if we're modifying a color.
@@ -1472,6 +1482,7 @@ void FRemoteControlModule::TestOrFinalizeOngoingChange(bool bForceEndChange)
 			if (OngoingModification->Reference.IsType<FRCObjectReference>())
 			{
 				FPropertyChangedEvent PropertyEvent = OngoingModification->Reference.Get<FRCObjectReference>().PropertyPathInfo.ToPropertyChangedEvent();
+				PropertyEvent.ChangeType = EPropertyChangeType::ValueSet;
 				if (UObject* Object = OngoingModification->Reference.Get<FRCObjectReference>().Object.Get())
 				{
 					Object->PostEditChangeProperty(PropertyEvent);
