Fix remote control changes not replicating in PIE

diff --git a/Engine/Plugins/VirtualProduction/RemoteControl/Source/RemoteControl/Private/RemoteControlModule.cpp b/Engine/Plugins/VirtualProduction/RemoteControl/Source/RemoteControl/Private/RemoteControlModule.cpp
--- a/Engine/Plugins/VirtualProduction/RemoteControl/Source/RemoteControl/Private/RemoteControlModule.cpp
+++ b/Engine/Plugins/VirtualProduction/RemoteControl/Source/RemoteControl/Private/RemoteControlModule.cpp
@@ -820,7 +820,7 @@ bool FRemoteControlModule::SetObjectProperties(const FRCObjectReference& ObjectA
 		}
 
 		// If a setter wasn't used, verify if the property should be allowed.
-		bool bObjectInGame = !GIsEditor || (ObjectAccess.Object.IsValid() && ObjectAccess.Object->GetOutermost()->HasAnyPackageFlags(PKG_PlayInEditor));
+		bool bObjectInGame = !GIsEditor;
 		if (!RemoteControlUtil::IsPropertyAllowed(ObjectAccess.Property.Get(), ObjectAccess.Access, bObjectInGame))
 		{
 			return false;
