Potential Security Issue, Login Screen

Have addeda a FHttpRequestHandler as a RequestPreprocessor which checks whether or not an Requests Header has a corresponding ApiKey to allow the execution, using MD5 hashing to obscure it a bit.
This affects Http Calls and Websocket Calls. Changes Wrapped RCWebsocketRequest to allow holding hashed ApiKeys from the outside. Mainly hooked as an additional Handler in WebRemoteControl.

diff --git a/Engine/Plugins/VirtualProduction/RemoteControl/Source/RemoteControlCommon/Public/RemoteControlSettings.h b/Engine/Plugins/VirtualProduction/RemoteControl/Source/RemoteControlCommon/Public/RemoteControlSettings.h
--- a/Engine/Plugins/VirtualProduction/RemoteControl/Source/RemoteControlCommon/Public/RemoteControlSettings.h
+++ b/Engine/Plugins/VirtualProduction/RemoteControl/Source/RemoteControlCommon/Public/RemoteControlSettings.h
@@ -7,6 +7,22 @@
 
 #include "RemoteControlSettings.generated.h"
 
+/**
+ * Passphrase Struct
+ */
+USTRUCT(BlueprintType)
+struct FRCPassphrase
+{
+	GENERATED_BODY()
+	
+	FRCPassphrase(){}
+
+	UPROPERTY(EditAnywhere, Category="Passphrase")
+	FString Identifier;
+	
+	UPROPERTY(EditAnywhere, Category="Passphrase")
+	FString Passphrase;
+};
 
 /**
  * Global remote control settings
@@ -40,6 +56,18 @@ public:
 		return NSLOCTEXT("RemoteControlSettings", "RemoteControlSettingsSection", "Remote Control");
 	}
 
+	virtual TArray<FString> GetHashedPassphrases() const
+	{
+		TArray<FString> OutArray;
+
+		for (const FRCPassphrase& Passphrase : Passphrases)
+		{
+			OutArray.Add(Passphrase.Passphrase);
+		}
+		
+		return OutArray;
+	}
+
 	/**
 	 * Should transactions be generated for events received through protocols (ie. MIDI, DMX etc.)
 	 * Disabling transactions improves performance but will prevent events from being transacted to Multi-User
@@ -60,6 +88,14 @@ public:
 	UPROPERTY(config, EditAnywhere, Category = "Remote Control Web Interface", DisplayName = "Log WebApp requests handle duration")
 	bool bWebAppLogRequestDuration = false;
 
+	/** Whether communication with the Web Interface should only be allowed with an Passphrase */
+	UPROPERTY(Config, EditAnywhere, Category = "Remote Control Web Interface", DisplayName = "Use Passphrase to block Access")
+	bool bUseRemoteControlPassphrase = false;
+
+	/** Whether the User should be warned that Passphrase usage is disabled or now. Initially activated */
+	UPROPERTY(Config, EditAnywhere, Category = "Remote Control Web Interface", DisplayName = "Warn that Passphrase might be disabled ")
+	bool bShowPassphraseDisabledWarning = true;
+
 	/** Whether web server is started automatically. */
 	UPROPERTY(config, EditAnywhere, Category = "Remote Control Web Server")
 	bool bAutoStartWebServer = true;
@@ -86,4 +122,7 @@ public:
 
 	UPROPERTY(config)
 	bool bUseRebindingContext = true;
+
+	UPROPERTY(config, EditAnywhere, Category = "Remote Control Web Interface", DisplayName = "Remote Control Web Interface Passphrase")
+	TArray<FRCPassphrase> Passphrases = {};
 };
diff --git a/Engine/Plugins/VirtualProduction/RemoteControl/Source/RemoteControlUI/Private/RemoteControlUIModule.cpp b/Engine/Plugins/VirtualProduction/RemoteControl/Source/RemoteControlUI/Private/RemoteControlUIModule.cpp
--- a/Engine/Plugins/VirtualProduction/RemoteControl/Source/RemoteControlUI/Private/RemoteControlUIModule.cpp
+++ b/Engine/Plugins/VirtualProduction/RemoteControl/Source/RemoteControlUI/Private/RemoteControlUIModule.cpp
@@ -25,6 +25,7 @@
 #include "ScopedTransaction.h"
 #include "Styling/AppStyle.h"
 #include "Textures/SlateIcon.h"
+#include "UI/Customizations/FPassphraseCustomization.h"
 #include "UI/Customizations/RemoteControlEntityCustomization.h"
 #include "UI/SRCPanelExposedField.h"
 #include "UI/SRCPanelExposedActor.h"
@@ -137,7 +138,7 @@ namespace RemoteControlUIModule
 		{
 			FRemoteControlProperty::StaticStruct()->GetFName(),
 			FRemoteControlFunction::StaticStruct()->GetFName(),
-			FRemoteControlActor::StaticStruct()->GetFName()
+			FRemoteControlActor::StaticStruct()->GetFName(),
 		};
 
 		return CustomizedStructNames;
@@ -744,6 +745,8 @@ void FRemoteControlUIModule::RegisterStructCustomizations()
 	{
 		PropertyEditorModule.RegisterCustomClassLayout(Name, FOnGetDetailCustomizationInstance::CreateStatic(&FRemoteControlEntityCustomization::MakeInstance));
 	}
+	
+	PropertyEditorModule.RegisterCustomPropertyTypeLayout(FRCPassphrase::StaticStruct()->GetFName(), FOnGetPropertyTypeCustomizationInstance::CreateStatic(&FPassphraseCustomization::MakeInstance));
 }
 
 void FRemoteControlUIModule::UnregisterStructCustomizations()
@@ -755,6 +758,11 @@ void FRemoteControlUIModule::UnregisterStructCustomizations()
 	{
 		PropertyEditorModule.UnregisterCustomClassLayout(RemoteControlUIModule::GetCustomizedStructNames()[NameIndex]);
 	}
+
+	if (UObjectInitialized())
+	{
+		PropertyEditorModule.UnregisterCustomPropertyTypeLayout(FRCPassphrase::StaticStruct()->GetFName());
+	}
 }
 
 void FRemoteControlUIModule::RegisterSettings()
diff --git a/Engine/Plugins/VirtualProduction/RemoteControl/Source/RemoteControlUI/Private/UI/Customizations/FPassphraseCustomization.cpp b/Engine/Plugins/VirtualProduction/RemoteControl/Source/RemoteControlUI/Private/UI/Customizations/FPassphraseCustomization.cpp
new file mode 100644
--- /dev/null
+++ b/Engine/Plugins/VirtualProduction/RemoteControl/Source/RemoteControlUI/Private/UI/Customizations/FPassphraseCustomization.cpp
@@ -0,0 +1,96 @@
+﻿#include "FPassphraseCustomization.h"
+
+#include "DetailWidgetRow.h"
+#include "RemoteControlSettings.h"
+#include "Widgets/Input/SButton.h"
+#include "Widgets/Input/SCheckBox.h"
+#include "Widgets/Input/SEditableTextBox.h"
+
+void FPassphraseCustomization::CustomizeHeader(TSharedRef<IPropertyHandle> PropertyHandle, FDetailWidgetRow& HeaderRow, IPropertyTypeCustomizationUtils& CustomizationUtils)
+{
+	TSharedPtr<SEditableTextBox> IdentifierBox;
+	TSharedPtr<SEditableTextBox> PassphraseBox;
+	TSharedPtr<SButton> EditButton;
+	TSharedPtr<SCheckBox> CheckBox;
+	
+	TSharedPtr<IPropertyHandle> IdentifierHandle = PropertyHandle->GetChildHandle(GET_MEMBER_NAME_CHECKED(FRCPassphrase, Identifier));
+	TSharedPtr<IPropertyHandle> PassphraseHandle = PropertyHandle->GetChildHandle(GET_MEMBER_NAME_CHECKED(FRCPassphrase, Passphrase));
+
+	HeaderRow.NameContent()
+	[
+		SNew(SHorizontalBox)
+		+ SHorizontalBox::Slot()
+		.HAlign(HAlign_Right)
+		.AutoWidth()
+		[
+			SAssignNew(IdentifierBox, SEditableTextBox)
+			.IsReadOnly(true)
+			.Text_Lambda([this, IdentifierHandle]()
+			{
+				FString StringValue;
+				IdentifierHandle->GetValue(StringValue);
+				
+				FText OutText = FText::FromString(StringValue);
+				return OutText;
+			})
+			.OnTextCommitted_Lambda([this, IdentifierHandle](const FText& Text, ETextCommit::Type CommitType)
+			{
+				IdentifierHandle->SetValue(Text.ToString());
+			})
+		]
+	];
+
+	HeaderRow.ValueContent()
+	[
+		SNew(SHorizontalBox)
+		+ SHorizontalBox::Slot()
+		.AutoWidth()
+		.HAlign(HAlign_Left)
+		[
+			SAssignNew(PassphraseBox, SEditableTextBox)
+			.IsReadOnly(true)
+			.IsPassword(true)
+			.HintText(FText::FromString("Passphrase"))
+			.SelectAllTextWhenFocused(true)
+			.Text_Lambda([this, PassphraseHandle]()
+			{
+				FString OutString;
+				PassphraseHandle->GetValue(OutString);
+				
+				return FText::FromString(OutString);
+			})
+			.OnTextCommitted_Lambda([this, PassphraseHandle](const FText& Text, ETextCommit::Type CommitType)
+			{
+				// Only Commits using Enter
+				if (ETextCommit::OnEnter == CommitType)
+				{
+					PassphraseHandle->SetValue(FMD5::HashAnsiString(*Text.ToString()));
+				}
+			})
+		]
+		+ SHorizontalBox::Slot()
+		.AutoWidth()
+		.HAlign(HAlign_Right)
+		[
+			SAssignNew(CheckBox, SCheckBox)
+			.ToolTipText(FText::FromString("Unhide Password"))
+			.IsChecked(ECheckBoxState::Unchecked)
+			.OnCheckStateChanged_Lambda([this, PassphraseBox](ECheckBoxState CheckState)
+			{
+				PassphraseBox->SetIsPassword(CheckState == ECheckBoxState::Checked ? false : true);
+			})
+		]
+		+ SHorizontalBox::Slot()
+		.HAlign(HAlign_Right)
+		[
+			SAssignNew(EditButton, SButton)
+			.Text(FText::FromString("Edit"))
+			.OnClicked_Lambda([PassphraseBox, IdentifierBox]()
+			{
+				IdentifierBox->SetIsReadOnly(!IdentifierBox->IsReadOnly());
+				PassphraseBox->SetIsReadOnly(!PassphraseBox->IsReadOnly());
+				return FReply::Handled();
+			})
+		]
+	];
+}
diff --git a/Engine/Plugins/VirtualProduction/RemoteControl/Source/RemoteControlUI/Private/UI/Customizations/FPassphraseCustomization.h b/Engine/Plugins/VirtualProduction/RemoteControl/Source/RemoteControlUI/Private/UI/Customizations/FPassphraseCustomization.h
new file mode 100644
--- /dev/null
+++ b/Engine/Plugins/VirtualProduction/RemoteControl/Source/RemoteControlUI/Private/UI/Customizations/FPassphraseCustomization.h
@@ -0,0 +1,18 @@
+﻿#pragma once
+
+#include "IPropertyTypeCustomization.h"
+
+class FPassphraseCustomization : public IPropertyTypeCustomization
+{
+public:
+	static TSharedRef<IPropertyTypeCustomization> MakeInstance()
+	{
+		return MakeShared<FPassphraseCustomization>();
+	}
+
+	virtual void CustomizeHeader(TSharedRef<IPropertyHandle> PropertyHandle, class FDetailWidgetRow& HeaderRow, IPropertyTypeCustomizationUtils& CustomizationUtils) override;
+
+	virtual void CustomizeChildren( TSharedRef<IPropertyHandle> PropertyHandle, class IDetailChildrenBuilder& ChildBuilder, IPropertyTypeCustomizationUtils& CustomizationUtils ) override
+	{
+	}
+};
diff --git a/Engine/Plugins/VirtualProduction/RemoteControl/Source/WebRemoteControl/Private/RemoteControlRequest.h b/Engine/Plugins/VirtualProduction/RemoteControl/Source/WebRemoteControl/Private/RemoteControlRequest.h
--- a/Engine/Plugins/VirtualProduction/RemoteControl/Source/WebRemoteControl/Private/RemoteControlRequest.h
+++ b/Engine/Plugins/VirtualProduction/RemoteControl/Source/WebRemoteControl/Private/RemoteControlRequest.h
@@ -35,6 +35,9 @@ struct FRCRequest
 		return StructParameters.FindChecked(ParameterName);
 	}
 
+	UPROPERTY()
+	FString Passphrase;
+
 	/** Holds the request's TCHAR payload. */
 	TArray<uint8> TCHARBody;
 
@@ -548,5 +551,4 @@ struct FRCWebSocketPresetRegisterBody : public FRCRequest
 	/** Whether changes to properties triggered remotely should fire an event. */
 	UPROPERTY()
 	bool IgnoreRemoteChanges = false;
-};
-
+};
\ No newline at end of file
diff --git a/Engine/Plugins/VirtualProduction/RemoteControl/Source/WebRemoteControl/Private/RemoteControlResponse.h b/Engine/Plugins/VirtualProduction/RemoteControl/Source/WebRemoteControl/Private/RemoteControlResponse.h
--- a/Engine/Plugins/VirtualProduction/RemoteControl/Source/WebRemoteControl/Private/RemoteControlResponse.h
+++ b/Engine/Plugins/VirtualProduction/RemoteControl/Source/WebRemoteControl/Private/RemoteControlResponse.h
@@ -75,6 +75,22 @@ struct FGetPresetResponse
 	FRCPresetDescription Preset;
 };
 
+USTRUCT()
+struct FCheckPassphraseResponse
+{
+	GENERATED_BODY()
+
+	FCheckPassphraseResponse() = default;
+
+	FCheckPassphraseResponse(bool bInKeyCorrect)
+		: keyCorrect(bInKeyCorrect)
+	{}
+
+	UPROPERTY()
+	bool keyCorrect;
+	
+};
+
 USTRUCT()
 struct FDescribeObjectResponse
 {
diff --git a/Engine/Plugins/VirtualProduction/RemoteControl/Source/WebRemoteControl/Private/RemoteControlWebSocketServer.cpp b/Engine/Plugins/VirtualProduction/RemoteControl/Source/WebRemoteControl/Private/RemoteControlWebSocketServer.cpp
--- a/Engine/Plugins/VirtualProduction/RemoteControl/Source/WebRemoteControl/Private/RemoteControlWebSocketServer.cpp
+++ b/Engine/Plugins/VirtualProduction/RemoteControl/Source/WebRemoteControl/Private/RemoteControlWebSocketServer.cpp
@@ -48,6 +48,10 @@ namespace RemoteControlWebSocketServer
 			{
 				Message.RequestPayload = MakeArrayView(InPayload).Slice(PayloadDelimiters.BlockStart, PayloadDelimiters.BlockEnd - PayloadDelimiters.BlockStart);
 			}
+			if (!Request.Passphrase.IsEmpty())
+			{
+				Message.Header.FindOrAdd(WebRemoteControlInternalUtils::PassphraseHeader) = TArray<FString>({Request.Passphrase});
+			}
 			Message.MessageName = MoveTemp(Request.MessageName);
 			ParsedMessage = MoveTemp(Message);
 		}
@@ -64,6 +68,23 @@ void FWebsocketMessageRouter::Dispatch(const FRemoteControlWebSocketMessage& Mes
 	}
 }
 
+void FWebsocketMessageRouter::AddPreDispatch(TFunction<bool(const FRemoteControlWebSocketMessage& Message)> WebsocketPreprocessor)
+{
+	DispatchPreProcessor.Add(WebsocketPreprocessor);
+}
+
+bool FWebsocketMessageRouter::PreDispatch(const FRemoteControlWebSocketMessage& Message) const
+{
+	for (TFunction<bool(const FRemoteControlWebSocketMessage&)> PreprocessFunction : DispatchPreProcessor)
+	{
+		if (!PreprocessFunction(Message))
+		{
+			return false;
+		}
+	}
+	return true;
+}
+
 void FWebsocketMessageRouter::BindRoute(const FString& MessageName, FWebSocketMessageDelegate OnMessageReceived)
 {
 	DispatchTable.Add(MessageName, MoveTemp(OnMessageReceived));
@@ -199,6 +220,12 @@ void FRCWebSocketServer::ReceivedRawPacket(void* Data, int32 Size, FGuid ClientI
 	if (TOptional<FRemoteControlWebSocketMessage> Message = RemoteControlWebSocketServer::ParseWebsocketMessage(Payload))
 	{
 		Message->ClientId = ClientId;
+		
+		if (!Router->PreDispatch(*Message))
+		{
+			return;
+		}
+	
 		Router->Dispatch(*Message);
 	}
 }
diff --git a/Engine/Plugins/VirtualProduction/RemoteControl/Source/WebRemoteControl/Private/RemoteControlWebSocketServer.h b/Engine/Plugins/VirtualProduction/RemoteControl/Source/WebRemoteControl/Private/RemoteControlWebSocketServer.h
--- a/Engine/Plugins/VirtualProduction/RemoteControl/Source/WebRemoteControl/Private/RemoteControlWebSocketServer.h
+++ b/Engine/Plugins/VirtualProduction/RemoteControl/Source/WebRemoteControl/Private/RemoteControlWebSocketServer.h
@@ -29,7 +29,24 @@ public:
 	 */
 	void UnbindRoute(const FString& MessageName);
 
+	/**
+	 * Adds a Preprocessor called before the Dispatch of the Route
+	 * @param WebsocketPreprocessor The Preprocessor Function to register 
+	 */
+	void AddPreDispatch(TFunction<bool(const struct FRemoteControlWebSocketMessage& Message)> WebsocketPreprocessor);
+
 private:
+	/** Preprocessors for the Dispatch Function */
+	TArray<TFunction<bool(const struct FRemoteControlWebSocketMessage& Message)>> DispatchPreProcessor;
+	
+	/**
+	 * Check if the Message to be dispatched can be executed
+	 * @param MessageName the name of the message to dispatch, used to find its handler
+	 * @param TCHARMessage the payload to dispatch.
+	 * @return true if can be dispatched, false if not
+	 */
+	bool PreDispatch(const struct FRemoteControlWebSocketMessage& Message) const;
+	
 	/**
 	 * Invoke the handler bound to a message name.
 	 * @param MessageName the name of the message to dispatch, used to find its handler.
diff --git a/Engine/Plugins/VirtualProduction/RemoteControl/Source/WebRemoteControl/Private/WebRemoteControl.cpp b/Engine/Plugins/VirtualProduction/RemoteControl/Source/WebRemoteControl/Private/WebRemoteControl.cpp
--- a/Engine/Plugins/VirtualProduction/RemoteControl/Source/WebRemoteControl/Private/WebRemoteControl.cpp
+++ b/Engine/Plugins/VirtualProduction/RemoteControl/Source/WebRemoteControl/Private/WebRemoteControl.cpp
@@ -13,6 +13,7 @@
 #include "RemoteControlWebsocketRoute.h"
 #include "WebRemoteControlInternalUtils.h"
 #include "WebRemoteControlExternalLogger.h"
+#include "WebRemoteControlUtils.h"
 #include "WebSocketMessageHandler.h"
 
 #if WITH_EDITOR
@@ -269,6 +270,27 @@ void FWebRemoteControlModule::StartupModule()
 	RegisterConsoleCommands();
 	RegisterRoutes();
 
+	WebSocketRouter->AddPreDispatch([this](const struct FRemoteControlWebSocketMessage& Message)
+	{
+		const TArray<FString>* InPassphrase = Message.Header.Find(WebRemoteControlInternalUtils::PassphraseHeader);
+		
+		const bool bCanBeDispatched = InPassphrase ? CheckPassphrase(InPassphrase->Last()) : CheckPassphrase("");
+
+		if (!bCanBeDispatched)
+		{
+			TArray<uint8> Response;
+			FRCRequestWrapper Wrapper;
+			Wrapper.RequestId = Message.MessageId;
+			
+			WebRemoteControlInternalUtils::CreateUTF8ErrorMessage(FString::Printf(TEXT("Given Passphrase is not correct!")), Wrapper.TCHARBody);
+			WebRemoteControlUtils::SerializeMessage(Wrapper, Response);
+
+			WebSocketServer.Send(Message.ClientId, MoveTemp(Response));
+		}
+
+		return bCanBeDispatched;
+	});
+
 	if (GetDefault<URemoteControlSettings>()->bAutoStartWebServer || CVarWebControlStartOnBoot.GetValueOnAnyThread() > 0)
 	{
 		StartHttpServer();
@@ -394,6 +416,31 @@ void FWebRemoteControlModule::StartHttpServer()
 			StartRoute(Route);
 		}
 
+		const FHttpRequestHandler ValidationRequestHandler = FHttpRequestHandler([this](const FHttpServerRequest& Request, const FHttpResultCallback& OnComplete)
+		{
+			TUniquePtr<FHttpServerResponse> Response = WebRemoteControlInternalUtils::CreateHttpResponse();
+
+			TArray<FString> ValueArray = {};
+			if (Request.Headers.Find(WebRemoteControlInternalUtils::PassphraseHeader))
+			{
+				ValueArray = Request.Headers[WebRemoteControlInternalUtils::PassphraseHeader];
+			}
+			
+			const FString Passphrase = !ValueArray.IsEmpty() ? ValueArray.Last() : "";
+
+			if (!CheckPassphrase(Passphrase))
+			{
+				WebRemoteControlInternalUtils::CreateUTF8ErrorMessage(FString::Printf(TEXT("Given Passphrase is not correct!")), Response->Body);
+				Response->Code = EHttpServerResponseCodes::Denied;
+				OnComplete(MoveTemp(Response));
+				return true;
+			}
+
+			return false;
+		});
+
+		HttpRouter->RegisterRequestPreprocessor(ValidationRequestHandler);
+		
 		// Go through externally registered request pre-processors and register them with the http router.
 		for (const TPair<FDelegateHandle, FHttpRequestHandler>& Handler : PreprocessorsToRegister)
 		{
@@ -528,6 +575,14 @@ void FWebRemoteControlModule::RegisterRoutes()
 		FRequestHandlerDelegate::CreateRaw(this, &FWebRemoteControlModule::HandleDescribeObjectRoute)
 		});
 
+	// Passphrase Checking
+	RegisterRoute({
+		TEXT("Check whether or no the given Passphrase is correct"),
+		FHttpPath(TEXT("/remote/passphrase/")),
+		EHttpServerRequestVerbs::VERB_GET,
+		FRequestHandlerDelegate::CreateRaw(this, &FWebRemoteControlModule::HandlePassphraseRoute)
+		});
+
 	// Preset API
 	RegisterRoute({
 		TEXT("Get a remote control preset's content."),
@@ -751,8 +806,11 @@ bool FWebRemoteControlModule::HandleBatchRequest(const FHttpServerRequest& Reque
 	JsonWriter->WriteIdentifierPrefix("Responses");
 	JsonWriter->WriteArrayStart();
 
+	BatchRequest.Passphrase = Request.Headers[WebRemoteControlInternalUtils::PassphraseHeader].Last();
+
 	for (FRCRequestWrapper& Wrapper : BatchRequest.Requests)
 	{
+		Wrapper.Passphrase = BatchRequest.Passphrase;
 		// This makes sure the Json writer is in a good state before writing raw data.
 		JsonWriter->WriteRawJSONValue(TEXT(""));
 		InvokeWrappedRequest(Wrapper, Writer, &Request);
@@ -1433,6 +1491,34 @@ bool FWebRemoteControlModule::HandleDescribeObjectRoute(const FHttpServerRequest
 	return true;
 }
 
+bool FWebRemoteControlModule::HandlePassphraseRoute(const FHttpServerRequest& Request, const FHttpResultCallback& OnComplete)
+{
+	TUniquePtr<FHttpServerResponse> Response = WebRemoteControlInternalUtils::CreateHttpResponse();
+	
+	FString Passphrase = "";
+	const TArray<FString>* PassphraseHeader = Request.Headers.Find(WebRemoteControlInternalUtils::PassphraseHeader);
+	
+	if (PassphraseHeader)
+	{
+		Passphrase = PassphraseHeader->Last();
+	}
+	
+	if (bool bIsCorrect = CheckPassphrase(Passphrase))
+	{
+		WebRemoteControlUtils::SerializeMessage(FCheckPassphraseResponse{ bIsCorrect}, Response->Body);
+		Response->Code = EHttpServerResponseCodes::Ok;
+	}
+	else
+	{
+		WebRemoteControlUtils::SerializeMessage(FCheckPassphraseResponse{ bIsCorrect}, Response->Body);
+		Response->Code = EHttpServerResponseCodes::Denied;
+	}
+
+	OnComplete(MoveTemp(Response));
+	return true;
+}
+
+
 bool FWebRemoteControlModule::HandleSearchActorRoute(const FHttpServerRequest& Request, const FHttpResultCallback& OnComplete)
 {
 	TUniquePtr<FHttpServerResponse> Response = WebRemoteControlInternalUtils::CreateHttpResponse(EHttpServerResponseCodes::NotSupported);
@@ -1647,6 +1733,34 @@ bool FWebRemoteControlModule::HandleEntityMetadataOperationsRoute(const FHttpSer
 	return true;
 }
 
+bool FWebRemoteControlModule::CheckPassphrase(const FString& HashedPassphrase) const
+{
+	bool bOutResult = !(GetMutableDefault<URemoteControlSettings>()->bUseRemoteControlPassphrase);
+
+	if (bOutResult)
+	{
+		return true;
+	}
+
+	TArray<FString> HashedPassphrases = GetMutableDefault<URemoteControlSettings>()->GetHashedPassphrases();
+	if (HashedPassphrases.IsEmpty())
+	{
+		return true;
+	}
+	
+	for (const FString& InPassphrase : HashedPassphrases)
+	{
+		bOutResult = bOutResult || InPassphrase == HashedPassphrase;
+
+		if (bOutResult)
+		{
+			break;
+		}
+	}
+	
+	return bOutResult;
+}
+
 bool FWebRemoteControlModule::HandleEntitySetLabelRoute(const FHttpServerRequest& Request, const FHttpResultCallback& OnComplete)
 {
 	TUniquePtr<FHttpServerResponse> Response = WebRemoteControlInternalUtils::CreateHttpResponse();
@@ -1717,6 +1831,11 @@ void FWebRemoteControlModule::HandleWebSocketHttpMessage(const FRemoteControlWeb
 	{
 		return;
 	}
+	
+	if (WebSocketMessage.Header.Find(WebRemoteControlInternalUtils::PassphraseHeader))
+	{
+		Wrapper.Passphrase = WebSocketMessage.Header[WebRemoteControlInternalUtils::PassphraseHeader][0];
+	}
 
 	LogRequestExternally(Wrapper.RequestId, TEXT("UE Received"));
 	
diff --git a/Engine/Plugins/VirtualProduction/RemoteControl/Source/WebRemoteControl/Private/WebRemoteControl.h b/Engine/Plugins/VirtualProduction/RemoteControl/Source/WebRemoteControl/Private/WebRemoteControl.h
--- a/Engine/Plugins/VirtualProduction/RemoteControl/Source/WebRemoteControl/Private/WebRemoteControl.h
+++ b/Engine/Plugins/VirtualProduction/RemoteControl/Source/WebRemoteControl/Private/WebRemoteControl.h
@@ -101,6 +101,9 @@ private:
 	/** Unregister console commands. */
 	void UnregisterConsoleCommands();
 
+	/** Checking ApiKey using Md5. */
+	bool CheckPassphrase(const FString& HashedPassphrase) const;
+
 	//~ Route handlers
 	bool HandleInfoRoute(const FHttpServerRequest& Request, const FHttpResultCallback& OnComplete);
 	bool HandleBatchRequest(const FHttpServerRequest& Request, const FHttpResultCallback& OnComplete);
@@ -123,6 +126,7 @@ private:
 	bool HandleSearchObjectRoute(const FHttpServerRequest& Request, const FHttpResultCallback& OnComplete);
 	bool HandleEntityMetadataOperationsRoute(const FHttpServerRequest& Request, const FHttpResultCallback& OnComplete);
 	bool HandleEntitySetLabelRoute(const FHttpServerRequest& Request, const FHttpResultCallback& OnComplete);
+	bool HandlePassphraseRoute(const FHttpServerRequest& Request, const FHttpResultCallback& OnComplete);
 
 	//~ Websocket route handlers
 	void HandleWebSocketHttpMessage(const struct FRemoteControlWebSocketMessage& WebSocketMessage);
diff --git a/Engine/Plugins/VirtualProduction/RemoteControl/Source/WebRemoteControl/Private/WebRemoteControlInternalUtils.cpp b/Engine/Plugins/VirtualProduction/RemoteControl/Source/WebRemoteControl/Private/WebRemoteControlInternalUtils.cpp
--- a/Engine/Plugins/VirtualProduction/RemoteControl/Source/WebRemoteControl/Private/WebRemoteControlInternalUtils.cpp
+++ b/Engine/Plugins/VirtualProduction/RemoteControl/Source/WebRemoteControl/Private/WebRemoteControlInternalUtils.cpp
@@ -85,6 +85,7 @@ TSharedRef<FHttpServerRequest> UnwrapHttpRequest(const FRCRequestWrapper& Wrappe
 	WrappedHttpRequest->RelativePath = Wrapper.URL;
 	WrappedHttpRequest->Verb = ParseHttpVerb(Wrapper.Verb);
 	WrappedHttpRequest->Body = Wrapper.TCHARBody;
+	WrappedHttpRequest->Headers.Add(WebRemoteControlInternalUtils::PassphraseHeader, { Wrapper.Passphrase });
 
 	return WrappedHttpRequest;
 }
diff --git a/Engine/Plugins/VirtualProduction/RemoteControl/Source/WebRemoteControl/Private/WebRemoteControlInternalUtils.h b/Engine/Plugins/VirtualProduction/RemoteControl/Source/WebRemoteControl/Private/WebRemoteControlInternalUtils.h
--- a/Engine/Plugins/VirtualProduction/RemoteControl/Source/WebRemoteControl/Private/WebRemoteControlInternalUtils.h
+++ b/Engine/Plugins/VirtualProduction/RemoteControl/Source/WebRemoteControl/Private/WebRemoteControlInternalUtils.h
@@ -49,6 +49,7 @@ namespace RemotePayloadSerializer
 namespace WebRemoteControlInternalUtils
 {
 	static const TCHAR* WrappedRequestHeader = TEXT("UE-Wrapped-Request");
+	static const FString PassphraseHeader = TEXT("Passphrase");
 
 	/**
 	 * Construct a default http response with CORS headers.
diff --git a/Engine/Plugins/VirtualProduction/RemoteControl/Source/WebRemoteControl/Public/RemoteControlWebsocketRoute.h b/Engine/Plugins/VirtualProduction/RemoteControl/Source/WebRemoteControl/Public/RemoteControlWebsocketRoute.h
--- a/Engine/Plugins/VirtualProduction/RemoteControl/Source/WebRemoteControl/Public/RemoteControlWebsocketRoute.h
+++ b/Engine/Plugins/VirtualProduction/RemoteControl/Source/WebRemoteControl/Public/RemoteControlWebsocketRoute.h
@@ -8,6 +8,7 @@ struct FRemoteControlWebSocketMessage
 	int32 MessageId = -1;
 	FGuid ClientId;
 	TArrayView<uint8> RequestPayload;
+	TMap<FString, TArray<FString>> Header;
 };
 
 DECLARE_DELEGATE_OneParam(FWebSocketMessageDelegate, const FRemoteControlWebSocketMessage& /** Message */);
diff --git a/Engine/Plugins/VirtualProduction/RemoteControlWebInterface/Source/RemoteControlWebInterface/Private/RCWebInterfaceCustomizations.cpp b/Engine/Plugins/VirtualProduction/RemoteControlWebInterface/Source/RemoteControlWebInterface/Private/RCWebInterfaceCustomizations.cpp
--- a/Engine/Plugins/VirtualProduction/RemoteControlWebInterface/Source/RemoteControlWebInterface/Private/RCWebInterfaceCustomizations.cpp
+++ b/Engine/Plugins/VirtualProduction/RemoteControlWebInterface/Source/RemoteControlWebInterface/Private/RCWebInterfaceCustomizations.cpp
@@ -205,6 +205,32 @@ void FRCWebInterfaceCustomizations::GeneratePanelExtensions(TArray<TSharedRef<SW
 
 	TSharedPtr<SWidget> Throbber;
 
+	OutExtensions.Add(
+		SNew(SButton)
+		.ButtonStyle(FEditorStyle::Get(), "FlatButton")
+		.ToolTipText(LOCTEXT("AppPassphraseDisabled", "Warning: Passphrase is disabled!"))
+		.Visibility_Lambda([]()
+		{
+			const URemoteControlSettings* RCSettings = GetDefault<URemoteControlSettings>();
+			bool bShouldWarn = RCSettings->bShowPassphraseDisabledWarning && !RCSettings->bUseRemoteControlPassphrase;
+			return bShouldWarn ? EVisibility::Visible : EVisibility::Collapsed;
+		})
+		.OnClicked_Lambda([]()
+		{
+			URemoteControlSettings* RCSettings = GetMutableDefault<URemoteControlSettings>();
+			RCSettings->bShowPassphraseDisabledWarning = false;
+			RCSettings->SaveConfig();
+			
+			return FReply::Handled();
+		})
+		[
+			SNew(STextBlock)
+			.TextStyle(FEditorStyle::Get(), "NormalText.Important")
+			.Font(FEditorStyle::Get().GetFontStyle("FontAwesome.10"))
+			.Text(FEditorFontGlyphs::Exclamation_Triangle)
+		]
+	);
+	
 	OutExtensions.Add(
 		SNew(SWidgetSwitcher)
 		.WidgetIndex_Lambda(GetDetailWidgetIndex)
diff --git a/Engine/Plugins/VirtualProduction/RemoteControlWebInterface/WebApp/Client/src/components/index.ts b/Engine/Plugins/VirtualProduction/RemoteControlWebInterface/WebApp/Client/src/components/index.ts
--- a/Engine/Plugins/VirtualProduction/RemoteControlWebInterface/WebApp/Client/src/components/index.ts
+++ b/Engine/Plugins/VirtualProduction/RemoteControlWebInterface/WebApp/Client/src/components/index.ts
@@ -4,4 +4,5 @@ export * from './DragDropWidgets';
 export * from './PropertiesDrawer';
 export * from './modals';
 export * from './Screen';
+export * from './Login';
 export * from './Tabs';
\ No newline at end of file
diff --git a/Engine/Plugins/VirtualProduction/RemoteControlWebInterface/WebApp/Client/src/reducers/api.ts b/Engine/Plugins/VirtualProduction/RemoteControlWebInterface/WebApp/Client/src/reducers/api.ts
--- a/Engine/Plugins/VirtualProduction/RemoteControlWebInterface/WebApp/Client/src/reducers/api.ts
+++ b/Engine/Plugins/VirtualProduction/RemoteControlWebInterface/WebApp/Client/src/reducers/api.ts
@@ -14,7 +14,9 @@ export type ApiState = {
   view: IView;
   status: {
     connected?: boolean;
+    keyCorrect?: boolean;
     loading?: boolean;
+    isOpen?: boolean;
   },
 };
 
@@ -22,15 +24,17 @@ export type ApiState = {
 let _preset;
 let _dispatch: Dispatch;
 let _socket: SocketIOClient.Socket;
+let _passphrase: string;
 const _host = (process.env.NODE_ENV === 'development' ? `http://${window.location.hostname}:7001` : '');
 
 function _initialize(dispatch: Dispatch, getState: () => { api: ApiState }) {
   _dispatch = dispatch;
+  _passphrase = localStorage.getItem('passphrase'); 
 
   _socket = io(`${_host}/`, { path: '/api/io' });
 
   _socket
-    .on('disconnect', () => dispatch(API.STATUS({ connected: false, version: undefined })))
+    .on('disconnect', () => dispatch(API.STATUS({ connected: false, isOpen: false, keyCorrect: false, version: undefined })))
     .on('presets', (presets: IPreset[]) => dispatch(API.PRESETS(presets)))
     .on('payloads', (payloads: IPayloads) => {
       dispatch(API.PAYLOADS(payloads));
@@ -57,12 +61,21 @@ function _initialize(dispatch: Dispatch, getState: () => { api: ApiState }) {
       dispatch(API.VIEW(view));
     })
     .on('connected', (connected: boolean, version: string) => {
-      dispatch(API.STATUS({ connected, version, loading: false }));
+      dispatch(API.STATUS({ connected, version }));
+
+      if (_passphrase !== null)
+        _api.passphrase.login(_passphrase);
+          
+      _api.presets.get();
+      _api.payload.all();
+    })
+    .on('opened', (isOpen: boolean) => {
+      dispatch(API.STATUS({ isOpen, loading: false }));
+    })
+    .on('passphrase', (keyCorrect: boolean) => {
+      dispatch(API.STATUS({ keyCorrect }));
+
 
-      if (connected) {
-          _api.presets.get();
-          _api.payload.all();
-      }
     })
     .on('loading', (loading: boolean) => {
       dispatch(API.STATUS({ loading }));
@@ -71,7 +84,7 @@ function _initialize(dispatch: Dispatch, getState: () => { api: ApiState }) {
 
 type IRequestCallback = Function | string | undefined;
 
-async function _request(method: string, url: string, body: string | object | undefined, callback: IRequestCallback): Promise<any> {
+async function _request(method: string, url: string, body: string | object | undefined, callback: IRequestCallback, passphrase?: string): Promise<any> {
   const request: RequestInit = { method, mode: 'cors', redirect: 'follow', headers: {} };
   if (body instanceof FormData || typeof(body) === 'string') {
     request.body = body;
@@ -79,13 +92,14 @@ async function _request(method: string, url: string, body: string | object | und
     request.body = JSON.stringify(body);
     request.headers['Content-Type'] = 'application/json';
   }
+  request.headers['passphrase'] = passphrase ?? _passphrase;
 
   const res = await fetch(_host + url, request);
 
   let answer: any = await res.text();
   if (answer.length > 0)
     answer = JSON.parse(answer);
-
+  
   if (!res.ok)
     throw answer;
 
@@ -95,8 +109,8 @@ async function _request(method: string, url: string, body: string | object | und
   return answer;
 }
 
-function _get(url: string, callback?: IRequestCallback)        { return _request('GET', url, undefined, callback) }
-function _put(url: string, body: any)                          { return _request('PUT', url, body, undefined) }
+function _get(url: string, callback?: IRequestCallback, passphrase?: string)        { return _request('GET', url, undefined, callback, passphrase) };
+function _put(url: string, body: any, passphrase?: string)                          { return _request('PUT', url, body, undefined, passphrase) };
 
 const API = {
   STATUS: createAction<any>('API_STATUS'),
@@ -137,6 +151,19 @@ export const _api = {
       _socket.emit('view', _preset, view);
     },
   },
+  passphrase: {
+    login: async(passphrase: string): Promise<boolean> => {
+      const ok = await _get('/api/passphrase', API.STATUS, passphrase);
+      if (!ok)
+        return false;
+
+      localStorage.setItem('passphrase', passphrase);
+      _passphrase = passphrase;
+      _api.presets.get();
+      _api.payload.all();
+      return true;
+    },
+  },
   payload: {
     get: (preset: string): Promise<IPayload> => _get(`/api/presets/payload?preset=${preset}`, API.PAYLOAD),
     all: (): Promise<IPayloads> => _get('/api/payloads', API.PAYLOADS),
@@ -195,6 +222,8 @@ const initialState: ApiState = {
   view: { tabs: null },
   status: {
     connected: false,
+    keyCorrect: false,
+    isOpen: false,
   },
 };
 
diff --git a/Engine/Plugins/VirtualProduction/RemoteControlWebInterface/WebApp/Server/src/Api.ts b/Engine/Plugins/VirtualProduction/RemoteControlWebInterface/WebApp/Server/src/Api.ts
--- a/Engine/Plugins/VirtualProduction/RemoteControlWebInterface/WebApp/Server/src/Api.ts
+++ b/Engine/Plugins/VirtualProduction/RemoteControlWebInterface/WebApp/Server/src/Api.ts
@@ -24,6 +24,11 @@ export class Api {
     this.send(res, UnrealEngine.getPresets());
   }
 
+  @Get('passphrase')
+  private async passphrase(req: Request, res: Response) {
+    this.send(res, Promise.resolve({ keyCorrect: await UnrealEngine.checkPassphrase(req.headers['passphrase'].toString()) }) );
+  }
+
   @Get('payloads')
   private payloads(req: Request, res: Response) {
     this.send(res, UnrealEngine.getPayloads());
diff --git a/Engine/Plugins/VirtualProduction/RemoteControlWebInterface/WebApp/Server/src/Notify.ts b/Engine/Plugins/VirtualProduction/RemoteControlWebInterface/WebApp/Server/src/Notify.ts
--- a/Engine/Plugins/VirtualProduction/RemoteControlWebInterface/WebApp/Server/src/Notify.ts
+++ b/Engine/Plugins/VirtualProduction/RemoteControlWebInterface/WebApp/Server/src/Notify.ts
@@ -21,6 +21,8 @@ export namespace Notify {
 
       if (UnrealEngine.isConnected())
         socket.emit('connected', true);
+        
+      socket.emit('opened', UnrealEngine.isOpenConnection());
 
       if (UnrealEngine.isLoading())
         socket.emit('loading', true);
diff --git a/Engine/Plugins/VirtualProduction/RemoteControlWebInterface/WebApp/Server/src/UnrealEngine.ts b/Engine/Plugins/VirtualProduction/RemoteControlWebInterface/WebApp/Server/src/UnrealEngine.ts
--- a/Engine/Plugins/VirtualProduction/RemoteControlWebInterface/WebApp/Server/src/UnrealEngine.ts
+++ b/Engine/Plugins/VirtualProduction/RemoteControlWebInterface/WebApp/Server/src/UnrealEngine.ts
@@ -30,6 +30,10 @@ namespace UnrealApi {
     Value: string;
   };
 
+  export type PassphraseCheck = {
+    keyCorrect: boolean;
+  };
+
   export type Request = {
     RequestId: number;
     URL: string;
@@ -74,15 +78,16 @@ export namespace UnrealEngine {
   let isPullingPresets: boolean;
   let isLoadingPresets: boolean;
   let autoPullTimeout: NodeJS.Timeout;
+  let isOpen: boolean;
 
   let pendings: { [id: string]: (reply: any) => void } = {};
-  let httpRequest: number = 1;
-  let wsRequest: number = 1;
+  let request: number = 1;
 
   let presets: IPresets = {};
   let registered: { [id: string]: boolean };
   let payloads: IPayloads = {};
   let views: { [preset: string]: IView } = {};
+  let workingPassphrase: string = undefined;
 
   export async function initialize() {
     connect();
@@ -93,6 +98,10 @@ export namespace UnrealEngine {
     return (connection?.readyState === WebSocket.OPEN);
   }
 
+  export function isOpenConnection(): boolean {
+    return isOpen;
+  }
+
   export function isLoading(): boolean {
     return isLoadingPresets;
   }
@@ -118,6 +127,28 @@ export namespace UnrealEngine {
       .on('close', onClose);
   }
 
+  export async function checkPassphrase(passphrase: string) : Promise<boolean> {
+    try {
+      let res = await get<UnrealApi.PassphraseCheck>(`/remote/passphrase/`, passphrase);
+      if (!res?.keyCorrect) {
+        return false;
+      }
+    } catch (error) {
+      return false;
+    }
+
+    if (workingPassphrase == undefined) {
+      workingPassphrase = passphrase;
+      await pullPresets(true);
+      pullTimer();   
+    }
+    else {
+      workingPassphrase = passphrase;
+    }
+
+    return true;
+  }
+
   async function onConnected() {
     if (connection.readyState !== WebSocket.OPEN)
       return;
@@ -126,13 +157,18 @@ export namespace UnrealEngine {
     registered = {};
     payloads = {};
     views = {};
-
+    isOpen = undefined;
+    
     clearQuitTimeout();
-
+    
     console.log('Connected to UE Remote WebSocket');
     Notify.emit('connected', true);
-    await refresh();
-    pullTimer();
+    
+    isOpen = await checkPassphrase("");
+    if (isOpen || workingPassphrase !== undefined) {
+      await pullPresets(true);
+      pullTimer();
+    }
   }
 
   async function refresh() { 
@@ -148,8 +184,9 @@ export namespace UnrealEngine {
 
       // First check if there are any new presets
       for (const preset of Presets) {
-        if (preset && !presets[preset.ID])
+        if (preset && !presets[preset.ID]){
           await refreshPreset(preset.ID, preset.Name);
+        }
       }
 
       // Check for deleted presets
@@ -166,8 +203,9 @@ export namespace UnrealEngine {
     }
 
     // Check for new presets once every 5 seconds
-    if (isConnected())
-       autoPullTimeout = setTimeout(pullTimer, 5 * 1000);
+    if (isConnected()) {
+      autoPullTimeout = setTimeout(pullTimer, 5 * 1000);
+    }
   }
  
   async function onMessage(data: WebSocket.Data) {
@@ -264,6 +302,7 @@ export namespace UnrealEngine {
     registered = {};
     payloads = {};
     views = {};
+    isOpen = undefined;
 
     Notify.emit('connected', false);
 
@@ -295,19 +334,20 @@ export namespace UnrealEngine {
       throw new Error('Websocket is not connected');
   }
 
-  function send(message: string, parameters: any) {
+  function send(message: string, parameters: any, passphrase?: string) {
     verifyConnection();
-    const Id = wsRequest++;
-    connection.send(JSON.stringify({ MessageName: message, Id, Parameters: parameters }));
+    const Id = parameters?.RequestId ?? request++;
+    passphrase = passphrase ?? workingPassphrase;
+    connection.send(JSON.stringify({ MessageName: message, Id, Passphrase: passphrase, Parameters: parameters}));
   }
 
-  function http<T>(Verb: string, URL: string, Body?: object, wantAnswer?: boolean): Promise<T> {
-    const RequestId = httpRequest++;
+  function http<T>(Verb: string, URL: string, Body?: object, wantAnswer?: boolean, passphrase?: string): Promise<T> {
+    const RequestId = request++;
     const payload = { RequestId, Verb, URL, Body };
     if (Program.logger)
       LogServer.logLoopback(payload);
 
-    send('http', payload);
+    send('http', payload, passphrase);
     if (!wantAnswer)
       return Promise.resolve(null);
 
@@ -316,12 +356,12 @@ export namespace UnrealEngine {
     });
   }
 
-  function get<T>(url: string): Promise<T> {
-    return http<T>('GET', url, undefined, true);
+  function get<T>(url: string, passphrase?: string): Promise<T> {
+    return http<T>('GET', url, undefined, true, passphrase);
   }
 
-  function put<T>(url: string, body: object): Promise<T> {
-    return http<T>('PUT', url, body, true);
+  function put<T>(url: string, body: object, passphrase?: string): Promise<T> {
+    return http<T>('PUT', url, body, true, passphrase);
   }
 
   function registerToPreset(PresetName: string): void {
@@ -356,8 +396,9 @@ export namespace UnrealEngine {
           continue;
 
         allPresets[Preset.ID] = Preset;
-        if (bInitial || !presets[Preset.ID])
+        if (bInitial || !presets[Preset.ID]){
           allPayloads[Preset.ID] = await pullPresetValues(Preset);
+        }
       }
 
       if (!equal(presets, allPresets)) {
@@ -546,7 +587,7 @@ export namespace UnrealEngine {
     const Requests = [];
     for (const property of Preset.ExposedProperties) {
       Requests.push({
-        RequestId: httpRequest++,
+        RequestId: request++,
         Verb: 'GET',
         URL: `/remote/preset/${Preset.ID}/property/${property.ID}`,
         PropertyId: property.ID,
@@ -611,7 +652,7 @@ export namespace UnrealEngine {
       if (Program.logger) {
         const presetObj = presets[preset];
         const log = {
-          RequestId: httpRequest,
+          RequestId: request,
           Stage: 'Set Property',
           Preset: presetObj?.Name,
           Property: presetObj?.Exposed[property]?.DisplayName,
diff --git a/Engine/Plugins/VirtualProduction/RemoteControlWebInterface/WebApp/_package.json b/Engine/Plugins/VirtualProduction/RemoteControlWebInterface/WebApp/_package.json
--- a/Engine/Plugins/VirtualProduction/RemoteControlWebInterface/WebApp/_package.json
+++ b/Engine/Plugins/VirtualProduction/RemoteControlWebInterface/WebApp/_package.json
@@ -1,6 +1,6 @@
 {
   "name": "conductor",
-  "version": "0.3.62",
+  "version": "0.3.63",
   "scripts": {
     "postinstall": "concurrently \"cd ./Client && npm install && cd ..\" \"cd ./Server && npm install && cd ..\"",
     "build": "node ./Scripts/build.js",
