Fixing Compilation Issues caused by [REDACTED_PERSON]. Copyright Notices

diff --git a/Engine/Plugins/VirtualProduction/RemoteControl/Source/RemoteControlUI/Private/UI/Customizations/FPassphraseCustomization.cpp b/Engine/Plugins/VirtualProduction/RemoteControl/Source/RemoteControlUI/Private/UI/Customizations/FPassphraseCustomization.cpp
--- a/Engine/Plugins/VirtualProduction/RemoteControl/Source/RemoteControlUI/Private/UI/Customizations/FPassphraseCustomization.cpp
+++ b/Engine/Plugins/VirtualProduction/RemoteControl/Source/RemoteControlUI/Private/UI/Customizations/FPassphraseCustomization.cpp
@@ -1,4 +1,6 @@
-﻿#include "FPassphraseCustomization.h"
+﻿// Copyright Epic Games, Inc. All Rights Reserved.
+
+#include "FPassphraseCustomization.h"
 
 #include "DetailWidgetRow.h"
 #include "RemoteControlSettings.h"
diff --git a/Engine/Plugins/VirtualProduction/RemoteControl/Source/RemoteControlUI/Private/UI/Customizations/FPassphraseCustomization.h b/Engine/Plugins/VirtualProduction/RemoteControl/Source/RemoteControlUI/Private/UI/Customizations/FPassphraseCustomization.h
--- a/Engine/Plugins/VirtualProduction/RemoteControl/Source/RemoteControlUI/Private/UI/Customizations/FPassphraseCustomization.h
+++ b/Engine/Plugins/VirtualProduction/RemoteControl/Source/RemoteControlUI/Private/UI/Customizations/FPassphraseCustomization.h
@@ -1,4 +1,6 @@
-﻿#pragma once
+﻿// Copyright Epic Games, Inc. All Rights Reserved.
+
+#pragma once
 
 #include "IPropertyTypeCustomization.h"
 
