Taking Fix from Initialization Error in Release-5.0

diff --git a/Engine/Plugins/VirtualProduction/RemoteControl/Source/WebRemoteControl/Private/RemoteControlResponse.h b/Engine/Plugins/VirtualProduction/RemoteControl/Source/WebRemoteControl/Private/RemoteControlResponse.h
--- a/Engine/Plugins/VirtualProduction/RemoteControl/Source/WebRemoteControl/Private/RemoteControlResponse.h
+++ b/Engine/Plugins/VirtualProduction/RemoteControl/Source/WebRemoteControl/Private/RemoteControlResponse.h
@@ -87,7 +87,7 @@ struct FCheckPassphraseResponse
 	{}
 
 	UPROPERTY()
-	bool keyCorrect;
+	bool keyCorrect = false;
 	
 };
 
