Fix values not transacting through nDisplay when driving Remote Control parameter

diff --git a/Engine/Plugins/Media/MediaIOFramework/Source/MediaIOCore/MediaIOCore.Build.cs b/Engine/Plugins/Media/MediaIOFramework/Source/MediaIOCore/MediaIOCore.Build.cs
--- a/Engine/Plugins/Media/MediaIOFramework/Source/MediaIOCore/MediaIOCore.Build.cs
+++ b/Engine/Plugins/Media/MediaIOFramework/Source/MediaIOCore/MediaIOCore.Build.cs
@@ -6,10 +6,6 @@ namespace UnrealBuildTool.Rules
 	{
 		public MediaIOCore(ReadOnlyTargetRules Target) : base(Target)
 		{
-			OptimizeCode = CodeOptimization.Never;
-			bUseUnity = false;
-			PCHUsage = PCHUsageMode.NoPCHs;
-		
 			PublicDependencyModuleNames.AddRange(
 				new string[]
 				{
diff --git a/Engine/Plugins/VirtualProduction/RemoteControl/Source/RemoteControl/Private/RemoteControlModule.cpp b/Engine/Plugins/VirtualProduction/RemoteControl/Source/RemoteControl/Private/RemoteControlModule.cpp
--- a/Engine/Plugins/VirtualProduction/RemoteControl/Source/RemoteControl/Private/RemoteControlModule.cpp
+++ b/Engine/Plugins/VirtualProduction/RemoteControl/Source/RemoteControl/Private/RemoteControlModule.cpp
@@ -117,7 +117,6 @@ namespace RemoteControlUtil
 		Filter.bIncludeOnlyOnDiskAssets = false;
 		Filter.ClassNames = {URemoteControlPreset::StaticClass()->GetFName()};
 		Filter.bRecursivePaths = true;
-		Filter.PackagePaths = {TEXT("/")};
 
 		return Filter;
 	}
@@ -1216,6 +1215,11 @@ void FRemoteControlModule::GetPresets(TArray<TSoftObjectPtr<URemoteControlPreset
 
 void FRemoteControlModule::GetPresetAssets(TArray<FAssetData>& OutPresetAssets) const
 {
+	if (CachedPresetsByName.Num() == 0)
+	{
+		CachePresets();
+	}
+
 	OutPresetAssets.Reserve(CachedPresetsByName.Num());
 	for (const TPair<FName, TArray<FAssetData>>& Entry : CachedPresetsByName)
 	{
@@ -1267,7 +1271,7 @@ void FRemoteControlModule::UnregisterEntityFactory(const FName InFactoryName)
 	EntityFactories.Remove(InFactoryName);
 }
 
-void FRemoteControlModule::CachePresets()
+void FRemoteControlModule::CachePresets() const
 {
 	TArray<FAssetData> Assets;
 	RemoteControlUtil::GetAllPresetAssets(Assets);
diff --git a/Engine/Plugins/VirtualProduction/RemoteControl/Source/RemoteControl/Private/RemoteControlModule.h b/Engine/Plugins/VirtualProduction/RemoteControl/Source/RemoteControl/Private/RemoteControlModule.h
--- a/Engine/Plugins/VirtualProduction/RemoteControl/Source/RemoteControl/Private/RemoteControlModule.h
+++ b/Engine/Plugins/VirtualProduction/RemoteControl/Source/RemoteControl/Private/RemoteControlModule.h
@@ -50,7 +50,7 @@ public:
 
 private:
 	/** Cache all presets in the project for the ResolvePreset function. */
-	void CachePresets();
+	void CachePresets() const;
 	
 	//~ Asset registry callbacks
 	void OnAssetAdded(const FAssetData& AssetData);
diff --git a/Engine/Plugins/VirtualProduction/RemoteControl/Source/RemoteControlCommon/Public/RCPropertyUtilities.h b/Engine/Plugins/VirtualProduction/RemoteControl/Source/RemoteControlCommon/Public/RCPropertyUtilities.h
--- a/Engine/Plugins/VirtualProduction/RemoteControl/Source/RemoteControlCommon/Public/RCPropertyUtilities.h
+++ b/Engine/Plugins/VirtualProduction/RemoteControl/Source/RemoteControlCommon/Public/RCPropertyUtilities.h
@@ -4,6 +4,7 @@
 
 #include "CoreMinimal.h"
 
+#include "Components/SceneComponent.h"
 #include "RCTypeTraits.h"
 #include "RCTypeUtilities.h"
 #include "Serialization/BufferArchive.h"
@@ -460,6 +461,15 @@ namespace RemoteControlPropertyUtilities
 
 	static UFunction* FindSetterFunctionInternal(FProperty* Property, UClass* OwnerClass)
 	{
+		static const FName RelativeLocationPropertyName = USceneComponent::GetRelativeLocationPropertyName();
+		static const FName RelativeRotationPropertyName = USceneComponent::GetRelativeRotationPropertyName();
+		static const FName RelativeScalePropertyName = USceneComponent::GetRelativeScale3DPropertyName();
+
+		static TMap<FName, FName> CommonPropertiesToFunctions = { 
+			{ RelativeLocationPropertyName, "K2_SetRelativeLocation"},
+			{ RelativeRotationPropertyName, "K2_SetRelativeRotation" },
+			{ RelativeScalePropertyName, "SetRelativeScale3D" }
+		};
 		// Check if the property setter is already cached.
 		TWeakObjectPtr<UFunction> SetterPtr = CachedSetterFunctions.FindRef(Property);
 		if (SetterPtr.IsValid())
@@ -476,7 +486,15 @@ namespace RemoteControlPropertyUtilities
 		}
 		else
 		{
-			return nullptr;
+			if (CommonPropertiesToFunctions.Contains(Property->GetFName()))
+			{
+				SetterFunction = OwnerClass->FindFunctionByName(*CommonPropertiesToFunctions[Property->GetFName()].ToString());
+			}
+
+			if (!SetterFunction)
+			{
+				return nullptr;
+			}
 		}
 #endif
 
