Add websocket route for subscribing/unsubscribing to actor add/delete/rename events

diff --git a/Engine/Plugins/VirtualProduction/RemoteControl/Source/WebRemoteControl/Private/RemoteControlModels.h b/Engine/Plugins/VirtualProduction/RemoteControl/Source/WebRemoteControl/Private/RemoteControlModels.h
--- a/Engine/Plugins/VirtualProduction/RemoteControl/Source/WebRemoteControl/Private/RemoteControlModels.h
+++ b/Engine/Plugins/VirtualProduction/RemoteControl/Source/WebRemoteControl/Private/RemoteControlModels.h
@@ -686,3 +686,32 @@ struct FRCAssetFilter
 	bool EnableBlueprintNativeClassFiltering = false;
 };
 
+/**
+ * A description of an actor that can both uniquely identify it and provide a user-friendly name.
+ */
+USTRUCT()
+struct FRCActorDescription
+{
+	GENERATED_BODY()
+
+	FRCActorDescription() = default;
+
+	FRCActorDescription(const AActor* InActor)
+		: Name(InActor->GetActorNameOrLabel())
+		, Path(InActor->GetPathName())
+	{
+	}
+
+	inline bool operator==(const FRCActorDescription& Other) const
+	{
+		return Path == Other.Path && Name == Other.Name;
+	}
+
+	/** A user-friendly name for the actor. */
+	UPROPERTY()
+	FString Name;
+
+	/** The path to the actor. */
+	UPROPERTY()
+	FString Path;
+};
\ No newline at end of file
diff --git a/Engine/Plugins/VirtualProduction/RemoteControl/Source/WebRemoteControl/Private/RemoteControlRequest.h b/Engine/Plugins/VirtualProduction/RemoteControl/Source/WebRemoteControl/Private/RemoteControlRequest.h
--- a/Engine/Plugins/VirtualProduction/RemoteControl/Source/WebRemoteControl/Private/RemoteControlRequest.h
+++ b/Engine/Plugins/VirtualProduction/RemoteControl/Source/WebRemoteControl/Private/RemoteControlRequest.h
@@ -525,7 +525,7 @@ struct FRCWebSocketRequest : public FRCRequest
 };
 
 /**
- * Holds a request made for web socket.
+ * Holds a request made via websocket to register for events about a given preset.
  */
 USTRUCT()
 struct FRCWebSocketPresetRegisterBody : public FRCRequest
@@ -551,4 +551,30 @@ struct FRCWebSocketPresetRegisterBody : public FRCRequest
 	/** Whether changes to properties triggered remotely should fire an event. */
 	UPROPERTY()
 	bool IgnoreRemoteChanges = false;
-};
\ No newline at end of file
+};
+
+/**
+ * Holds a request made via websocket to register for spawn/destroy events about a given actor type.
+ */
+USTRUCT()
+struct FRCWebSocketActorRegisterBody : public FRCRequest
+{
+	GENERATED_BODY()
+
+	FRCWebSocketActorRegisterBody()
+	{
+		AddStructParameter(ParametersFieldLabel());
+	}
+
+	/**
+	 * Get the label for the property value struct.
+	 */
+	static FString ParametersFieldLabel() { return TEXT("Parameters"); }
+
+	/**
+	 * Name of the actor class to register for.
+	 */
+	UPROPERTY()
+	FName ClassName;
+};
+
diff --git a/Engine/Plugins/VirtualProduction/RemoteControl/Source/WebRemoteControl/Private/RemoteControlResponse.h b/Engine/Plugins/VirtualProduction/RemoteControl/Source/WebRemoteControl/Private/RemoteControlResponse.h
--- a/Engine/Plugins/VirtualProduction/RemoteControl/Source/WebRemoteControl/Private/RemoteControlResponse.h
+++ b/Engine/Plugins/VirtualProduction/RemoteControl/Source/WebRemoteControl/Private/RemoteControlResponse.h
@@ -396,4 +396,36 @@ struct FRCPresetEntitiesModifiedEvent
 	 */
 	UPROPERTY()
 	FRCPresetModifiedEntitiesDescription ModifiedEntities;
-};
\ No newline at end of file
+};
+
+/**
+ * Event triggered when the list of actors in the current scene (or their names) changes.
+ */
+USTRUCT()
+struct FRCActorsChangedEvent
+{
+	GENERATED_BODY()
+
+	FRCActorsChangedEvent()
+		: Type(TEXT("ActorsChanged"))
+	{
+	}
+
+	/**
+	 * Type of the event.
+	 */
+	UPROPERTY()
+	FString Type;
+
+	/** Actors that were added. */
+	UPROPERTY()
+	TArray<FRCActorDescription> AddedActors;
+
+	/** Actors that were renamed. */
+	UPROPERTY()
+	TArray<FRCActorDescription> RenamedActors;
+
+	/** Actors that were deleted. */
+	UPROPERTY()
+	TArray<FRCActorDescription> DeletedActors;
+};
diff --git a/Engine/Plugins/VirtualProduction/RemoteControl/Source/WebRemoteControl/Private/WebSocketMessageHandler.cpp b/Engine/Plugins/VirtualProduction/RemoteControl/Source/WebRemoteControl/Private/WebSocketMessageHandler.cpp
--- a/Engine/Plugins/VirtualProduction/RemoteControl/Source/WebRemoteControl/Private/WebSocketMessageHandler.cpp
+++ b/Engine/Plugins/VirtualProduction/RemoteControl/Source/WebRemoteControl/Private/WebSocketMessageHandler.cpp
@@ -5,9 +5,9 @@
 #include "Algo/ForEach.h"
 #include "Misc/CoreDelegates.h"
 #include "Misc/Guid.h"
+#include "EngineUtils.h"
 #include "IRemoteControlModule.h"
 #include "GameFramework/Actor.h"
-#include "RemoteControlModels.h"
 #include "RemoteControlPreset.h"
 #include "RemoteControlRequest.h"
 #include "RemoteControlResponse.h"
@@ -17,6 +17,10 @@
 #include "WebRemoteControl.h"
 #include "WebRemoteControlInternalUtils.h"
 
+#if WITH_EDITOR
+#include "Editor.h"
+#endif
+
 static TAutoConsoleVariable<int32> CVarWebRemoteControlFramesBetweenPropertyNotifications(TEXT("WebControl.FramesBetweenPropertyNotifications"), 5, TEXT("The number of frames between sending batches of property notifications."));
 
 namespace WebSocketMessageHandlerStructUtils
@@ -208,6 +212,18 @@ namespace WebSocketMessageHandlerStructUtils
 	}
 }
 
+namespace WebSocketMessageHandlerMiscUtils
+{
+	UWorld* GetEditorWorld()
+	{
+#if WITH_EDITOR
+		return GEditor ? GEditor->GetEditorWorldContext(false).World() : nullptr;
+#else
+		return nullptr;
+#endif
+	}
+}
+
 FWebSocketMessageHandler::FWebSocketMessageHandler(FRCWebSocketServer* InServer, const FGuid& InActingClientId)
 	: Server(InServer)
 	, ActingClientId(InActingClientId)
@@ -218,26 +234,59 @@ FWebSocketMessageHandler::FWebSocketMessageHandler(FRCWebSocketServer* InServer,
 void FWebSocketMessageHandler::RegisterRoutes(FWebRemoteControlModule* WebRemoteControl)
 {
 	FCoreDelegates::OnEndFrame.AddRaw(this, &FWebSocketMessageHandler::OnEndFrame);
+
+#if WITH_EDITOR
+	FCoreUObjectDelegates::OnObjectPropertyChanged.AddRaw(this, &FWebSocketMessageHandler::OnObjectPropertyChanged);
+	FCoreUObjectDelegates::OnObjectTransacted.AddRaw(this, &FWebSocketMessageHandler::OnObjectTransacted);
+
 	Server->OnConnectionClosed().AddRaw(this, &FWebSocketMessageHandler::OnConnectionClosedCallback);
+
+	if (GEngine)
+	{
+		RegisterActorHandlers();
+	}
+	else
+	{
+		FCoreDelegates::OnPostEngineInit.AddRaw(this, &FWebSocketMessageHandler::RegisterActorHandlers);
+	}
+#endif
 	
 	// WebSocket routes
-	TUniquePtr<FRemoteControlWebsocketRoute> RegisterRoute = MakeUnique<FRemoteControlWebsocketRoute>(
-		TEXT("Route a message for custom websocket route"),
+	TUniquePtr<FRemoteControlWebsocketRoute> PresetRegisterRoute = MakeUnique<FRemoteControlWebsocketRoute>(
+		TEXT("Subscribe to events emitted by a Remote Control Preset"),
 		TEXT("preset.register"),
 		FWebSocketMessageDelegate::CreateRaw(this, &FWebSocketMessageHandler::HandleWebSocketPresetRegister)
 		);
 
-	WebRemoteControl->RegisterWebsocketRoute(*RegisterRoute);
-	Routes.Emplace(MoveTemp(RegisterRoute));
+	WebRemoteControl->RegisterWebsocketRoute(*PresetRegisterRoute);
+	Routes.Emplace(MoveTemp(PresetRegisterRoute));
 
-	TUniquePtr<FRemoteControlWebsocketRoute> UnregisterRoute = MakeUnique<FRemoteControlWebsocketRoute>(
-		TEXT("Route a message for custom websocket route"),
+	TUniquePtr<FRemoteControlWebsocketRoute> PresetUnregisterRoute = MakeUnique<FRemoteControlWebsocketRoute>(
+		TEXT("Unsubscribe to events emitted by a Remote Control Preset"),
 		TEXT("preset.unregister"),
 		FWebSocketMessageDelegate::CreateRaw(this, &FWebSocketMessageHandler::HandleWebSocketPresetUnregister)
 		);
 
-	WebRemoteControl->RegisterWebsocketRoute(*UnregisterRoute);
-	Routes.Emplace(MoveTemp(UnregisterRoute));
+	WebRemoteControl->RegisterWebsocketRoute(*PresetUnregisterRoute);
+	Routes.Emplace(MoveTemp(PresetUnregisterRoute));
+
+	TUniquePtr<FRemoteControlWebsocketRoute> ActorRegisterRoute = MakeUnique<FRemoteControlWebsocketRoute>(
+		TEXT("Subscribe to events emitted when actors of a particular type are added to/deleted from/renamed in the editor world"),
+		TEXT("actors.register"),
+		FWebSocketMessageDelegate::CreateRaw(this, &FWebSocketMessageHandler::HandleWebSocketActorRegister)
+		);
+
+	WebRemoteControl->RegisterWebsocketRoute(*ActorRegisterRoute);
+	Routes.Emplace(MoveTemp(ActorRegisterRoute));
+
+	TUniquePtr<FRemoteControlWebsocketRoute> ActorUnregisterRoute = MakeUnique<FRemoteControlWebsocketRoute>(
+		TEXT("Unsubscribe to events emitted when actors of a particular type are added to/deleted from/renamed in the editor world"),
+		TEXT("actors.unregister"),
+		FWebSocketMessageDelegate::CreateRaw(this, &FWebSocketMessageHandler::HandleWebSocketActorUnregister)
+		);
+
+	WebRemoteControl->RegisterWebsocketRoute(*ActorUnregisterRoute);
+	Routes.Emplace(MoveTemp(ActorUnregisterRoute));
 }
 
 void FWebSocketMessageHandler::UnregisterRoutes(FWebRemoteControlModule* WebRemoteControl)
@@ -245,15 +294,39 @@ void FWebSocketMessageHandler::UnregisterRoutes(FWebRemoteControlModule* WebRemo
 	Server->OnConnectionClosed().RemoveAll(this);
 	FCoreDelegates::OnEndFrame.RemoveAll(this);
 
+#if WITH_EDITOR
+	FCoreUObjectDelegates::OnObjectPropertyChanged.RemoveAll(this);
+	FCoreUObjectDelegates::OnObjectTransacted.RemoveAll(this);
+
+	if (GEngine)
+	{
+		GEngine->OnLevelActorAdded().Remove(OnActorAddedHandle);
+		GEngine->OnLevelActorDeleted().Remove(OnActorDeletedHandle);
+		GEngine->OnLevelActorListChanged().Remove(OnActorListChangedHandle);
+	}
+#endif
+
 	for (const TUniquePtr<FRemoteControlWebsocketRoute>& Route : Routes)
 	{
 		WebRemoteControl->UnregisterWebsocketRoute(*Route);
 	}
 }
 
+void FWebSocketMessageHandler::RegisterActorHandlers()
+{
+#if WITH_EDITOR
+	if (GEngine)
+	{
+		OnActorAddedHandle = GEngine->OnLevelActorAdded().AddRaw(this, &FWebSocketMessageHandler::OnActorAdded);
+		OnActorDeletedHandle = GEngine->OnLevelActorDeleted().AddRaw(this, &FWebSocketMessageHandler::OnActorDeleted);
+		OnActorListChangedHandle = GEngine->OnLevelActorListChanged().AddRaw(this, &FWebSocketMessageHandler::OnActorListChanged);
+	}
+#endif
+}
+
 void FWebSocketMessageHandler::NotifyPropertyChangedRemotely(const FGuid& OriginClientId, const FGuid& PresetId, const FGuid& ExposedPropertyId)
 {
-	if (TArray<FGuid>* SubscribedClients = WebSocketNotificationMap.Find(PresetId))
+	if (TArray<FGuid>* SubscribedClients = PresetNotificationMap.Find(PresetId))
 	{
 		if (SubscribedClients->Contains(OriginClientId))
 		{
@@ -314,12 +387,12 @@ void FWebSocketMessageHandler::HandleWebSocketPresetRegister(const FRemoteContro
 
 	ClientConfigMap.FindOrAdd(WebSocketMessage.ClientId).bIgnoreRemoteChanges = Body.IgnoreRemoteChanges;
 	
-	TArray<FGuid>* ClientIds = WebSocketNotificationMap.Find(Preset->GetPresetId());
+	TArray<FGuid>* ClientIds = PresetNotificationMap.Find(Preset->GetPresetId());
 
 	// Don't register delegates for a preset more than once.
 	if (!ClientIds)
 	{
-		ClientIds = &WebSocketNotificationMap.Add(Preset->GetPresetId());
+		ClientIds = &PresetNotificationMap.Add(Preset->GetPresetId());
 
 		//Register to any useful callback for the given preset
 		Preset->OnExposedPropertiesModified().AddRaw(this, &FWebSocketMessageHandler::OnPresetExposedPropertiesModified);
@@ -359,13 +432,88 @@ void FWebSocketMessageHandler::HandleWebSocketPresetUnregister(const FRemoteCont
 
 	if (Preset)
 	{
-		if (TArray<FGuid>* RegisteredClients = WebSocketNotificationMap.Find(Preset->GetPresetId()))
+		if (TArray<FGuid>* RegisteredClients = PresetNotificationMap.Find(Preset->GetPresetId()))
 		{
 			RegisteredClients->Remove(WebSocketMessage.ClientId);
 		}
 	}
 }
 
+void FWebSocketMessageHandler::HandleWebSocketActorRegister(const FRemoteControlWebSocketMessage& WebSocketMessage)
+{
+	FRCWebSocketActorRegisterBody Body;
+	if (!WebRemoteControlInternalUtils::DeserializeRequestPayload(WebSocketMessage.RequestPayload, nullptr, Body))
+	{
+		return;
+	}
+
+	const TSubclassOf<AActor> ActorClass = StaticLoadClass(AActor::StaticClass(), nullptr, *Body.ClassName.ToString());
+	if (!ActorClass)
+	{
+		return;
+	}
+
+	const UWorld* World = WebSocketMessageHandlerMiscUtils::GetEditorWorld();
+	if (!World)
+	{
+		return;
+	}
+
+	// Register the client for future updates
+	TArray<FGuid>* ClientIds = ActorNotificationMap.Find(ActorClass);
+	if (!ClientIds)
+	{
+		ClientIds = &ActorNotificationMap.Add(ActorClass);
+	}
+
+	ClientIds->AddUnique(WebSocketMessage.ClientId);
+
+	// Register events for each actor and send the existing list of actors as "added" so the client is caught up
+	FRCActorsChangedEvent Event;
+	for (AActor* Actor : TActorRange<AActor>(World, ActorClass))
+	{
+		Event.AddedActors.Add(FRCActorDescription(Actor));
+		StartWatchingActor(Actor, WebSocketMessage.ClientId);
+	}
+
+	TArray<uint8> Payload;
+	WebRemoteControlUtils::SerializeMessage(Event, Payload);
+	Server->Send(WebSocketMessage.ClientId, Payload);
+}
+
+void FWebSocketMessageHandler::HandleWebSocketActorUnregister(const FRemoteControlWebSocketMessage& WebSocketMessage)
+{
+	FRCWebSocketActorRegisterBody Body;
+	if (!WebRemoteControlInternalUtils::DeserializeRequestPayload(WebSocketMessage.RequestPayload, nullptr, Body))
+	{
+		return;
+	}
+
+	TSubclassOf<AActor> ActorClass = StaticLoadClass(AActor::StaticClass(), nullptr, *Body.ClassName.ToString());
+	if (!ActorClass)
+	{
+		return;
+	}
+
+	// Unregister if already registered
+	if (TArray<FGuid>* RegisteredClients = ActorNotificationMap.Find(ActorClass))
+	{
+		RegisteredClients->Remove(WebSocketMessage.ClientId);
+		if (RegisteredClients->IsEmpty())
+		{
+			ActorNotificationMap.Remove(ActorClass);
+		}
+	}
+
+	// Stop watching all actors for this client
+	TArray<AActor*> WatchedActorPointers;
+	WatchedActors.GetKeys(WatchedActorPointers);
+	for (AActor* Actor : WatchedActorPointers)
+	{
+		StopWatchingActor(Actor, WebSocketMessage.ClientId);
+	}
+}
+
 void FWebSocketMessageHandler::ProcessChangedProperties()
 {
 	//Go over each property that were changed for each preset
@@ -447,7 +595,7 @@ void FWebSocketMessageHandler::OnPropertyExposed(URemoteControlPreset* Owner, co
 		return;
 	}
 
-	if (WebSocketNotificationMap.Num() <= 0)
+	if (PresetNotificationMap.Num() <= 0)
 	{
 		return;
 	}
@@ -463,7 +611,7 @@ void FWebSocketMessageHandler::OnPresetExposedPropertiesModified(URemoteControlP
 		return;
 	}
 
-	if (WebSocketNotificationMap.Num() <= 0)
+	if (PresetNotificationMap.Num() <= 0)
 	{
 		return;
 	}
@@ -471,7 +619,7 @@ void FWebSocketMessageHandler::OnPresetExposedPropertiesModified(URemoteControlP
 	//Cache the property field that changed for end of frame notification
 	TMap<FGuid, TSet<FGuid>>& EventsForClient = PerFrameModifiedProperties.FindOrAdd(Owner->GetPresetId());
 	
-	if (TArray<FGuid>* SubscribedClients = WebSocketNotificationMap.Find(Owner->GetPresetId()))
+	if (TArray<FGuid>* SubscribedClients = PresetNotificationMap.Find(Owner->GetPresetId()))
 	{
 		for (const FGuid& ModifiedPropertyId : ModifiedPropertyIds)
 		{
@@ -505,7 +653,7 @@ void FWebSocketMessageHandler::OnPropertyUnexposed(URemoteControlPreset* Owner,
 		return;
 	}
 
-	if (WebSocketNotificationMap.Num() <= 0)
+	if (PresetNotificationMap.Num() <= 0)
 	{
 		return;
 	}
@@ -526,7 +674,7 @@ void FWebSocketMessageHandler::OnFieldRenamed(URemoteControlPreset* Owner, FName
 		return;
 	}
 
-	if (WebSocketNotificationMap.Num() <= 0)
+	if (PresetNotificationMap.Num() <= 0)
 	{
 		return;
 	}
@@ -543,7 +691,7 @@ void FWebSocketMessageHandler::OnMetadataModified(URemoteControlPreset* Owner)
 		return;
 	}
 
-	if (WebSocketNotificationMap.Num() <= 0)
+	if (PresetNotificationMap.Num() <= 0)
 	{
 		return;
 	}
@@ -559,7 +707,7 @@ void FWebSocketMessageHandler::OnActorPropertyChanged(URemoteControlPreset* Owne
 		return;
 	}
 
-	if (WebSocketNotificationMap.Num() <= 0)
+	if (PresetNotificationMap.Num() <= 0)
 	{
 		return;
 	}
@@ -583,7 +731,7 @@ void FWebSocketMessageHandler::OnActorPropertyChanged(URemoteControlPreset* Owne
 	TMap<FGuid, TMap<FRemoteControlActor, TArray<FRCObjectReference>>>& EventsForClient = PerFrameActorPropertyChanged.FindOrAdd(Owner->GetPresetId());
 
 	// Dont send events to the client that triggered it.
-	if (TArray<FGuid>* SubscribedClients = WebSocketNotificationMap.Find(Owner->GetPresetId()))
+	if (TArray<FGuid>* SubscribedClients = PresetNotificationMap.Find(Owner->GetPresetId()))
 	{
 		for (const FGuid& Client : *SubscribedClients)
 		{
@@ -606,7 +754,7 @@ void FWebSocketMessageHandler::OnEntitiesModified(URemoteControlPreset* Owner, c
 	
 	TArray<uint8> Payload;
 	WebRemoteControlUtils::SerializeMessage(FRCPresetEntitiesModifiedEvent{Owner, ModifiedEntities.Array()}, Payload);
-	BroadcastToListeners(Owner->GetPresetId(), Payload);
+	BroadcastToPresetListeners(Owner->GetPresetId(), Payload);
 }
 
 void FWebSocketMessageHandler::OnLayoutModified(URemoteControlPreset* Owner)
@@ -616,7 +764,7 @@ void FWebSocketMessageHandler::OnLayoutModified(URemoteControlPreset* Owner)
 		return;
 	}
 
-	if (WebSocketNotificationMap.Num() <= 0)
+	if (PresetNotificationMap.Num() <= 0)
 	{
 		return;
 	}
@@ -628,7 +776,7 @@ void FWebSocketMessageHandler::OnLayoutModified(URemoteControlPreset* Owner)
 void FWebSocketMessageHandler::OnConnectionClosedCallback(FGuid ClientId)
 {
 	//Cleanup client that were waiting for callbacks
-	for (auto Iter = WebSocketNotificationMap.CreateIterator(); Iter; ++Iter)
+	for (auto Iter = PresetNotificationMap.CreateIterator(); Iter; ++Iter)
 	{
 		Iter.Value().Remove(ClientId);
 	}
@@ -640,7 +788,7 @@ void FWebSocketMessageHandler::OnConnectionClosedCallback(FGuid ClientId)
 void FWebSocketMessageHandler::OnEndFrame()
 {
 	//Early exit if no clients are requesting notifications
-	if (WebSocketNotificationMap.Num() <= 0)
+	if (PresetNotificationMap.Num() <= 0 && ActorNotificationMap.Num() <= 0)
 	{
 		return;
 	}
@@ -657,6 +805,7 @@ void FWebSocketMessageHandler::OnEndFrame()
 		ProcessRenamedFields();
 		ProcessModifiedMetadata();
 		ProcessModifiedPresetLayouts();
+		ProcessActorChanges();
 	}
 }
 
@@ -697,7 +846,7 @@ void FWebSocketMessageHandler::ProcessAddedProperties()
 
 		TArray<uint8> Payload;
 		WebRemoteControlUtils::SerializeMessage(FRCPresetFieldsAddedEvent{ Preset->GetFName(), Preset->GetPresetId(), AddedPropertiesDescription }, Payload);
-		BroadcastToListeners(Entry.Key, Payload);
+		BroadcastToPresetListeners(Entry.Key, Payload);
 	}
 
 	PerFrameAddedProperties.Empty();
@@ -722,7 +871,7 @@ void FWebSocketMessageHandler::ProcessRemovedProperties()
 		
 		TArray<uint8> Payload;
 		WebRemoteControlUtils::SerializeMessage(FRCPresetFieldsRemovedEvent{ Preset->GetFName(), Preset->GetPresetId(), Entry.Value.Value, Entry.Value.Key }, Payload);
-		BroadcastToListeners(Entry.Key, Payload);
+		BroadcastToPresetListeners(Entry.Key, Payload);
 	}
 	
 	PerFrameRemovedProperties.Empty();
@@ -745,7 +894,7 @@ void FWebSocketMessageHandler::ProcessRenamedFields()
 
 		TArray<uint8> Payload;
 		WebRemoteControlUtils::SerializeMessage(FRCPresetFieldsRenamedEvent{Preset->GetFName(), Preset->GetPresetId(), Entry.Value}, Payload);
-		BroadcastToListeners(Entry.Key, Payload);
+		BroadcastToPresetListeners(Entry.Key, Payload);
 	}
 
 	PerFrameRenamedFields.Empty();
@@ -761,7 +910,7 @@ void FWebSocketMessageHandler::ProcessModifiedMetadata()
 			{
 				TArray<uint8> Payload;
 				WebRemoteControlUtils::SerializeMessage(FRCPresetMetadataModified{ Preset }, Payload);
-				BroadcastToListeners(Entry, Payload);
+				BroadcastToPresetListeners(Entry, Payload);
 			}
 		}
 	}
@@ -779,7 +928,7 @@ void FWebSocketMessageHandler::ProcessModifiedPresetLayouts()
 			{
 				TArray<uint8> Payload;
 				WebRemoteControlUtils::SerializeMessage(FRCPresetLayoutModified{ Preset }, Payload);
-				BroadcastToListeners(Entry, Payload);
+				BroadcastToPresetListeners(Entry, Payload);
 			}
 		}
 	}
@@ -787,9 +936,74 @@ void FWebSocketMessageHandler::ProcessModifiedPresetLayouts()
 	PerFrameModifiedPresetLayouts.Empty();
 }
 
-void FWebSocketMessageHandler::BroadcastToListeners(const FGuid& TargetPresetId, const TArray<uint8>& Payload)
+void FWebSocketMessageHandler::ProcessActorChanges()
+{
+	// Get the set of all clients that will receive updates (so we can batch all update types together)
+	TArray<FGuid, TInlineAllocator<8>> ClientsToNotify;
+	
+	for (const TPair<FGuid, TArray<TWeakObjectPtr<AActor>>>& AddedPair : PerFrameActorsAdded)
+	{
+		ClientsToNotify.AddUnique(AddedPair.Key);
+	}
+
+	for (const TPair<FGuid, TArray<TWeakObjectPtr<AActor>>>& RenamedPair : PerFrameActorsRenamed)
+	{
+		ClientsToNotify.AddUnique(RenamedPair.Key);
+	}
+
+	for (const TPair<FGuid, TArray<FRCActorDescription>>& DeletedPair : PerFrameActorsDeleted)
+	{
+		ClientsToNotify.AddUnique(DeletedPair.Key);
+	}
+
+	// Send the updates to each client
+	for (const FGuid& ClientId : ClientsToNotify)
+	{
+		FRCActorsChangedEvent Event;
+
+		// Added actors
+		if (const TArray<TWeakObjectPtr<AActor>>* AddedActors = PerFrameActorsAdded.Find(ClientId))
+		{
+			for (TWeakObjectPtr<AActor> Actor : *AddedActors)
+			{
+				if (Actor.IsValid())
+				{
+					Event.AddedActors.Add(FRCActorDescription(Actor.Get()));
+				}
+			}
+		}
+
+		// Renamed actors
+		if (const TArray<TWeakObjectPtr<AActor>>* RenamedActors = PerFrameActorsRenamed.Find(ClientId))
+		{
+			for (TWeakObjectPtr<AActor> Actor : *RenamedActors)
+			{
+				if (Actor.IsValid())
+				{
+					Event.RenamedActors.Add(FRCActorDescription(Actor.Get()));
+				}
+			}
+		}
+
+		// Deleted actors
+		if (const TArray<FRCActorDescription>* DeletedActors = PerFrameActorsDeleted.Find(ClientId))
+		{
+			Event.DeletedActors = *DeletedActors;
+		}
+
+		TArray<uint8> Payload;
+		WebRemoteControlUtils::SerializeMessage(Event, Payload);
+		Server->Send(ClientId, Payload);
+	}
+
+	PerFrameActorsAdded.Empty();
+	PerFrameActorsRenamed.Empty();
+	PerFrameActorsDeleted.Empty();
+}
+
+void FWebSocketMessageHandler::BroadcastToPresetListeners(const FGuid& TargetPresetId, const TArray<uint8>& Payload)
 {
-	const TArray<FGuid>& Listeners = WebSocketNotificationMap.FindChecked(TargetPresetId);
+	const TArray<FGuid>& Listeners = PresetNotificationMap.FindChecked(TargetPresetId);
 	for (const FGuid& Listener : Listeners)
 	{
 		Server->Send(Listener, Payload);
@@ -798,7 +1012,7 @@ void FWebSocketMessageHandler::BroadcastToListeners(const FGuid& TargetPresetId,
 
 bool FWebSocketMessageHandler::ShouldProcessEventForPreset(const FGuid& PresetId) const
 {
-	return WebSocketNotificationMap.Contains(PresetId) && WebSocketNotificationMap[PresetId].Num() > 0;
+	return PresetNotificationMap.Contains(PresetId) && PresetNotificationMap[PresetId].Num() > 0;
 }
 
 bool FWebSocketMessageHandler::WritePropertyChangeEventPayload(URemoteControlPreset* InPreset, const TSet<FGuid>& InModifiedPropertyIds, TArray<uint8>& OutBuffer)
@@ -869,3 +1083,257 @@ bool FWebSocketMessageHandler::WriteActorPropertyChangePayload(URemoteControlPre
 
 	return bHasProperty;
 }
+
+void FWebSocketMessageHandler::OnActorAdded(AActor* Actor)
+{
+	TArray<FGuid, TInlineAllocator<8>> WatchingIds;
+
+	const FString ActorPath = Actor->GetPathName();
+
+	for (const TPair<TWeakObjectPtr<UClass>, TArray<FGuid>>& WatchedClasses : ActorNotificationMap)
+	{
+		if (!WatchedClasses.Key.IsValid())
+		{
+			continue;
+		}
+
+		// All the clients in the list are subscribed to events for actors of this type
+		if (Actor->IsA(WatchedClasses.Key.Get()))
+		{
+			for (const FGuid& ClientId : WatchedClasses.Value)
+			{
+				TArray<TWeakObjectPtr<AActor>>& AddedActors = PerFrameActorsAdded.FindOrAdd(ClientId);
+				AddedActors.AddUnique(Actor);
+				WatchingIds.AddUnique(ClientId);
+				
+				// If this actor was queued for a delete event, cancel it so that it's clear that the actor has been re-created.
+				TArray<FRCActorDescription>* DeletedActors = PerFrameActorsDeleted.Find(ClientId);
+				if (DeletedActors)
+				{
+					const int32 DeletedIndex = DeletedActors->IndexOfByPredicate([&ActorPath](const FRCActorDescription& Description) {
+						return Description.Path == ActorPath;
+					});
+
+					if (DeletedIndex != INDEX_NONE)
+					{
+						DeletedActors->RemoveAt(DeletedIndex);
+					}
+				}
+			}
+		}
+	}
+
+	if (WatchingIds.Num() > 0)
+	{
+		// At least one subscriber cares about this actor, so we should listen to its events
+		FWatchedActorData& ActorData = WatchedActors.Add(Actor, FWatchedActorData(Actor));
+		ActorData.Clients = WatchingIds;
+	}
+}
+
+void FWebSocketMessageHandler::OnActorDeleted(AActor* Actor)
+{
+	FWatchedActorData* ActorData = WatchedActors.Find(Actor);
+	if (!ActorData)
+	{
+		return;
+	}
+
+	for (const FGuid& ClientId : ActorData->Clients)
+	{
+		TArray<FRCActorDescription>& DeletedActors = PerFrameActorsDeleted.FindOrAdd(ClientId);
+		DeletedActors.AddUnique(ActorData->Description);
+
+		// If this actor was queued for an add event, cancel it so that it's clear that the actor has been deleted again.
+		TArray<TWeakObjectPtr<AActor>>* AddedActors = PerFrameActorsAdded.Find(ClientId);
+		if (AddedActors)
+		{
+			const int32 AddedIndex = AddedActors->IndexOfByPredicate([&Actor](TWeakObjectPtr<AActor> AddedActor) {
+				return AddedActor.GetEvenIfUnreachable() == Actor;
+			});
+
+			if (AddedIndex != INDEX_NONE)
+			{
+				AddedActors->RemoveAt(AddedIndex);
+			}
+		}
+	}
+
+	WatchedActors.Remove(Actor);
+}
+
+void FWebSocketMessageHandler::OnActorListChanged()
+{
+	// We don't know exactly what changed, so manually check all the actors we know about
+
+	const UWorld* World = WebSocketMessageHandlerMiscUtils::GetEditorWorld();
+	if (!World)
+	{
+		return;
+	}
+
+	TSet<AActor*> RemainingActors;
+	WatchedActors.GetKeys(RemainingActors);
+
+	TArray<AActor*> NewActors;
+
+	// Find any new actors
+	for (AActor* Actor : TActorRange<AActor>(World))
+	{
+		if (!WatchedActors.Contains(Actor))
+		{
+			NewActors.Add(Actor);
+		}
+		RemainingActors.Remove(Actor);
+	}
+
+	// Fire events for any actors that are now missing, which have presumably been deleted
+	for (AActor* Actor : RemainingActors)
+	{
+		OnActorDeleted(Actor);
+	}
+
+	// Stale actors may not have been removed from the map due to how invalid keys are hashed, so remove them directly
+	for (auto It = WatchedActors.CreateIterator(); It; ++It)
+	{
+		TWeakObjectPtr<AActor> Actor = It.Key();
+		if (!Actor.IsValid())
+		{
+			It.RemoveCurrent();
+		}
+	}
+
+	// Fire events for new actors (this must be done second since we could be re-creating actors with the same paths, e.g. by reloading a world)
+	for (AActor* Actor : NewActors)
+	{
+		OnActorAdded(Actor);
+	}
+}
+
+void FWebSocketMessageHandler::OnObjectPropertyChanged(UObject* Object, struct FPropertyChangedEvent& Event)
+{
+	// We only care about name changes
+	const FName LabelProperty(TEXT("ActorLabel"));
+	if (Event.GetPropertyName() != LabelProperty)
+	{
+		return;
+	}
+
+	// We only care about actors
+	AActor* Actor = Cast<AActor>(Object);
+	if (!Actor)
+	{
+		return;
+	}
+
+	// We only care about actors that we're watching
+	FWatchedActorData* ActorData = WatchedActors.Find(Actor);
+	if (!ActorData)
+	{
+		return;
+	}
+
+	UpdateWatchedActorName(Actor, *ActorData);
+}
+
+void FWebSocketMessageHandler::OnObjectTransacted(UObject* Object, const class FTransactionObjectEvent& TransactionEvent)
+{
+	// We only care about undo/redo
+	if (TransactionEvent.GetEventType() != ETransactionObjectEventType::UndoRedo)
+	{
+		return;
+	}
+
+	// We only care about actors
+	AActor* Actor = Cast<AActor>(Object);
+	if (!Actor)
+	{
+		return;
+	}
+
+	// Check if the actor was created/deleted by the transaction
+	if (TransactionEvent.HasPendingKillChange())
+	{
+		if (!IsValid(Actor))
+		{
+			// Actor was undone; treat as a delete
+			OnActorDeleted(Actor);
+		}
+		else
+		{
+			// Actor was redone; treat as a create
+			OnActorAdded(Actor);
+		}
+
+		// In either case, we can bail early since a rename no longer matters
+		return;
+	}
+
+	// We only care about renames for actors that we're watching
+	FWatchedActorData* ActorData = WatchedActors.Find(Actor);
+	if (!ActorData)
+	{
+		return;
+	}
+
+	// Check if the actor was renamed by the transaction
+	const FName LabelProperty(TEXT("ActorLabel"));
+	for (FName Property : TransactionEvent.GetChangedProperties())
+	{
+		if (Property == LabelProperty)
+		{
+			UpdateWatchedActorName(Actor, *ActorData);
+			return;
+		}
+	}
+}
+
+void FWebSocketMessageHandler::StartWatchingActor(AActor* Actor, FGuid ClientId)
+{
+	FWatchedActorData* ActorData = WatchedActors.Find(Actor);
+
+	if (!ActorData)
+	{
+		ActorData = &WatchedActors.Add(Actor, FWatchedActorData(Actor));
+	}
+
+	ActorData->Clients.Add(ClientId);
+}
+
+void FWebSocketMessageHandler::StopWatchingActor(AActor* Actor, FGuid ClientId)
+{
+	FWatchedActorData* ActorData = WatchedActors.Find(Actor);
+	if (ActorData)
+	{
+		bool bAnyRemoved = ActorData->Clients.Remove(ClientId) > 0;
+
+		if (bAnyRemoved && ActorData->Clients.IsEmpty())
+		{
+			// Nobody is watching anymore, so we can forget about it
+			WatchedActors.Remove(Actor);
+		}
+	}
+}
+
+void FWebSocketMessageHandler::UpdateWatchedActorName(AActor* Actor, FWebSocketMessageHandler::FWatchedActorData& ActorData)
+{
+	// Update our cached name
+	ActorData.Description.Name = Actor->GetActorNameOrLabel();
+
+	// Mark that this has been renamed
+	for (const FGuid& ClientId : ActorData.Clients)
+	{
+		if (TArray<TWeakObjectPtr<AActor>>* AddedActors = PerFrameActorsAdded.Find(ClientId))
+		{
+			// If the actor was just added this frame, we don't need to report the rename since the name will be included
+			// with the add event. This happens with copy+paste, which renames immediately after creation.
+			if (AddedActors->Contains(Actor))
+			{
+				continue;
+			}
+		}
+
+		TArray<TWeakObjectPtr<AActor>>& RenamedActors = PerFrameActorsRenamed.FindOrAdd(ClientId);
+		RenamedActors.AddUnique(Actor);
+	}
+}
\ No newline at end of file
diff --git a/Engine/Plugins/VirtualProduction/RemoteControl/Source/WebRemoteControl/Private/WebSocketMessageHandler.h b/Engine/Plugins/VirtualProduction/RemoteControl/Source/WebRemoteControl/Private/WebSocketMessageHandler.h
--- a/Engine/Plugins/VirtualProduction/RemoteControl/Source/WebRemoteControl/Private/WebSocketMessageHandler.h
+++ b/Engine/Plugins/VirtualProduction/RemoteControl/Source/WebRemoteControl/Private/WebSocketMessageHandler.h
@@ -6,6 +6,7 @@
 
 #include "RemoteControlActor.h"
 #include "RemoteControlField.h"
+#include "RemoteControlModels.h"
 #include "RemoteControlPreset.h"
 
 
@@ -35,6 +36,24 @@ public:
 	void NotifyPropertyChangedRemotely(const FGuid& OriginClientId, const FGuid& PresetId, const FGuid& ExposedPropertyId);
 
 private:
+
+	/** Data about a watched actor so we know who to notify and what to send if the actor is garbage-collected before we know it's been deleted. */
+	struct FWatchedActorData
+	{
+		FWatchedActorData(AActor* InActor)
+			: Description(InActor)
+		{
+		}
+
+		/** Description of the actor. */
+		FRCActorDescription Description;
+
+		/** Which clients are listening for events about this actor. */
+		TArray<FGuid> Clients;
+	};
+
+	/** Register handlers for actors being added/deleted (must happen after engine init). */
+	void RegisterActorHandlers();
 	
 	/** Handles registration to callbacks to a given preset */
 	void HandleWebSocketPresetRegister(const FRemoteControlWebSocketMessage& WebSocketMessage);
@@ -42,6 +61,12 @@ private:
 	/** Handles unregistration to callbacks to a given preset */
 	void HandleWebSocketPresetUnregister(const FRemoteControlWebSocketMessage& WebSocketMessage);
 
+	/** Handles registration to callbacks for creation/destruction/rename of a given actor type */
+	void HandleWebSocketActorRegister(const FRemoteControlWebSocketMessage& WebSocketMessage);
+
+	/** Handles registration to callbacks for creation/destruction/rename of a given actor type */
+	void HandleWebSocketActorUnregister(const FRemoteControlWebSocketMessage& WebSocketMessage);
+
 	//Preset callbacks
 	void OnPresetExposedPropertiesModified(URemoteControlPreset* Owner, const TSet<FGuid>& ModifiedPropertyIds);
 	void OnPropertyExposed(URemoteControlPreset* Owner,  const FGuid& EntityId);
@@ -79,11 +104,14 @@ private:
 	/** If a preset layout is modified, notify listeners. */
 	void ProcessModifiedPresetLayouts();
 
+	/** If actors were added/removed/renamed, notify listeners. */
+	void ProcessActorChanges();
+
 	/** 
 	 * Send a payload to all clients bound to a certain preset.
-	 * @note: TargetPresetName must be in the WebSocketNotificationMap.
+	 * @note: TargetPresetName must be in the PresetNotificationMap.
 	 */
-	void BroadcastToListeners(const FGuid& TargetPresetId, const TArray<uint8>& Payload);
+	void BroadcastToPresetListeners(const FGuid& TargetPresetId, const TArray<uint8>& Payload);
 
 	/**
 	 * Returns whether an event targeting a particular preset should be processed.
@@ -100,6 +128,46 @@ private:
 	 */
 	bool WriteActorPropertyChangePayload(URemoteControlPreset* InPreset, const TMap<FRemoteControlActor, TArray<FRCObjectReference>>& InModifications, FMemoryWriter& InWriter);
 
+	/**
+	 * Called when an actor is added to the current world.
+	 */
+	void OnActorAdded(AActor* Actor);
+
+	/**
+	 * Called when an actor is deleted from the current world.
+	 */
+	void OnActorDeleted(AActor* Actor);
+
+	/**
+	 * Called when an untracked change to the actor list happens in the editor.
+	 */
+	void OnActorListChanged();
+
+	/**
+	 * Called when an object's property is changed in the editor. Used to detect name changes on subscribed actors.
+	 */
+	void OnObjectPropertyChanged(UObject* Object, struct FPropertyChangedEvent& Event);
+
+	/**
+	 * Called when an object is affected by an editor transaction. Used to detect undo/redo of creating/deleting subscribed actors.
+	 */
+	void OnObjectTransacted(UObject* Object, const class FTransactionObjectEvent& TransactionEvent);
+
+	/**
+	 * Start watching an actor for the given client.
+	 */
+	void StartWatchingActor(AActor* Actor, FGuid ClientId);
+
+	/**
+	 * Stop watching an actor for the given client.
+	 */
+	void StopWatchingActor(AActor* Actor, FGuid ClientId);
+
+	/**
+	 * Update our cache of watched actor's name and notify any subscribers.
+	 */
+	void UpdateWatchedActorName(AActor* Actor, FWatchedActorData& ActorData);
+
 private:
 
 	/** Web Socket server. */
@@ -108,7 +176,11 @@ private:
 	TArray<TUniquePtr<FRemoteControlWebsocketRoute>> Routes;
 
 	/** All websockets connections associated to a preset notifications */
-	TMap<FGuid, TArray<FGuid>> WebSocketNotificationMap;
+	TMap<FGuid, TArray<FGuid>> PresetNotificationMap;
+
+	/** All websocket client IDs associated with an actor class. */
+	typedef TMap<TWeakObjectPtr<UClass>, TArray<FGuid>, FDefaultSetAllocator, TWeakObjectPtrMapKeyFuncs<TWeakObjectPtr<UClass>, TArray<FGuid>>> FActorNotificationMap;
+	FActorNotificationMap ActorNotificationMap;
 
 	/** Configuration for a given client related to how events should be handled. */
 	struct FRCClientConfig
@@ -141,6 +213,18 @@ private:
 	/** Fields that were renamed for a frame, per preset */
 	TMap<FGuid, TArray<TTuple<FName, FName>>> PerFrameRenamedFields;
 
+	/** Actors that were added for a frame, per subscribed client. */
+	TMap<FGuid, TArray<TWeakObjectPtr<AActor>>> PerFrameActorsAdded;
+
+	/** Actors that were renamed for a frame, per subscribed client. */
+	TMap<FGuid, TArray<TWeakObjectPtr<AActor>>> PerFrameActorsRenamed;
+
+	/**
+	 * Actors that were removed for a frame, per subscribed client.
+	 * Stored by description in case the actor is garbage collected before its name and path can be collected.
+	 */
+	TMap<FGuid, TArray<FRCActorDescription>> PerFrameActorsDeleted;
+
 	/** Presets that had their metadata modified for a frame */
 	TSet<FGuid> PerFrameModifiedMetadata;
 
@@ -152,4 +236,19 @@ private:
 	
 	/** Frame counter for delaying property change checks. */
 	int32 PropertyNotificationFrameCounter = 0;
+
+	/** Handle for when an actor is added to the world. */
+	FDelegateHandle OnActorAddedHandle;
+
+	/** Handle for when an actor is deleted from the world. */
+	FDelegateHandle OnActorDeletedHandle;
+
+	/** Handle for when the list of actors changes. */
+	FDelegateHandle OnActorListChangedHandle;
+
+	/**
+	 * Actors that we are actively watching to send events to subscribers.
+	 * The key is not a weak pointer, so it shouldn't be accessed in case it's stale. Use the value instead.
+	 */
+	TMap<AActor*, FWatchedActorData> WatchedActors;
 };
\ No newline at end of file
