RemoteControl: Have the multicast beacon listener join the multicast group on all interfaces instead of just the default interface (0.0.0.0).

With this, Unreal instances running on PCs with multiple interfaces, including VPNs, can receive the multicast beacon packet from the tablet device when the default interface is not in the corresponding subnet.

diff --git a/Engine/Plugins/VirtualProduction/RemoteControl/Source/WebRemoteControl/Private/StageAppBeaconReceiver.cpp b/Engine/Plugins/VirtualProduction/RemoteControl/Source/WebRemoteControl/Private/StageAppBeaconReceiver.cpp
--- a/Engine/Plugins/VirtualProduction/RemoteControl/Source/WebRemoteControl/Private/StageAppBeaconReceiver.cpp
+++ b/Engine/Plugins/VirtualProduction/RemoteControl/Source/WebRemoteControl/Private/StageAppBeaconReceiver.cpp
@@ -33,22 +33,48 @@ void FStageAppBeaconReceiver::Startup()
 		return;
 	}
 
+	WebsocketPort = GetDefault<URemoteControlSettings>()->RemoteControlWebSocketServerPort;
+
 	Socket = FUdpSocketBuilder(TEXT("StageAppBeaconResponder"))
 		.AsNonBlocking()
 		.AsReusable()
 		.BoundToPort(Settings.DiscoveryPort)
-		.JoinedToGroup(DiscoveryAddress)
+		.JoinedToGroup(DiscoveryAddress) // joins default interface
 		.WithMulticastLoopback()
 		.WithReceiveBufferSize(256);
 
-	WebsocketPort = GetDefault<URemoteControlSettings>()->RemoteControlWebSocketServerPort;
-
 	if (!Socket)
 	{
 		UE_LOG(LogRemoteControl, Warning, TEXT("StageAppBeaconResponder failed to create multicast socket"));
 		return;
 	}
 
+	// Listen for multicast packets on all interfaces.
+	// Note: If a new interface appears later on, the socket will need to join it in order to receive multicast packets from it.
+	{
+		TSharedRef<FInternetAddr> DiscoveryInternetAddr = FIPv4Endpoint(DiscoveryAddress, 0).ToInternetAddr(); // Note: Port will be ignored by JoinMulticastGroup
+
+		TArray<TSharedPtr<FInternetAddr>> LocapIps;
+		ISocketSubsystem::Get(PLATFORM_SOCKETSUBSYSTEM)->GetLocalAdapterAddresses(LocapIps);
+
+		for (TSharedPtr<FInternetAddr>& LocapIp : LocapIps)
+		{
+			if (!LocapIp.IsValid())
+			{
+				continue;
+			}
+			
+			const bool bJoinedGroup = Socket->JoinMulticastGroup(*DiscoveryInternetAddr, *LocapIp);
+
+			if (!bJoinedGroup)
+			{
+				UE_LOG(LogRemoteControl, Warning, 
+					TEXT("StageAppBeaconResponder failed to join multicast group '%s' on detected local interface '%s'"), 
+					*DiscoveryAddress.ToString(), *LocapIp->ToString(false));
+			}
+		}
+	}
+
 	Thread.Reset(FRunnableThread::Create(this, TEXT("StageAppBeaconResponderThread"), 0, TPri_Normal, FPlatformAffinity::GetPoolThreadMask()));
 }
 
