[WebRemoteControl] Change actor websocket subscription to split events by subscribed class

diff --git a/Engine/Plugins/VirtualProduction/RemoteControl/Source/WebRemoteControl/Private/RemoteControlResponse.h b/Engine/Plugins/VirtualProduction/RemoteControl/Source/WebRemoteControl/Private/RemoteControlResponse.h
--- a/Engine/Plugins/VirtualProduction/RemoteControl/Source/WebRemoteControl/Private/RemoteControlResponse.h
+++ b/Engine/Plugins/VirtualProduction/RemoteControl/Source/WebRemoteControl/Private/RemoteControlResponse.h
@@ -398,6 +398,27 @@ struct FRCPresetEntitiesModifiedEvent
 	FRCPresetModifiedEntitiesDescription ModifiedEntities;
 };
 
+/**
+ * Data about actors that have changed in the scene.
+ */
+USTRUCT()
+struct FRCActorsChangedData
+{
+	GENERATED_BODY()
+
+	/** Actors that were added. */
+	UPROPERTY()
+	TArray<FRCActorDescription> AddedActors;
+
+	/** Actors that were renamed. */
+	UPROPERTY()
+	TArray<FRCActorDescription> RenamedActors;
+
+	/** Actors that were deleted. */
+	UPROPERTY()
+	TArray<FRCActorDescription> DeletedActors;
+};
+
 /**
  * Event triggered when the list of actors in the current scene (or their names) changes.
  */
@@ -417,15 +438,9 @@ struct FRCActorsChangedEvent
 	UPROPERTY()
 	FString Type;
 
-	/** Actors that were added. */
-	UPROPERTY()
-	TArray<FRCActorDescription> AddedActors;
-
-	/** Actors that were renamed. */
-	UPROPERTY()
-	TArray<FRCActorDescription> RenamedActors;
-
-	/** Actors that were deleted. */
+	/**
+	 * Map from class name to changes in actors of that type.
+	 */
 	UPROPERTY()
-	TArray<FRCActorDescription> DeletedActors;
+	TMap<FString, FRCActorsChangedData> Changes;
 };
diff --git a/Engine/Plugins/VirtualProduction/RemoteControl/Source/WebRemoteControl/Private/WebSocketMessageHandler.cpp b/Engine/Plugins/VirtualProduction/RemoteControl/Source/WebRemoteControl/Private/WebSocketMessageHandler.cpp
--- a/Engine/Plugins/VirtualProduction/RemoteControl/Source/WebRemoteControl/Private/WebSocketMessageHandler.cpp
+++ b/Engine/Plugins/VirtualProduction/RemoteControl/Source/WebRemoteControl/Private/WebSocketMessageHandler.cpp
@@ -10,7 +10,6 @@
 #include "GameFramework/Actor.h"
 #include "RemoteControlPreset.h"
 #include "RemoteControlRequest.h"
-#include "RemoteControlResponse.h"
 #include "RemoteControlRoute.h"
 #include "RemoteControlReflectionUtils.h"
 #include "RemoteControlWebsocketRoute.h"
@@ -459,21 +458,26 @@ void FWebSocketMessageHandler::HandleWebSocketActorRegister(const FRemoteControl
 		return;
 	}
 
-	// Register the client for future updates
-	TArray<FGuid>* ClientIds = ActorNotificationMap.Find(ActorClass);
-	if (!ClientIds)
+	FWatchedClassData* WatchedClassData = ActorNotificationMap.Find(ActorClass);
+	FString ClassPath = ActorClass->GetPathName();
+
+	// Start watching the class if we aren't already
+	if (!WatchedClassData)
 	{
-		ClientIds = &ActorNotificationMap.Add(ActorClass);
+		WatchedClassData = &ActorNotificationMap.Add(ActorClass);
+		WatchedClassData->CachedPath = ClassPath;
 	}
 
-	ClientIds->AddUnique(WebSocketMessage.ClientId);
+	// Register the client for future updates
+	WatchedClassData->Clients.AddUnique(WebSocketMessage.ClientId);
 
 	// Register events for each actor and send the existing list of actors as "added" so the client is caught up
 	FRCActorsChangedEvent Event;
+	FRCActorsChangedData& ChangeData = Event.Changes.Add(ClassPath);
 	for (AActor* Actor : TActorRange<AActor>(World, ActorClass))
 	{
-		Event.AddedActors.Add(FRCActorDescription(Actor));
-		StartWatchingActor(Actor, WebSocketMessage.ClientId);
+		ChangeData.AddedActors.Add(FRCActorDescription(Actor));
+		StartWatchingActor(Actor, ActorClass);
 	}
 
 	TArray<uint8> Payload;
@@ -495,23 +499,7 @@ void FWebSocketMessageHandler::HandleWebSocketActorUnregister(const FRemoteContr
 		return;
 	}
 
-	// Unregister if already registered
-	if (TArray<FGuid>* RegisteredClients = ActorNotificationMap.Find(ActorClass))
-	{
-		RegisteredClients->Remove(WebSocketMessage.ClientId);
-		if (RegisteredClients->IsEmpty())
-		{
-			ActorNotificationMap.Remove(ActorClass);
-		}
-	}
-
-	// Stop watching all actors for this client
-	TArray<AActor*> WatchedActorPointers;
-	WatchedActors.GetKeys(WatchedActorPointers);
-	for (AActor* Actor : WatchedActorPointers)
-	{
-		StopWatchingActor(Actor, WebSocketMessage.ClientId);
-	}
+	UnregisterClientForActorClass(WebSocketMessage.ClientId, ActorClass);
 }
 
 void FWebSocketMessageHandler::ProcessChangedProperties()
@@ -775,12 +763,21 @@ void FWebSocketMessageHandler::OnLayoutModified(URemoteControlPreset* Owner)
 
 void FWebSocketMessageHandler::OnConnectionClosedCallback(FGuid ClientId)
 {
-	//Cleanup client that were waiting for callbacks
+	// Clean up clients that were waiting for preset callbacks
 	for (auto Iter = PresetNotificationMap.CreateIterator(); Iter; ++Iter)
 	{
 		Iter.Value().Remove(ClientId);
 	}
 
+	// Clean up clients that were waiting for actor callbacks
+	TArray<TWeakObjectPtr<UClass>> WatchedClasses;
+	ActorNotificationMap.GenerateKeyArray(WatchedClasses);
+
+	for (const TWeakObjectPtr<UClass>& WatchedClass : WatchedClasses)
+	{
+		UnregisterClientForActorClass(ClientId, WatchedClass.Get());
+	}
+
 	/** Remove this client's config. */
 	ClientConfigMap.Remove(ClientId);
 }
@@ -938,67 +935,169 @@ void FWebSocketMessageHandler::ProcessModifiedPresetLayouts()
 
 void FWebSocketMessageHandler::ProcessActorChanges()
 {
-	// Get the set of all clients that will receive updates (so we can batch all update types together)
-	TArray<FGuid, TInlineAllocator<8>> ClientsToNotify;
+	// Get the set of all classes with subscribed clients (so we can batch all update types together)
+	TArray<UClass*, TInlineAllocator<8>> ChangedClasses;
+	TArray<const FDeletedActorsData*> DeletedClasses;
 	
-	for (const TPair<FGuid, TArray<TWeakObjectPtr<AActor>>>& AddedPair : PerFrameActorsAdded)
+	for (const TPair<TWeakObjectPtr<UClass>, TArray<TWeakObjectPtr<AActor>>>& AddedPair : PerFrameActorsAdded)
+	{
+		if (AddedPair.Key.IsValid())
+		{
+			ChangedClasses.AddUnique(AddedPair.Key.Get());
+		}
+	}
+
+	for (const TPair<TWeakObjectPtr<UClass>, TArray<TWeakObjectPtr<AActor>>>& RenamedPair : PerFrameActorsRenamed)
+	{
+		if (RenamedPair.Key.IsValid())
+		{
+			ChangedClasses.AddUnique(RenamedPair.Key.Get());
+		}
+	}
+
+	for (const TPair<TWeakObjectPtr<UClass>, FDeletedActorsData>& DeletedPair : PerFrameActorsDeleted)
+	{
+		if (DeletedPair.Key.IsValid())
+		{
+			ChangedClasses.AddUnique(DeletedPair.Key.Get());
+		}
+		else
+		{
+			// The class itself was deleted, so we'll have to look it up by path instead of by pointer
+			DeletedClasses.Add(&DeletedPair.Value);
+		}
+	}
+
+	if (ChangedClasses.Num() == 0 && DeletedClasses.Num() == 0)
 	{
-		ClientsToNotify.AddUnique(AddedPair.Key);
+		return;
 	}
 
-	for (const TPair<FGuid, TArray<TWeakObjectPtr<AActor>>>& RenamedPair : PerFrameActorsRenamed)
+	// Map from actor class' path to changed actor data for that class.
+	TMap<FString, FRCActorsChangedData> ChangesByClassPath;
+
+	// Map from client ID to which class paths they're going to get an update about.
+	TMap<FGuid, TArray<FString>> ClientsToNotify;
+
+	// Gather changes for each class that still exists
+	for (UClass* ActorClass : ChangedClasses)
 	{
-		ClientsToNotify.AddUnique(RenamedPair.Key);
+		GatherActorChangesForClass(ActorClass, ChangesByClassPath, ClientsToNotify);
 	}
 
-	for (const TPair<FGuid, TArray<FRCActorDescription>>& DeletedPair : PerFrameActorsDeleted)
+	// Gather changes for deleted classes, which we need to handle different since their pointers are invalid
+	for (const FDeletedActorsData* DeletedActorsData : DeletedClasses)
 	{
-		ClientsToNotify.AddUnique(DeletedPair.Key);
+		GatherActorChangesForDeletedClass(DeletedActorsData, ChangesByClassPath, ClientsToNotify);
 	}
 
-	// Send the updates to each client
-	for (const FGuid& ClientId : ClientsToNotify)
+	// Update each client that cares about the changes we're processing
+	for (const TPair<FGuid, TArray<FString>>& ClientData : ClientsToNotify)
 	{
 		FRCActorsChangedEvent Event;
 
-		// Added actors
-		if (const TArray<TWeakObjectPtr<AActor>>* AddedActors = PerFrameActorsAdded.Find(ClientId))
+		for (const FString& ActorClassPath : ClientData.Value)
 		{
-			for (TWeakObjectPtr<AActor> Actor : *AddedActors)
+			const FRCActorsChangedData* ChangeData = ChangesByClassPath.Find(ActorClassPath);
+			if (!ensureMsgf(ChangeData != nullptr, TEXT("Found no change data for an actor class that supposedly has changes")))
 			{
-				if (Actor.IsValid())
-				{
-					Event.AddedActors.Add(FRCActorDescription(Actor.Get()));
-				}
+				continue;
+			}
+
+			Event.Changes.Add(ActorClassPath, *ChangeData);
+		}
+
+		TArray<uint8> Payload;
+		WebRemoteControlUtils::SerializeMessage(Event, Payload);
+		Server->Send(ClientData.Key, Payload);
+	}
+
+	PerFrameActorsAdded.Empty();
+	PerFrameActorsRenamed.Empty();
+	PerFrameActorsDeleted.Empty();
+}
+
+void FWebSocketMessageHandler::GatherActorChangesForClass(UClass* ActorClass, TMap<FString, FRCActorsChangedData>& OutChangesByClassPath,
+	TMap<FGuid, TArray<FString>>& OutClientsToNotify)
+{
+	const FWatchedClassData* WatchedData = ActorNotificationMap.Find(ActorClass);
+
+	if (!ensureMsgf(WatchedData != nullptr, TEXT("An actor was still being watched for a class that is no longer in ActorNotificationMap")))
+	{
+		return;
+	}
+
+	// Each client watching this class should be notified about the changes
+	const FString ActorClassPath = ActorClass->GetPathName();
+	for (const FGuid& ClientId : WatchedData->Clients)
+	{
+		OutClientsToNotify.FindOrAdd(ClientId).AddUnique(ActorClassPath);
+	}
+
+	FRCActorsChangedData& ChangeData = OutChangesByClassPath.Add(ActorClassPath);
+
+	// Added actors
+	if (const TArray<TWeakObjectPtr<AActor>>* AddedActors = PerFrameActorsAdded.Find(ActorClass))
+	{
+		for (const TWeakObjectPtr<AActor>& Actor : *AddedActors)
+		{
+			if (Actor.IsValid())
+			{
+				ChangeData.AddedActors.Add(FRCActorDescription(Actor.Get()));
 			}
 		}
+	}
 
-		// Renamed actors
-		if (const TArray<TWeakObjectPtr<AActor>>* RenamedActors = PerFrameActorsRenamed.Find(ClientId))
+	// Renamed actors
+	if (const TArray<TWeakObjectPtr<AActor>>* RenamedActors = PerFrameActorsRenamed.Find(ActorClass))
+	{
+		for (const TWeakObjectPtr<AActor>& Actor : *RenamedActors)
 		{
-			for (TWeakObjectPtr<AActor> Actor : *RenamedActors)
+			if (Actor.IsValid())
 			{
-				if (Actor.IsValid())
-				{
-					Event.RenamedActors.Add(FRCActorDescription(Actor.Get()));
-				}
+				ChangeData.RenamedActors.Add(FRCActorDescription(Actor.Get()));
 			}
 		}
+	}
+
+	// Deleted actors
+	if (const FDeletedActorsData* DeletedActorsData = PerFrameActorsDeleted.Find(ActorClass))
+	{
+		ChangeData.DeletedActors = DeletedActorsData->Actors;
+	}
+}
+
+void FWebSocketMessageHandler::GatherActorChangesForDeletedClass(const FWebSocketMessageHandler::FDeletedActorsData* DeletedActorsData,
+	TMap<FString, FRCActorsChangedData>& OutChangesByClassPath, TMap<FGuid, TArray<FString>>& OutClientsToNotify)
+{
+	const FString& ActorClassPath = DeletedActorsData->ClassPath;
 
-		// Deleted actors
-		if (const TArray<FRCActorDescription>* DeletedActors = PerFrameActorsDeleted.Find(ClientId))
+	FRCActorsChangedData& ChangeData = OutChangesByClassPath.Add(ActorClassPath);
+	ChangeData.DeletedActors = DeletedActorsData->Actors;
+
+	// Find who we're supposed to notify about this. We have to manually iterate the map to find it since
+	// the class pointer is now invalid.
+	const TArray<FGuid>* WatchingClients = nullptr;
+
+	for (auto It = ActorNotificationMap.CreateConstIterator(); It; ++It)
+	{
+		if (It->Value.CachedPath == ActorClassPath)
 		{
-			Event.DeletedActors = *DeletedActors;
+			WatchingClients = &It->Value.Clients;
+			break;
 		}
+	}
 
-		TArray<uint8> Payload;
-		WebRemoteControlUtils::SerializeMessage(Event, Payload);
-		Server->Send(ClientId, Payload);
+	if (!WatchingClients)
+	{
+		ensureMsgf(false, TEXT("An actor was still being watched for a deleted class that is no longer in ActorNotificationMap"));
+		return;
 	}
 
-	PerFrameActorsAdded.Empty();
-	PerFrameActorsRenamed.Empty();
-	PerFrameActorsDeleted.Empty();
+	for (const FGuid& ClientId : *WatchingClients)
+	{
+		OutClientsToNotify.FindOrAdd(ClientId).AddUnique(ActorClassPath);
+	}
 }
 
 void FWebSocketMessageHandler::BroadcastToPresetListeners(const FGuid& TargetPresetId, const TArray<uint8>& Payload)
@@ -1086,48 +1185,49 @@ bool FWebSocketMessageHandler::WriteActorPropertyChangePayload(URemoteControlPre
 
 void FWebSocketMessageHandler::OnActorAdded(AActor* Actor)
 {
-	TArray<FGuid, TInlineAllocator<8>> WatchingIds;
+	// Array of classes this actor is a child of and which are being watched by a client
+	TArray<TWeakObjectPtr<UClass>, TInlineAllocator<8>> WatchedClasses;
 
 	const FString ActorPath = Actor->GetPathName();
 
-	for (const TPair<TWeakObjectPtr<UClass>, TArray<FGuid>>& WatchedClasses : ActorNotificationMap)
+	for (const TPair<TWeakObjectPtr<UClass>, FWatchedClassData>& WatchedClassPair : ActorNotificationMap)
 	{
-		if (!WatchedClasses.Key.IsValid())
+		UClass* WatchedClass = WatchedClassPair.Key.Get();
+
+		if (WatchedClass == nullptr)
 		{
+			// We don't need to send an add if the class has already been deleted
 			continue;
 		}
 
-		// All the clients in the list are subscribed to events for actors of this type
-		if (Actor->IsA(WatchedClasses.Key.Get()))
+		// Any classes in this list have at least one client subscribed to updates
+		if (Actor->IsA(WatchedClass))
 		{
-			for (const FGuid& ClientId : WatchedClasses.Value)
-			{
-				TArray<TWeakObjectPtr<AActor>>& AddedActors = PerFrameActorsAdded.FindOrAdd(ClientId);
-				AddedActors.AddUnique(Actor);
-				WatchingIds.AddUnique(ClientId);
+			TArray<TWeakObjectPtr<AActor>>& AddedActors = PerFrameActorsAdded.FindOrAdd(WatchedClass);
+			AddedActors.AddUnique(Actor);
+			WatchedClasses.AddUnique(WatchedClass);
 				
-				// If this actor was queued for a delete event, cancel it so that it's clear that the actor has been re-created.
-				TArray<FRCActorDescription>* DeletedActors = PerFrameActorsDeleted.Find(ClientId);
-				if (DeletedActors)
-				{
-					const int32 DeletedIndex = DeletedActors->IndexOfByPredicate([&ActorPath](const FRCActorDescription& Description) {
-						return Description.Path == ActorPath;
-					});
+			// If this actor was queued for a delete event, cancel it so that it's clear that the actor has been re-created.
+			FDeletedActorsData* DeletedActorsData = PerFrameActorsDeleted.Find(WatchedClass);
+			if (DeletedActorsData)
+			{
+				const int32 DeletedIndex = DeletedActorsData->Actors.IndexOfByPredicate([&ActorPath](const FRCActorDescription& Description) {
+					return Description.Path == ActorPath;
+				});
 
-					if (DeletedIndex != INDEX_NONE)
-					{
-						DeletedActors->RemoveAt(DeletedIndex);
-					}
+				if (DeletedIndex != INDEX_NONE)
+				{
+					DeletedActorsData->Actors.RemoveAt(DeletedIndex);
 				}
 			}
 		}
 	}
 
-	if (WatchingIds.Num() > 0)
+	if (WatchedClasses.Num() > 0)
 	{
 		// At least one subscriber cares about this actor, so we should listen to its events
 		FWatchedActorData& ActorData = WatchedActors.Add(Actor, FWatchedActorData(Actor));
-		ActorData.Clients = WatchingIds;
+		ActorData.WatchedClasses = WatchedClasses;
 	}
 }
 
@@ -1139,13 +1239,25 @@ void FWebSocketMessageHandler::OnActorDeleted(AActor* Actor)
 		return;
 	}
 
-	for (const FGuid& ClientId : ActorData->Clients)
+	for (TWeakObjectPtr<UClass> WatchedClass : ActorData->WatchedClasses)
 	{
-		TArray<FRCActorDescription>& DeletedActors = PerFrameActorsDeleted.FindOrAdd(ClientId);
-		DeletedActors.AddUnique(ActorData->Description);
+		UClass* WatchedClassPtr = WatchedClass.Get();
+		if (!WatchedClassPtr)
+		{
+			continue;
+		}
+
+		FDeletedActorsData* DeletedActors = PerFrameActorsDeleted.Find(WatchedClassPtr);
+		if (!DeletedActors)
+		{
+			// No actors of this class have been deleted this frame, so store the class path in case the class gets deleted too.
+			DeletedActors = &PerFrameActorsDeleted.Add(WatchedClassPtr);
+			DeletedActors->ClassPath = WatchedClassPtr->GetPathName();
+		}
+		DeletedActors->Actors.AddUnique(ActorData->Description);
 
 		// If this actor was queued for an add event, cancel it so that it's clear that the actor has been deleted again.
-		TArray<TWeakObjectPtr<AActor>>* AddedActors = PerFrameActorsAdded.Find(ClientId);
+		TArray<TWeakObjectPtr<AActor>>* AddedActors = PerFrameActorsAdded.Find(WatchedClassPtr);
 		if (AddedActors)
 		{
 			const int32 AddedIndex = AddedActors->IndexOfByPredicate([&Actor](TWeakObjectPtr<AActor> AddedActor) {
@@ -1288,7 +1400,7 @@ void FWebSocketMessageHandler::OnObjectTransacted(UObject* Object, const class F
 	}
 }
 
-void FWebSocketMessageHandler::StartWatchingActor(AActor* Actor, FGuid ClientId)
+void FWebSocketMessageHandler::StartWatchingActor(AActor* Actor, UClass* WatchedClass)
 {
 	FWatchedActorData* ActorData = WatchedActors.Find(Actor);
 
@@ -1297,19 +1409,19 @@ void FWebSocketMessageHandler::StartWatchingActor(AActor* Actor, FGuid ClientId)
 		ActorData = &WatchedActors.Add(Actor, FWatchedActorData(Actor));
 	}
 
-	ActorData->Clients.Add(ClientId);
+	ActorData->WatchedClasses.Add(WatchedClass);
 }
 
-void FWebSocketMessageHandler::StopWatchingActor(AActor* Actor, FGuid ClientId)
+void FWebSocketMessageHandler::StopWatchingActor(AActor* Actor, UClass* WatchedClass)
 {
 	FWatchedActorData* ActorData = WatchedActors.Find(Actor);
 	if (ActorData)
 	{
-		bool bAnyRemoved = ActorData->Clients.Remove(ClientId) > 0;
+		bool bAnyRemoved = ActorData->WatchedClasses.Remove(WatchedClass) > 0;
 
-		if (bAnyRemoved && ActorData->Clients.IsEmpty())
+		if (bAnyRemoved && ActorData->WatchedClasses.IsEmpty())
 		{
-			// Nobody is watching anymore, so we can forget about it
+			// Nobody is watching anymore, so we can forget about the actor
 			WatchedActors.Remove(Actor);
 		}
 	}
@@ -1321,9 +1433,9 @@ void FWebSocketMessageHandler::UpdateWatchedActorName(AActor* Actor, FWebSocketM
 	ActorData.Description.Name = Actor->GetActorNameOrLabel();
 
 	// Mark that this has been renamed
-	for (const FGuid& ClientId : ActorData.Clients)
+	for (TWeakObjectPtr<UClass> ActorClass : ActorData.WatchedClasses)
 	{
-		if (TArray<TWeakObjectPtr<AActor>>* AddedActors = PerFrameActorsAdded.Find(ClientId))
+		if (TArray<TWeakObjectPtr<AActor>>* AddedActors = PerFrameActorsAdded.Find(ActorClass))
 		{
 			// If the actor was just added this frame, we don't need to report the rename since the name will be included
 			// with the add event. This happens with copy+paste, which renames immediately after creation.
@@ -1333,7 +1445,33 @@ void FWebSocketMessageHandler::UpdateWatchedActorName(AActor* Actor, FWebSocketM
 			}
 		}
 
-		TArray<TWeakObjectPtr<AActor>>& RenamedActors = PerFrameActorsRenamed.FindOrAdd(ClientId);
+		TArray<TWeakObjectPtr<AActor>>& RenamedActors = PerFrameActorsRenamed.FindOrAdd(ActorClass);
 		RenamedActors.AddUnique(Actor);
 	}
-}
\ No newline at end of file
+}
+
+void FWebSocketMessageHandler::UnregisterClientForActorClass(const FGuid& ClientId, TSubclassOf<AActor> ActorClass)
+{
+	// Unregister if already registered
+	bool bIsClassNoLongerWatched = false;
+	if (FWatchedClassData* WatchedClassData = ActorNotificationMap.Find(ActorClass))
+	{
+		WatchedClassData->Clients.Remove(ClientId);
+		if (WatchedClassData->Clients.IsEmpty())
+		{
+			ActorNotificationMap.Remove(ActorClass);
+			bIsClassNoLongerWatched = true;
+		}
+	}
+
+	// Nobody is watching this class anymore, so stop watching actors for that class
+	if (bIsClassNoLongerWatched)
+	{
+		TArray<AActor*> WatchedActorPointers;
+		WatchedActors.GetKeys(WatchedActorPointers);
+		for (AActor* Actor : WatchedActorPointers)
+		{
+			StopWatchingActor(Actor, ActorClass);
+		}
+	}
+}
diff --git a/Engine/Plugins/VirtualProduction/RemoteControl/Source/WebRemoteControl/Private/WebSocketMessageHandler.h b/Engine/Plugins/VirtualProduction/RemoteControl/Source/WebRemoteControl/Private/WebSocketMessageHandler.h
--- a/Engine/Plugins/VirtualProduction/RemoteControl/Source/WebRemoteControl/Private/WebSocketMessageHandler.h
+++ b/Engine/Plugins/VirtualProduction/RemoteControl/Source/WebRemoteControl/Private/WebSocketMessageHandler.h
@@ -8,6 +8,7 @@
 #include "RemoteControlField.h"
 #include "RemoteControlModels.h"
 #include "RemoteControlPreset.h"
+#include "RemoteControlResponse.h"
 
 
 struct FGuid;
@@ -48,10 +49,35 @@ private:
 		/** Description of the actor. */
 		FRCActorDescription Description;
 
-		/** Which clients are listening for events about this actor. */
+		/** Which classes this actor is a member of that are causing it to be watched. */
+		TArray<TWeakObjectPtr<UClass>> WatchedClasses;
+	};
+
+	/**
+	 * Data about a class being watched by one or more clients.
+	 */
+	struct FWatchedClassData
+	{
+		/** The clients watching this class. */
 		TArray<FGuid> Clients;
+
+		/** The cached path name of the class so we can still send events about it even if it gets deleted. */
+		FString CachedPath;
+	};
+
+	/**
+	 * Data about actors of a shared class that have been deleted recently.
+	 */
+	struct FDeletedActorsData
+	{
+		/** The cached path name of the shared class so we can still send events about it even if it gets deleted. */
+		FString ClassPath;
+
+		/** Deleted actors, stored as descriptions in case the actor is garbage collected before its name and path can be collected. */
+		TArray<FRCActorDescription> Actors;
 	};
 
+
 	/** Register handlers for actors being added/deleted (must happen after engine init). */
 	void RegisterActorHandlers();
 	
@@ -107,6 +133,24 @@ private:
 	/** If actors were added/removed/renamed, notify listeners. */
 	void ProcessActorChanges();
 
+	/**
+	 * Gather all the changes in the last frame for an actor class.
+	 * 
+	 * @param ActorClass The class of actor for which to gather changes.
+	 * @param OutChangesByClassPath A map from class path to data about changes to actors of that class.
+	 * @param OutClientsToNotify A map from client ID to class paths with actor changes the client cares about.
+	 */
+	void GatherActorChangesForClass(UClass* ActorClass, TMap<FString, FRCActorsChangedData>& OutChangesByClassPath, TMap<FGuid, TArray<FString>>& OutClientsToNotify);
+
+	/**
+	 * Gather all the changes in the last frame about an actor class that has already been deleted.
+	 *
+	 * @param DeletedActorsData Data about the class and actors that have been deleted.
+	 * @param OutChangesByClassPath A map from class path to data about changes to actors of that class.
+	 * @param OutClientsToNotify A map from client ID to class paths with actor changes the client cares about.
+	 */
+	void GatherActorChangesForDeletedClass(const FDeletedActorsData* DeletedActorsData, TMap<FString, FRCActorsChangedData>& OutChangesByClassPath, TMap<FGuid, TArray<FString>>& OutClientsToNotify);
+
 	/** 
 	 * Send a payload to all clients bound to a certain preset.
 	 * @note: TargetPresetName must be in the PresetNotificationMap.
@@ -154,22 +198,37 @@ private:
 	void OnObjectTransacted(UObject* Object, const class FTransactionObjectEvent& TransactionEvent);
 
 	/**
-	 * Start watching an actor for the given client.
+	 * Start watching an actor because it's a member of the given class.
 	 */
-	void StartWatchingActor(AActor* Actor, FGuid ClientId);
+	void StartWatchingActor(AActor* Actor, UClass* WatchedClass);
 
 	/**
-	 * Stop watching an actor for the given client.
+	 * Stop watching an actor because it's a member of the given class.
+	 * This should only be called if nobody is watching for this class anymore.
 	 */
-	void StopWatchingActor(AActor* Actor, FGuid ClientId);
+	void StopWatchingActor(AActor* Actor, UClass* WatchedClass);
 
 	/**
 	 * Update our cache of watched actor's name and notify any subscribers.
 	 */
 	void UpdateWatchedActorName(AActor* Actor, FWatchedActorData& ActorData);
 
+	/**
+	 * Unregister a client from messages about the given actor class.
+	 */
+	void UnregisterClientForActorClass(const FGuid& ClientId, TSubclassOf<AActor> ActorClass);
+
 private:
 
+	/** Map type from class to Guids of clients listening for changes to actors of that class. */
+	typedef TMap<TWeakObjectPtr<UClass>, FWatchedClassData, FDefaultSetAllocator, TWeakObjectPtrMapKeyFuncs<TWeakObjectPtr<UClass>, FWatchedClassData>> FActorNotificationMap;
+
+	/** Map type from class to array of actors of that class that have changed recently. */
+	typedef TMap<TWeakObjectPtr<UClass>, TArray<TWeakObjectPtr<AActor>>, FDefaultSetAllocator, TWeakObjectPtrMapKeyFuncs<TWeakObjectPtr<UClass>, TArray<TWeakObjectPtr<AActor>>>> FChangedActorMap;
+
+	/** Map type from class to data about actors of that class that have been deleted recently. */
+	typedef TMap<TWeakObjectPtr<UClass>, FDeletedActorsData, FDefaultSetAllocator, TWeakObjectPtrMapKeyFuncs<TWeakObjectPtr<UClass>, FDeletedActorsData>> FDeletedActorMap;
+
 	/** Web Socket server. */
 	FRCWebSocketServer* Server = nullptr;
 
@@ -179,7 +238,6 @@ private:
 	TMap<FGuid, TArray<FGuid>> PresetNotificationMap;
 
 	/** All websocket client IDs associated with an actor class. */
-	typedef TMap<TWeakObjectPtr<UClass>, TArray<FGuid>, FDefaultSetAllocator, TWeakObjectPtrMapKeyFuncs<TWeakObjectPtr<UClass>, TArray<FGuid>>> FActorNotificationMap;
 	FActorNotificationMap ActorNotificationMap;
 
 	/** Configuration for a given client related to how events should be handled. */
@@ -213,17 +271,14 @@ private:
 	/** Fields that were renamed for a frame, per preset */
 	TMap<FGuid, TArray<TTuple<FName, FName>>> PerFrameRenamedFields;
 
-	/** Actors that were added for a frame, per subscribed client. */
-	TMap<FGuid, TArray<TWeakObjectPtr<AActor>>> PerFrameActorsAdded;
+	/** Actors that were added for a frame, per watched class. */
+	FChangedActorMap PerFrameActorsAdded;
 
-	/** Actors that were renamed for a frame, per subscribed client. */
-	TMap<FGuid, TArray<TWeakObjectPtr<AActor>>> PerFrameActorsRenamed;
+	/** Actors that were renamed for a frame, per watched class. */
+	FChangedActorMap PerFrameActorsRenamed;
 
-	/**
-	 * Actors that were removed for a frame, per subscribed client.
-	 * Stored by description in case the actor is garbage collected before its name and path can be collected.
-	 */
-	TMap<FGuid, TArray<FRCActorDescription>> PerFrameActorsDeleted;
+	/** Actors that were removed for a frame, per watched class. */
+	FDeletedActorMap PerFrameActorsDeleted;
 
 	/** Presets that had their metadata modified for a frame */
 	TSet<FGuid> PerFrameModifiedMetadata;
