Add: Remote Control Panel, Save Button/Function

diff --git a/Engine/Plugins/VirtualProduction/RemoteControl/Source/RemoteControlUI/Private/UI/SRemoteControlPanel.cpp b/Engine/Plugins/VirtualProduction/RemoteControl/Source/RemoteControlUI/Private/UI/SRemoteControlPanel.cpp
--- a/Engine/Plugins/VirtualProduction/RemoteControl/Source/RemoteControlUI/Private/UI/SRemoteControlPanel.cpp
+++ b/Engine/Plugins/VirtualProduction/RemoteControl/Source/RemoteControlUI/Private/UI/SRemoteControlPanel.cpp
@@ -11,8 +11,10 @@
 #include "EditorFontGlyphs.h"
 #include "EngineUtils.h"
 #include "Engine/Selection.h"
+#include "FileHelpers.h"
 #include "Layout/Visibility.h"
 #include "Input/Reply.h"
+#include "IRemoteControlModule.h"
 #include "IRemoteControlProtocolWidgetsModule.h"
 #include "ISettingsModule.h"
 #include "IStructureDetailsView.h"
@@ -142,7 +144,33 @@ void SRemoteControlPanel::Construct(const FArguments& InArgs, URemoteControlPres
 			SNew(SHorizontalBox)
 
 			+ SHorizontalBox::Slot()
-			.Padding(FMargin(5.0f, 0.0f))
+			.Padding(FMargin(0, 0, 2.0f, 0.0f))
+			.VAlign(VAlign_Center)
+			.AutoWidth()
+			[
+				SNew(SButton)
+				.Visibility_Lambda([this]() { return bIsInEditMode ? EVisibility::Visible : EVisibility::Collapsed; })
+				.ButtonStyle(FAppStyle::Get(), "FlatButton")
+				.OnClicked(this, &SRemoteControlPanel::OnSavePreset)
+				[
+					SNew(SHorizontalBox)
+					+ SHorizontalBox::Slot()
+					.AutoWidth()
+					.Padding(0,0,5, 0)
+					[
+						SNew(SImage)
+						.Image(FAppStyle::Get().GetBrush("Icons.Save"))
+					]
+					+ SHorizontalBox::Slot()
+					.FillWidth(1.0f)
+					[
+						SNew(STextBlock)
+						.Text(LOCTEXT("Save", "Save"))
+					]
+				]
+			]
+			+ SHorizontalBox::Slot()
+			.Padding(FMargin(0, 0, 5.0f, 0.0f))
 			.VAlign(VAlign_Center)
 			.AutoWidth()
 			[
@@ -950,6 +978,25 @@ FReply SRemoteControlPanel::OnCreateGroup()
 	return FReply::Handled();
 }
 
+FReply SRemoteControlPanel::OnSavePreset()
+{
+	if (!Preset.IsValid() || !Preset->IsAsset())
+	{
+		UE_LOG(LogRemoteControl, Log, TEXT("Invalid object to save: %s"), (Preset.IsValid()) ? *Preset->GetFullName() : TEXT("Null Object"));
+		
+		return FReply::Handled();
+	}
+
+	TArray<UPackage*> PackagesToSave;
+	constexpr bool bCheckDirtyOnAssetSave = false;
+	constexpr bool bPromptToSave = false;
+	
+	PackagesToSave.Add(Preset->GetOutermost());
+	FEditorFileUtils::PromptForCheckoutAndSave(PackagesToSave, bCheckDirtyOnAssetSave, bPromptToSave);
+	
+	return FReply::Handled();
+}
+
 void SRemoteControlPanel::ExposeProperty(UObject* Object, FRCFieldPathInfo Path)
 {
 	if (Path.Resolve(Object))
diff --git a/Engine/Plugins/VirtualProduction/RemoteControl/Source/RemoteControlUI/Private/UI/SRemoteControlPanel.h b/Engine/Plugins/VirtualProduction/RemoteControl/Source/RemoteControlUI/Private/UI/SRemoteControlPanel.h
--- a/Engine/Plugins/VirtualProduction/RemoteControl/Source/RemoteControlUI/Private/UI/SRemoteControlPanel.h
+++ b/Engine/Plugins/VirtualProduction/RemoteControl/Source/RemoteControlUI/Private/UI/SRemoteControlPanel.h
@@ -131,6 +131,9 @@ private:
 	/** Handles creating a new group. */
 	FReply OnCreateGroup();
 
+	/** Handles saving the Preset. */
+	FReply OnSavePreset();
+
 	/** Expose a property using its handle. */
 	void ExposeProperty(UObject* Object, FRCFieldPathInfo FieldPath);
 
