[WebRemoteControl] Add transient presets, routes to expose/unexpose preset properties, and delta API for presets

diff --git a/Engine/Plugins/VirtualProduction/RemoteControl/Source/RemoteControl/Private/RemoteControlModule.cpp b/Engine/Plugins/VirtualProduction/RemoteControl/Source/RemoteControl/Private/RemoteControlModule.cpp
--- a/Engine/Plugins/VirtualProduction/RemoteControl/Source/RemoteControl/Private/RemoteControlModule.cpp
+++ b/Engine/Plugins/VirtualProduction/RemoteControl/Source/RemoteControl/Private/RemoteControlModule.cpp
@@ -1204,6 +1204,59 @@ URemoteControlPreset* FRemoteControlModule::ResolvePreset(const FGuid& PresetId)
 	return nullptr;
 }
 
+URemoteControlPreset* FRemoteControlModule::CreateTransientPreset()
+{
+	const FName AssetName(*FString::Printf(TEXT("TransientRCPreset%d"), NextTransientPresetIndex));
+	++NextTransientPresetIndex;
+		
+	const FString AssetPath = FString::Printf(TEXT("/Temp/%s"), *AssetName.ToString());
+	if (UPackage* Package = CreatePackage(*AssetPath))
+	{
+		if (URemoteControlPreset* Preset = NewObject<URemoteControlPreset>(Package, AssetName, RF_Transient | RF_Public | RF_Standalone))
+		{
+			FAssetRegistryModule::AssetCreated(Preset);
+			const FAssetData AssetData(Preset);
+			const FGuid PresetId = RemoteControlUtil::GetPresetId(AssetData);
+
+			TransientPresets.Add(AssetData);
+			OnAssetAdded(AssetData);
+
+			return Preset;
+		}
+	}
+
+	return nullptr;
+}
+
+bool FRemoteControlModule::DestroyTransientPreset(FName PresetName)
+{
+	if (URemoteControlPreset* Preset = ResolvePreset(PresetName))
+	{
+		for (const FAssetData& TransientAssetData : TransientPresets)
+		{
+			if (TransientAssetData.GetAsset() == Preset)
+			{
+				// This call will also remove the asset from TransientPresets
+				OnAssetRemoved(TransientAssetData);
+				
+				FAssetRegistryModule::AssetDeleted(Preset);
+
+				if (UPackage* Package = Preset->GetPackage())
+				{
+					Package->MarkAsGarbage();
+				}
+
+				Preset->ClearFlags(RF_Public | RF_Standalone);
+				Preset->MarkAsGarbage();
+
+				return true;
+			}
+		}
+	}
+
+	return false;
+}
+
 void FRemoteControlModule::GetPresets(TArray<TSoftObjectPtr<URemoteControlPreset>>& OutPresets) const
 {
 	OutPresets.Reserve(CachedPresetsByName.Num());
@@ -1213,7 +1266,7 @@ void FRemoteControlModule::GetPresets(TArray<TSoftObjectPtr<URemoteControlPreset
 	}
 }
 
-void FRemoteControlModule::GetPresetAssets(TArray<FAssetData>& OutPresetAssets) const
+void FRemoteControlModule::GetPresetAssets(TArray<FAssetData>& OutPresetAssets, bool bIncludeTransient) const
 {
 	if (CachedPresetsByName.Num() == 0)
 	{
@@ -1221,9 +1274,26 @@ void FRemoteControlModule::GetPresetAssets(TArray<FAssetData>& OutPresetAssets)
 	}
 
 	OutPresetAssets.Reserve(CachedPresetsByName.Num());
-	for (const TPair<FName, TArray<FAssetData>>& Entry : CachedPresetsByName)
+	
+	if (bIncludeTransient)
+	{
+		for (const TPair<FName, TArray<FAssetData>>& Entry : CachedPresetsByName)
+		{
+			OutPresetAssets.Append(Entry.Value);
+		}
+	}
+	else
 	{
-		OutPresetAssets.Append(Entry.Value);
+		for (const TPair<FName, TArray<FAssetData>>& Entry : CachedPresetsByName)
+		{
+			for (const FAssetData& AssetData : Entry.Value)
+			{
+				if (!TransientPresets.Contains(AssetData))
+				{
+					OutPresetAssets.Add(AssetData);
+				}
+			}
+		}
 	}
 }
 
@@ -1321,6 +1391,8 @@ void FRemoteControlModule::OnAssetRemoved(const FAssetData& AssetData)
 		return;
 	}
 
+	TransientPresets.Remove(AssetData);
+
 	const FGuid PresetId = RemoteControlUtil::GetPresetId(AssetData);
 	if (FName* PresetName = CachedPresetNamesById.Find(PresetId))
 	{
diff --git a/Engine/Plugins/VirtualProduction/RemoteControl/Source/RemoteControl/Private/RemoteControlModule.h b/Engine/Plugins/VirtualProduction/RemoteControl/Source/RemoteControl/Private/RemoteControlModule.h
--- a/Engine/Plugins/VirtualProduction/RemoteControl/Source/RemoteControl/Private/RemoteControlModule.h
+++ b/Engine/Plugins/VirtualProduction/RemoteControl/Source/RemoteControl/Private/RemoteControlModule.h
@@ -2,11 +2,11 @@
 #pragma once
 
 #include "IRemoteControlModule.h"
+#include "AssetRegistry/AssetData.h"
 #include "CoreMinimal.h"
 #include "Engine/EngineTypes.h"
 
 class IRemoteControlInterceptionFeatureProcessor;
-struct FAssetData;
 
 /**
  * Implementation of the RemoteControl interface
@@ -37,8 +37,10 @@ public:
 	virtual TOptional<FExposedProperty> ResolvePresetProperty(const FResolvePresetFieldArgs& Args) const override;
 	virtual URemoteControlPreset* ResolvePreset(FName PresetName) const override;
 	virtual URemoteControlPreset* ResolvePreset(const FGuid& PresetId) const override;
+	virtual URemoteControlPreset* CreateTransientPreset() override;
+	virtual bool DestroyTransientPreset(FName PresetName) override;
 	virtual void GetPresets(TArray<TSoftObjectPtr<URemoteControlPreset>>& OutPresets) const override;
-	virtual void GetPresetAssets(TArray<FAssetData>& OutPresetAssets) const override;
+	virtual void GetPresetAssets(TArray<FAssetData>& OutPresetAssets, bool bIncludeTransient = true) const override;
 	virtual const TMap<FName, FEntityMetadataInitializer>& GetDefaultMetadataInitializers() const override;
 	virtual bool RegisterDefaultEntityMetadata(FName MetadataKey, FEntityMetadataInitializer MetadataInitializer) override;
 	virtual void UnregisterDefaultEntityMetadata(FName MetadataKey) override;
@@ -94,6 +96,12 @@ private:
 	/** Cache of ids to preset names. */
 	mutable TMap<FGuid, FName> CachedPresetNamesById;
 
+	/** Temporary presets that aren't saved as assets or directly visible to the editor's user. */
+	TSet<FAssetData> TransientPresets;
+
+	/** Index of the next created transient preset, to avoid naming collisions. */
+	uint32 NextTransientPresetIndex = 0;
+
 	/** Map of registered default metadata initializers. */
 	TMap<FName, FEntityMetadataInitializer> DefaultMetadataInitializers;
 
diff --git a/Engine/Plugins/VirtualProduction/RemoteControl/Source/RemoteControl/Public/IRemoteControlModule.h b/Engine/Plugins/VirtualProduction/RemoteControl/Source/RemoteControl/Public/IRemoteControlModule.h
--- a/Engine/Plugins/VirtualProduction/RemoteControl/Source/RemoteControl/Public/IRemoteControlModule.h
+++ b/Engine/Plugins/VirtualProduction/RemoteControl/Source/RemoteControl/Public/IRemoteControlModule.h
@@ -324,6 +324,20 @@ public:
      */
 	virtual URemoteControlPreset* ResolvePreset(const FGuid& PresetId) const = 0;
 
+	/**
+	 * Create a transient preset.
+	 * Make sure to call DestroyTransientPreset when done with the preset or it will stay in memory.
+	 * @return the new preset, or nullptr if it couldn't be created.
+	 */
+	virtual URemoteControlPreset* CreateTransientPreset() = 0;
+
+	/**
+	 * Destroy a transient preset.
+	 * @arg PresetName name of the preset to destroy.
+	 * @return true if a transient preset with that name existed and was destroyed.
+	 */
+	virtual bool DestroyTransientPreset(FName PresetName) = 0;
+
 	/**
 	 * Get all the presets currently registered with the module.
 	 */
@@ -331,8 +345,9 @@ public:
 
 	/**
 	 * Get all the preset asset currently registered with the module.
+	 * @arg bIncludeTransient Whether to include transient presets.
 	 */
-	virtual void GetPresetAssets(TArray<FAssetData>& OutPresetAssets) const = 0;
+	virtual void GetPresetAssets(TArray<FAssetData>& OutPresetAssets, bool bIncludeTransient = true) const = 0;
 	
 	/**
 	 * Get the map of registered default entity metadata initializers. 
diff --git a/Engine/Plugins/VirtualProduction/RemoteControl/Source/WebRemoteControl/Private/RemoteControlRequest.h b/Engine/Plugins/VirtualProduction/RemoteControl/Source/WebRemoteControl/Private/RemoteControlRequest.h
--- a/Engine/Plugins/VirtualProduction/RemoteControl/Source/WebRemoteControl/Private/RemoteControlRequest.h
+++ b/Engine/Plugins/VirtualProduction/RemoteControl/Source/WebRemoteControl/Private/RemoteControlRequest.h
@@ -270,6 +270,13 @@ struct FRCPresetSetPropertyRequest : public FRCRequest
 	 */
 	static FString PropertyValueLabel() { return TEXT("PropertyValue"); }
 
+	/**
+	 * Which type of operation should be performed on the value of the property.
+	 * This will be ignored if ResetToDefault is true.
+	 */
+	UPROPERTY()
+	ERCModifyOperation Operation = ERCModifyOperation::EQUAL;
+
 	/**
 	 * Whether a transaction should be created for the call.
 	 */
@@ -305,6 +312,39 @@ struct FRCPresetCallRequest : public FRCRequest
 	bool GenerateTransaction = false;
 };
 
+/**
+ * Holds a request to expose a property on a preset
+ */
+USTRUCT()
+struct FRCPresetExposePropertyRequest : public FRCRequest
+{
+	GENERATED_BODY()
+
+	/**
+	 * The path of the target object.
+	 */
+	UPROPERTY()
+	FString ObjectPath;
+
+	/**
+	 * The property to expose.
+	 */
+	UPROPERTY()
+	FString PropertyName;
+
+	/**
+	 * The label to give the new exposed property (optional).
+	 */
+	UPROPERTY()
+	FString Label;
+
+	/**
+	 * The name of the group in which to place the new exposed property (optional).
+	 */
+	UPROPERTY()
+	FString GroupName;
+};
+
 /**
  * Holds a request to describe an object using its path.
  */
diff --git a/Engine/Plugins/VirtualProduction/RemoteControl/Source/WebRemoteControl/Private/WebRemoteControl.cpp b/Engine/Plugins/VirtualProduction/RemoteControl/Source/WebRemoteControl/Private/WebRemoteControl.cpp
--- a/Engine/Plugins/VirtualProduction/RemoteControl/Source/WebRemoteControl/Private/WebRemoteControl.cpp
+++ b/Engine/Plugins/VirtualProduction/RemoteControl/Source/WebRemoteControl/Private/WebRemoteControl.cpp
@@ -622,6 +622,20 @@ void FWebRemoteControlModule::RegisterRoutes()
 		FRequestHandlerDelegate::CreateRaw(this, &FWebRemoteControlModule::HandlePresetGetPropertyRoute)
 		});
 
+	RegisterRoute({
+		TEXT("Expose a property on a preset."),
+		FHttpPath(TEXT("/remote/preset/:preset/expose/property")),
+		EHttpServerRequestVerbs::VERB_PUT,
+		FRequestHandlerDelegate::CreateRaw(this, &FWebRemoteControlModule::HandlePresetExposePropertyRoute)
+		});
+
+	RegisterRoute({
+		TEXT("Unexpose a property on a preset."),
+		FHttpPath(TEXT("/remote/preset/:preset/unexpose/property/:property")),
+		EHttpServerRequestVerbs::VERB_PUT,
+		FRequestHandlerDelegate::CreateRaw(this, &FWebRemoteControlModule::HandlePresetUnexposePropertyRoute)
+		});
+
 	RegisterRoute({
 		TEXT("Get an exposed actor's properties."),
 		FHttpPath(TEXT("/remote/preset/:preset/actor/:actor")),
@@ -715,6 +729,20 @@ void FWebRemoteControlModule::RegisterRoutes()
 		FRequestHandlerDelegate::CreateRaw(this, &FWebRemoteControlModule::HandleEntitySetLabelRoute)
 	});
 
+	RegisterRoute({
+		TEXT("Create a temporary preset which won't be visible in the editor or saved as an asset"),
+		FHttpPath(TEXT("/remote/preset/transient")),
+		EHttpServerRequestVerbs::VERB_PUT,
+		FRequestHandlerDelegate::CreateRaw(this, &FWebRemoteControlModule::HandleCreateTransientPresetRoute)
+	});
+
+	RegisterRoute({
+		TEXT("Delete a transient preset"),
+		FHttpPath(TEXT("/remote/preset/transient/:preset")),
+		EHttpServerRequestVerbs::VERB_DELETE,
+		FRequestHandlerDelegate::CreateRaw(this, &FWebRemoteControlModule::HandleDeleteTransientPresetRoute)
+	});
+
 	//**************************************
 	// Special websocket route just using http request
 	RegisterWebsocketRoute({
@@ -1170,7 +1198,7 @@ bool FWebRemoteControlModule::HandlePresetSetPropertyRoute(const FHttpServerRequ
 		{
 			NewPayloadReader.Seek(0);
 			// Set a ERCPayloadType and TCHARBody in order to follow the replication path
-			bSuccess &= IRemoteControlModule::Get().SetObjectProperties(ObjectRef, Backend, ERCPayloadType::Json, NewPayload);
+			bSuccess &= IRemoteControlModule::Get().SetObjectProperties(ObjectRef, Backend, ERCPayloadType::Json, NewPayload, SetPropertyRequest.Operation);
 		}
 	}
 
@@ -1247,6 +1275,115 @@ bool FWebRemoteControlModule::HandlePresetGetPropertyRoute(const FHttpServerRequ
 	return true;
 }
 
+bool FWebRemoteControlModule::HandlePresetExposePropertyRoute(const FHttpServerRequest& Request, const FHttpResultCallback& OnComplete)
+{
+	TUniquePtr<FHttpServerResponse> Response = WebRemoteControlInternalUtils::CreateHttpResponse();
+
+	const FString PresetName = Request.PathParams.FindChecked(TEXT("preset"));
+
+	URemoteControlPreset* Preset = WebRemoteControl::GetPreset(*PresetName);
+	if (Preset == nullptr)
+	{
+		Response->Code = EHttpServerResponseCodes::NotFound;
+		WebRemoteControlInternalUtils::CreateUTF8ErrorMessage(*FString::Printf(TEXT("Unable to resolve the preset '%s'."), *PresetName), Response->Body);
+		OnComplete(MoveTemp(Response));
+		return true;
+	}
+
+	FRCPresetExposePropertyRequest ExposePropertyRequest;
+	if (!WebRemoteControlInternalUtils::DeserializeRequest(Request, &OnComplete, ExposePropertyRequest))
+	{
+		return true;
+	}
+
+	FRCObjectReference ObjectRef;
+	const bool bFoundProperty = IRemoteControlModule::Get().ResolveObject(ERCAccess::WRITE_ACCESS, ExposePropertyRequest.ObjectPath, ExposePropertyRequest.PropertyName, ObjectRef);
+
+	if (!bFoundProperty)
+	{
+		Response->Code = EHttpServerResponseCodes::NotFound;
+		WebRemoteControlInternalUtils::CreateUTF8ErrorMessage(*FString::Printf(TEXT("Unable to resolve the property '%s' on object '%s'."),
+			*ExposePropertyRequest.PropertyName, *ExposePropertyRequest.ObjectPath), Response->Body);
+		OnComplete(MoveTemp(Response));
+		return true;
+	}
+
+	FRemoteControlPresetExposeArgs ExposeArgs;
+	ExposeArgs.Label = ExposePropertyRequest.Label;
+
+	if (ExposePropertyRequest.GroupName.Len() > 0)
+	{
+		const FRemoteControlPresetGroup* Group = Preset->Layout.GetGroupByName(FName(ExposePropertyRequest.GroupName));
+		if (Group == nullptr)
+		{
+			Response->Code = EHttpServerResponseCodes::NotFound;
+			WebRemoteControlInternalUtils::CreateUTF8ErrorMessage(*FString::Printf(TEXT("Unable to resolve the property group '%s' on preset '%s'."),
+				*ExposePropertyRequest.GroupName, *PresetName), Response->Body);
+			OnComplete(MoveTemp(Response));
+			return true;
+		}
+
+		ExposeArgs.GroupId = Group->Id;
+	}
+
+	TSharedPtr<FRemoteControlProperty> ExposedProperty = Preset->ExposeProperty(ObjectRef.Object.Get(), ObjectRef.PropertyPathInfo, ExposeArgs).Pin();
+
+	if (ExposedProperty.IsValid())
+	{
+		TArray<uint8> WorkingBuffer;
+		FMemoryWriter Writer(WorkingBuffer);
+
+		FStructOnScope PropertyValueOnScope = WebRemoteControlStructUtils::CreatePropertyValueOnScope(ExposedProperty, ObjectRef);
+		FStructOnScope FinalStruct = WebRemoteControlStructUtils::CreateGetPropertyOnScope(ExposedProperty, ObjectRef, MoveTemp(PropertyValueOnScope));
+		WebRemoteControlInternalUtils::SerializeStructOnScope(FinalStruct, Writer);
+
+		Response->Code = EHttpServerResponseCodes::Ok;
+		WebRemoteControlUtils::ConvertToUTF8(WorkingBuffer, Response->Body);
+	}
+	else
+	{
+		Response->Code = EHttpServerResponseCodes::ServerError;
+		WebRemoteControlInternalUtils::CreateUTF8ErrorMessage(TEXT("Failed to expose the property."), Response->Body);
+	}
+
+	OnComplete(MoveTemp(Response));
+	return true;
+}
+
+
+bool FWebRemoteControlModule::HandlePresetUnexposePropertyRoute(const FHttpServerRequest& Request, const FHttpResultCallback& OnComplete)
+{
+	TUniquePtr<FHttpServerResponse> Response = WebRemoteControlInternalUtils::CreateHttpResponse();
+
+	const FString PresetName = Request.PathParams.FindChecked(TEXT("preset"));
+
+	URemoteControlPreset* Preset = WebRemoteControl::GetPreset(*PresetName);
+	if (Preset == nullptr)
+	{
+		Response->Code = EHttpServerResponseCodes::NotFound;
+		WebRemoteControlInternalUtils::CreateUTF8ErrorMessage(*FString::Printf(TEXT("Unable to resolve the preset '%s'."), *PresetName), Response->Body);
+		OnComplete(MoveTemp(Response));
+		return true;
+	}
+
+	const FString PropertyLabelString = Request.PathParams.FindChecked(TEXT("property"));
+	const FName PropertyLabel = FName(PropertyLabelString);
+	if (!Preset->GetExposedEntityId(PropertyLabel).IsValid())
+	{
+		Response->Code = EHttpServerResponseCodes::NotFound;
+		WebRemoteControlInternalUtils::CreateUTF8ErrorMessage(*FString::Printf(TEXT("Exposeed property '%s' not found on preset '%s'."),
+			*PropertyLabelString, *PresetName), Response->Body);
+		OnComplete(MoveTemp(Response));
+		return true;
+	}
+
+	Preset->Unexpose(PropertyLabel);
+
+	Response->Code = EHttpServerResponseCodes::Ok;
+	OnComplete(MoveTemp(Response));
+	return true;
+}
+
 bool FWebRemoteControlModule::HandlePresetGetExposedActorPropertyRoute(const FHttpServerRequest& Request, const FHttpResultCallback& OnComplete)
 {
 	TUniquePtr<FHttpServerResponse> Response = WebRemoteControlInternalUtils::CreateHttpResponse();
@@ -1404,7 +1541,7 @@ bool FWebRemoteControlModule::HandlePresetSetExposedActorPropertyRoute(const FHt
 				{
 					NewPayloadReader.Seek(0);
 					// Set a ERCPayloadType and NewPayload in order to follow the replication path
-					bSuccess &= IRemoteControlModule::Get().SetObjectProperties(ObjectRef, Backend, ERCPayloadType::Json, NewPayload);
+					bSuccess &= IRemoteControlModule::Get().SetObjectProperties(ObjectRef, Backend, ERCPayloadType::Json, NewPayload, SetPropertyRequest.Operation);
 				}
 			}
 		}
@@ -1459,7 +1596,7 @@ bool FWebRemoteControlModule::HandleGetPresetsRoute(const FHttpServerRequest& Re
 	TUniquePtr<FHttpServerResponse> Response = WebRemoteControlInternalUtils::CreateHttpResponse(EHttpServerResponseCodes::Ok);
 
 	TArray<FAssetData> PresetAssets;
-	IRemoteControlModule::Get().GetPresetAssets(PresetAssets);
+	IRemoteControlModule::Get().GetPresetAssets(PresetAssets, false);
 
 	WebRemoteControlUtils::SerializeMessage(FListPresetsResponse{ PresetAssets }, Response->Body);
 
@@ -1815,6 +1952,47 @@ bool FWebRemoteControlModule::HandleEntitySetLabelRoute(const FHttpServerRequest
 	return true;
 }
 
+bool FWebRemoteControlModule::HandleCreateTransientPresetRoute(const FHttpServerRequest& Request, const FHttpResultCallback& OnComplete)
+{
+	TUniquePtr<FHttpServerResponse> Response = WebRemoteControlInternalUtils::CreateHttpResponse();
+
+	const URemoteControlPreset* Preset = IRemoteControlModule::Get().CreateTransientPreset();
+	if (!Preset)
+	{
+		Response->Code = EHttpServerResponseCodes::ServerError;
+		WebRemoteControlInternalUtils::CreateUTF8ErrorMessage(TEXT("Unable to create a transient preset."), Response->Body);
+		OnComplete(MoveTemp(Response));
+		return true;
+	}
+
+	WebRemoteControlUtils::SerializeMessage(FGetPresetResponse{ Preset }, Response->Body);
+	Response->Code = EHttpServerResponseCodes::Ok;
+
+	OnComplete(MoveTemp(Response));
+	return true;
+}
+
+bool FWebRemoteControlModule::HandleDeleteTransientPresetRoute(const FHttpServerRequest& Request, const FHttpResultCallback& OnComplete)
+{
+	TUniquePtr<FHttpServerResponse> Response = WebRemoteControlInternalUtils::CreateHttpResponse();
+
+	FString PresetName = Request.PathParams.FindChecked(TEXT("preset"));
+
+	const bool bSuccess = IRemoteControlModule::Get().DestroyTransientPreset(FName(*PresetName));
+	if (!bSuccess)
+	{
+		Response->Code = EHttpServerResponseCodes::NotFound;
+		WebRemoteControlInternalUtils::CreateUTF8ErrorMessage(TEXT("No transient preset with that name exists."), Response->Body);
+		OnComplete(MoveTemp(Response));
+		return true;
+	}
+
+	Response->Code = EHttpServerResponseCodes::Ok;
+
+	OnComplete(MoveTemp(Response));
+	return true;
+}
+
 void FWebRemoteControlModule::HandleWebSocketHttpMessage(const FRemoteControlWebSocketMessage& WebSocketMessage)
 {
 	// Sets the acting client id for the duration of the message handling.
diff --git a/Engine/Plugins/VirtualProduction/RemoteControl/Source/WebRemoteControl/Private/WebRemoteControl.h b/Engine/Plugins/VirtualProduction/RemoteControl/Source/WebRemoteControl/Private/WebRemoteControl.h
--- a/Engine/Plugins/VirtualProduction/RemoteControl/Source/WebRemoteControl/Private/WebRemoteControl.h
+++ b/Engine/Plugins/VirtualProduction/RemoteControl/Source/WebRemoteControl/Private/WebRemoteControl.h
@@ -113,6 +113,8 @@ private:
 	bool HandlePresetCallFunctionRoute(const FHttpServerRequest& Request, const FHttpResultCallback& OnComplete);
 	bool HandlePresetSetPropertyRoute(const FHttpServerRequest& Request, const FHttpResultCallback& OnComplete);
 	bool HandlePresetGetPropertyRoute(const FHttpServerRequest& Request, const FHttpResultCallback& OnComplete);
+	bool HandlePresetExposePropertyRoute(const FHttpServerRequest& Request, const FHttpResultCallback& OnComplete);
+	bool HandlePresetUnexposePropertyRoute(const FHttpServerRequest& Request, const FHttpResultCallback& OnComplete);
 	bool HandlePresetGetExposedActorPropertyRoute(const FHttpServerRequest& Request, const FHttpResultCallback& OnComplete);
 	bool HandlePresetGetExposedActorPropertiesRoute(const FHttpServerRequest& Request, const FHttpResultCallback& OnComplete);
 	bool HandlePresetSetExposedActorPropertyRoute(const FHttpServerRequest& Request, const FHttpResultCallback& OnComplete);
@@ -127,6 +129,8 @@ private:
 	bool HandleEntityMetadataOperationsRoute(const FHttpServerRequest& Request, const FHttpResultCallback& OnComplete);
 	bool HandleEntitySetLabelRoute(const FHttpServerRequest& Request, const FHttpResultCallback& OnComplete);
 	bool HandlePassphraseRoute(const FHttpServerRequest& Request, const FHttpResultCallback& OnComplete);
+	bool HandleCreateTransientPresetRoute(const FHttpServerRequest& Request, const FHttpResultCallback& OnComplete);
+	bool HandleDeleteTransientPresetRoute(const FHttpServerRequest& Request, const FHttpResultCallback& OnComplete);
 
 	//~ Websocket route handlers
 	void HandleWebSocketHttpMessage(const struct FRemoteControlWebSocketMessage& WebSocketMessage);
