[WebRemoteControl] Wait until WebSocket server is initialized to bind routes

diff --git a/Engine/Plugins/VirtualProduction/RemoteControl/Source/WebRemoteControl/Private/WebRemoteControl.cpp b/Engine/Plugins/VirtualProduction/RemoteControl/Source/WebRemoteControl/Private/WebRemoteControl.cpp
--- a/Engine/Plugins/VirtualProduction/RemoteControl/Source/WebRemoteControl/Private/WebRemoteControl.cpp
+++ b/Engine/Plugins/VirtualProduction/RemoteControl/Source/WebRemoteControl/Private/WebRemoteControl.cpp
@@ -390,12 +390,22 @@ void FWebRemoteControlModule::UnregisterRoute(const FRemoteControlRoute& Route)
 
 void FWebRemoteControlModule::RegisterWebsocketRoute(const FRemoteControlWebsocketRoute& Route)
 {
-	WebSocketRouter->BindRoute(Route.MessageName, Route.Delegate);
+	RegisteredWebSocketRoutes.Add(Route);
+
+	if (WebSocketRouter)
+	{
+		WebSocketRouter->BindRoute(Route.MessageName, Route.Delegate);
+	}
 }
 
 void FWebRemoteControlModule::UnregisterWebsocketRoute(const FRemoteControlWebsocketRoute& Route)
 {
-	WebSocketRouter->UnbindRoute(Route.MessageName);
+	RegisteredWebSocketRoutes.Remove(Route);
+
+	if (WebSocketRouter)
+	{
+		WebSocketRouter->UnbindRoute(Route.MessageName);
+	}
 }
 
 void FWebRemoteControlModule::SendWebsocketMessage(const FGuid& InTargetClientId, const TArray<uint8>& InUTF8Payload)
@@ -507,6 +517,11 @@ void FWebRemoteControlModule::StartWebSocketServer()
 
 		UE_LOG(LogRemoteControl, Log, TEXT("Web Remote Control WebSocket server started on port %d"), WebSocketServerPort);
 		OnWebSocketServerStartedDelegate.Broadcast(WebSocketServerPort);
+
+		for (FRemoteControlWebsocketRoute& Route : RegisteredWebSocketRoutes)
+		{
+			WebSocketRouter->BindRoute(Route.MessageName, Route.Delegate);
+		}
 	}
 }
 
diff --git a/Engine/Plugins/VirtualProduction/RemoteControl/Source/WebRemoteControl/Private/WebRemoteControl.h b/Engine/Plugins/VirtualProduction/RemoteControl/Source/WebRemoteControl/Private/WebRemoteControl.h
--- a/Engine/Plugins/VirtualProduction/RemoteControl/Source/WebRemoteControl/Private/WebRemoteControl.h
+++ b/Engine/Plugins/VirtualProduction/RemoteControl/Source/WebRemoteControl/Private/WebRemoteControl.h
@@ -160,7 +160,7 @@ private:
 	TSet<FRemoteControlRoute> RegisteredHttpRoutes;
 
 	/** Set of routes that will be activated on websocket server start. */
-	TSet<FRemoteControlRoute> RegisteredWebSocketRoutes;
+	TSet<FRemoteControlWebsocketRoute> RegisteredWebSocketRoutes;
 
 	/** Port of the remote control http server. */
 	uint32 HttpServerPort;
