[WebRemoteControl] WebSocket routes to set preset properties and batch messages.
Added sequence numbers to preset property update messages. Clients can send a sequence number when changing a property via WebSocket, and the property updates will include the latest received number for the corresponding client.

diff --git a/Engine/Plugins/VirtualProduction/RemoteControl/Source/WebRemoteControl/Private/RemoteControlModels.h b/Engine/Plugins/VirtualProduction/RemoteControl/Source/WebRemoteControl/Private/RemoteControlModels.h
--- a/Engine/Plugins/VirtualProduction/RemoteControl/Source/WebRemoteControl/Private/RemoteControlModels.h
+++ b/Engine/Plugins/VirtualProduction/RemoteControl/Source/WebRemoteControl/Private/RemoteControlModels.h
@@ -116,6 +116,12 @@ struct FRCPropertyDescription
 
 		Name = Property->GetName();
 
+#if WITH_EDITORONLY_DATA
+		DisplayName = Property->GetDisplayNameText();
+#else
+		DisplayName = FText::FromString(Name);
+#endif
+
 		const FProperty* ValueProperty = Property;
 		if (const FArrayProperty* ArrayProperty = CastField<FArrayProperty>(Property))
 		{
@@ -188,6 +194,10 @@ struct FRCPropertyDescription
 	/** Name of the exposed property */
 	UPROPERTY()
 	FString Name;
+
+	/** Display name of the exposed property */
+	UPROPERTY()
+	FText DisplayName;
 	
 	/** Description of the exposed property */
 	UPROPERTY()
diff --git a/Engine/Plugins/VirtualProduction/RemoteControl/Source/WebRemoteControl/Private/RemoteControlRequest.h b/Engine/Plugins/VirtualProduction/RemoteControl/Source/WebRemoteControl/Private/RemoteControlRequest.h
--- a/Engine/Plugins/VirtualProduction/RemoteControl/Source/WebRemoteControl/Private/RemoteControlRequest.h
+++ b/Engine/Plugins/VirtualProduction/RemoteControl/Source/WebRemoteControl/Private/RemoteControlRequest.h
@@ -564,6 +564,21 @@ struct FRCWebSocketRequest : public FRCRequest
 	int32 Id = INDEX_NONE;
 };
 
+/**
+ * Holds a request that wraps multiple requests..
+ */
+USTRUCT()
+struct FRCWebSocketBatchRequest : public FRCRequest
+{
+	GENERATED_BODY()
+
+	/**
+	 * The list of batched requests.
+	 */
+	UPROPERTY()
+	TArray<FRCWebSocketRequest> Requests;
+};
+
 /**
  * Holds a request made via websocket to register for events about a given preset.
  */
@@ -618,3 +633,60 @@ struct FRCWebSocketActorRegisterBody : public FRCRequest
 	FName ClassName;
 };
 
+/**
+ * Holds a request made via websocket to modify a property exposed in a preset.
+ */
+USTRUCT()
+struct FRCWebSocketPresetSetPropertyBody : public FRCRequest
+{
+	GENERATED_BODY()
+
+	FRCWebSocketPresetSetPropertyBody()
+	{
+		AddStructParameter(PropertyValueLabel());
+	}
+
+	/**
+	 * Get the label for the PropertyValue struct.
+	 */
+	static FString PropertyValueLabel() { return TEXT("PropertyValue"); }
+
+	/**
+	 * The name of the remote control preset to which the property belongs.
+	 */
+	UPROPERTY()
+	FName PresetName;
+
+	/**
+	 * The label of the property to modify.
+	 */
+	UPROPERTY()
+	FName PropertyLabel;
+
+	/**
+	 * Which type of operation should be performed on the value of the property.
+	 * This will be ignored if ResetToDefault is true.
+	 */
+	UPROPERTY()
+	ERCModifyOperation Operation = ERCModifyOperation::EQUAL;
+
+	/**
+	 * Whether a transaction should be created for the call.
+	 */
+	UPROPERTY()
+	bool GenerateTransaction = false;
+
+	/**
+	 * If true, ignore the other parameters and just reset the property to its default value.
+	 */
+	UPROPERTY()
+	bool ResetToDefault = false;
+
+	/**
+	 * The sequence number of this change. The highest sequence number received from this client will be
+	 * sent back to the client in future PresetFieldsChanged events.
+	 */
+	UPROPERTY()
+	int64 SequenceNumber = -1;
+};
+
diff --git a/Engine/Plugins/VirtualProduction/RemoteControl/Source/WebRemoteControl/Private/RemoteControlWebSocketServer.cpp b/Engine/Plugins/VirtualProduction/RemoteControl/Source/WebRemoteControl/Private/RemoteControlWebSocketServer.cpp
--- a/Engine/Plugins/VirtualProduction/RemoteControl/Source/WebRemoteControl/Private/RemoteControlWebSocketServer.cpp
+++ b/Engine/Plugins/VirtualProduction/RemoteControl/Source/WebRemoteControl/Private/RemoteControlWebSocketServer.cpp
@@ -73,6 +73,16 @@ void FWebsocketMessageRouter::AddPreDispatch(TFunction<bool(const FRemoteControl
 	DispatchPreProcessor.Add(WebsocketPreprocessor);
 }
 
+void FWebsocketMessageRouter::AttemptDispatch(const struct FRemoteControlWebSocketMessage& Message)
+{
+	if (!PreDispatch(Message))
+	{
+		return;
+	}
+
+	Dispatch(Message);
+}
+
 bool FWebsocketMessageRouter::PreDispatch(const FRemoteControlWebSocketMessage& Message) const
 {
 	for (TFunction<bool(const FRemoteControlWebSocketMessage&)> PreprocessFunction : DispatchPreProcessor)
@@ -220,13 +230,7 @@ void FRCWebSocketServer::ReceivedRawPacket(void* Data, int32 Size, FGuid ClientI
 	if (TOptional<FRemoteControlWebSocketMessage> Message = RemoteControlWebSocketServer::ParseWebsocketMessage(Payload))
 	{
 		Message->ClientId = ClientId;
-		
-		if (!Router->PreDispatch(*Message))
-		{
-			return;
-		}
-	
-		Router->Dispatch(*Message);
+		Router->AttemptDispatch(*Message);
 	}
 }
 
diff --git a/Engine/Plugins/VirtualProduction/RemoteControl/Source/WebRemoteControl/Private/RemoteControlWebSocketServer.h b/Engine/Plugins/VirtualProduction/RemoteControl/Source/WebRemoteControl/Private/RemoteControlWebSocketServer.h
--- a/Engine/Plugins/VirtualProduction/RemoteControl/Source/WebRemoteControl/Private/RemoteControlWebSocketServer.h
+++ b/Engine/Plugins/VirtualProduction/RemoteControl/Source/WebRemoteControl/Private/RemoteControlWebSocketServer.h
@@ -35,6 +35,11 @@ public:
 	 */
 	void AddPreDispatch(TFunction<bool(const struct FRemoteControlWebSocketMessage& Message)> WebsocketPreprocessor);
 
+	/**
+	 * Attempt to dispatch a WebSocket message if its PreDispatch passes.
+	 */
+	void AttemptDispatch(const struct FRemoteControlWebSocketMessage& Message);
+
 private:
 	/** Preprocessors for the Dispatch Function */
 	TArray<TFunction<bool(const struct FRemoteControlWebSocketMessage& Message)>> DispatchPreProcessor;
diff --git a/Engine/Plugins/VirtualProduction/RemoteControl/Source/WebRemoteControl/Private/WebRemoteControl.cpp b/Engine/Plugins/VirtualProduction/RemoteControl/Source/WebRemoteControl/Private/WebRemoteControl.cpp
--- a/Engine/Plugins/VirtualProduction/RemoteControl/Source/WebRemoteControl/Private/WebRemoteControl.cpp
+++ b/Engine/Plugins/VirtualProduction/RemoteControl/Source/WebRemoteControl/Private/WebRemoteControl.cpp
@@ -766,6 +766,12 @@ void FWebRemoteControlModule::RegisterRoutes()
 		FWebSocketMessageDelegate::CreateRaw(this, &FWebRemoteControlModule::HandleWebSocketHttpMessage)
 		});
 
+	RegisterWebsocketRoute({
+		TEXT("Batch multiple WebSocket messages into one request"),
+		TEXT("batch"),
+		FWebSocketMessageDelegate::CreateRaw(this, &FWebRemoteControlModule::HandleWebSocketBatchMessage)
+		});
+
 	WebSocketHandler->RegisterRoutes(this);
 
 	EditorRoutes.RegisterRoutes(this);
@@ -1172,50 +1178,8 @@ bool FWebRemoteControlModule::HandlePresetSetPropertyRoute(const FHttpServerRequ
 		return true;
 	}
 
-	FRCObjectReference ObjectRef;
-
-	// Replace PropertyValue with the underlying property name.
-	TArray<uint8> NewPayload;
-	RemotePayloadSerializer::ReplaceFirstOccurence(SetPropertyRequest.TCHARBody, TEXT("PropertyValue"), RemoteControlProperty->FieldName.ToString(), NewPayload);
-
-	// Then deserialize the payload onto all the bound objects.
-	FMemoryReader NewPayloadReader(NewPayload);
-	FRCJsonStructDeserializerBackend Backend(NewPayloadReader);
-
-	ObjectRef.Property = RemoteControlProperty->GetProperty();
-	ObjectRef.Access = SetPropertyRequest.GenerateTransaction ? ERCAccess::WRITE_TRANSACTION_ACCESS : ERCAccess::WRITE_ACCESS;
-
-	bool bSuccess = true;
-
-	RemoteControlProperty->EnableEditCondition();
-
-	for (UObject* Object : RemoteControlProperty->GetBoundObjects())
-	{
-		IRemoteControlModule::Get().ResolveObjectProperty(ObjectRef.Access, Object, RemoteControlProperty->FieldPathInfo.ToString(), ObjectRef);
-
-		// Notify the handler before the change to ensure that the notification triggered by PostEditChange is ignored by the handler 
-		// if the client does not want remote change notifications.
-		if (ActingClientId.IsValid())
-		{
-			// Don't manually trigger a property change modification if this request gets converted to a function call.
-			if (ObjectRef.IsValid() && !RemoteControlPropertyUtilities::FindSetterFunction(ObjectRef.Property.Get(), ObjectRef.Object->GetClass()))
-			{
-				WebSocketHandler->NotifyPropertyChangedRemotely(ActingClientId, Preset->GetPresetId(), RemoteControlProperty->GetId());
-			}
-		}
-		if (SetPropertyRequest.ResetToDefault)
-		{
-			// set interception flag as an extra argument {}
-			constexpr bool bAllowIntercept = true;
-			bSuccess &= IRemoteControlModule::Get().ResetObjectProperties(ObjectRef, bAllowIntercept);
-		}
-		else
-		{
-			NewPayloadReader.Seek(0);
-			// Set a ERCPayloadType and TCHARBody in order to follow the replication path
-			bSuccess &= IRemoteControlModule::Get().SetObjectProperties(ObjectRef, Backend, ERCPayloadType::Json, NewPayload, SetPropertyRequest.Operation);
-		}
-	}
+	const bool bSuccess = WebRemoteControlInternalUtils::ModifyPropertyUsingPayload(
+		*RemoteControlProperty.Get(), SetPropertyRequest, SetPropertyRequest.TCHARBody, ActingClientId, *WebSocketHandler);
 
 	if (bSuccess)
 	{
@@ -2044,6 +2008,38 @@ void FWebRemoteControlModule::HandleWebSocketHttpMessage(const FRemoteControlWeb
 	LogRequestExternally(Wrapper.RequestId, TEXT("UE Sent"));
 }
 
+void FWebRemoteControlModule::HandleWebSocketBatchMessage(const FRemoteControlWebSocketMessage& WebSocketMessage)
+{
+	if (!WebSocketRouter)
+	{
+		return;
+	}
+
+	FRCWebSocketBatchRequest Body;
+	if (!WebRemoteControlInternalUtils::DeserializeRequestPayload(WebSocketMessage.RequestPayload, nullptr, Body))
+	{
+		return;
+	}
+
+	for (FRCWebSocketRequest& Request : Body.Requests)
+	{
+		FRemoteControlWebSocketMessage Message;
+
+		const FBlockDelimiters PayloadDelimiters = Request.GetParameterDelimiters(FRCWebSocketRequest::ParametersFieldLabel());
+		if (PayloadDelimiters.BlockStart != PayloadDelimiters.BlockEnd)
+		{
+			Message.RequestPayload = MakeArrayView(WebSocketMessage.RequestPayload).Slice(PayloadDelimiters.BlockStart, PayloadDelimiters.BlockEnd - PayloadDelimiters.BlockStart);
+		}
+
+		Message.ClientId = WebSocketMessage.ClientId;
+		Message.Header = WebSocketMessage.Header;
+		Message.MessageId = Request.Id;
+		Message.MessageName = Request.MessageName;
+
+		WebSocketRouter->AttemptDispatch(Message);
+	}
+}
+
 void FWebRemoteControlModule::InvokeWrappedRequest(const FRCRequestWrapper& Wrapper, FMemoryWriter& OutUTF8PayloadWriter, const FHttpServerRequest* TemplateRequest)
 {
 	TSharedRef<FHttpServerRequest> UnwrappedRequest = RemotePayloadSerializer::UnwrapHttpRequest(Wrapper, TemplateRequest);
diff --git a/Engine/Plugins/VirtualProduction/RemoteControl/Source/WebRemoteControl/Private/WebRemoteControl.h b/Engine/Plugins/VirtualProduction/RemoteControl/Source/WebRemoteControl/Private/WebRemoteControl.h
--- a/Engine/Plugins/VirtualProduction/RemoteControl/Source/WebRemoteControl/Private/WebRemoteControl.h
+++ b/Engine/Plugins/VirtualProduction/RemoteControl/Source/WebRemoteControl/Private/WebRemoteControl.h
@@ -134,6 +134,7 @@ private:
 
 	//~ Websocket route handlers
 	void HandleWebSocketHttpMessage(const struct FRemoteControlWebSocketMessage& WebSocketMessage);
+	void HandleWebSocketBatchMessage(const FRemoteControlWebSocketMessage& WebSocketMessage);
 
 	void InvokeWrappedRequest(const struct FRCRequestWrapper& Wrapper, FMemoryWriter& OutUTF8PayloadWriter, const FHttpServerRequest* TemplateRequest = nullptr);
 
diff --git a/Engine/Plugins/VirtualProduction/RemoteControl/Source/WebRemoteControl/Private/WebRemoteControlInternalUtils.cpp b/Engine/Plugins/VirtualProduction/RemoteControl/Source/WebRemoteControl/Private/WebRemoteControlInternalUtils.cpp
--- a/Engine/Plugins/VirtualProduction/RemoteControl/Source/WebRemoteControl/Private/WebRemoteControlInternalUtils.cpp
+++ b/Engine/Plugins/VirtualProduction/RemoteControl/Source/WebRemoteControl/Private/WebRemoteControlInternalUtils.cpp
@@ -328,10 +328,11 @@ bool WebRemoteControlInternalUtils::GetBatchRequestStructDelimiters(TConstArrayV
 		ErrorText = TEXT("Expected json object.");
 	}
 
+	bool bHasEncounteredError = false;
 	int32 LastEncounteredRequestId = -1;
 
 	// Mark the start/end of the param object in the payload
-	while (JsonReader->ReadNext(Notation) && ErrorText.IsEmpty())
+	while (!bHasEncounteredError && JsonReader->ReadNext(Notation))
 	{
 		switch (Notation)
 		{
@@ -358,6 +359,7 @@ bool WebRemoteControlInternalUtils::GetBatchRequestStructDelimiters(TConstArrayV
 			}
 			break;
 		case EJsonNotation::Error:
+			bHasEncounteredError = true;
 			ErrorText = JsonReader->GetErrorMessage();
 			break;
 		default:
@@ -379,6 +381,86 @@ bool WebRemoteControlInternalUtils::GetBatchRequestStructDelimiters(TConstArrayV
 	return true;
 }
 
+bool WebRemoteControlInternalUtils::GetBatchWebSocketRequestStructDelimiters(TConstArrayView<uint8> InTCHARPayload, TArray<FBlockDelimiters>& OutStructParameters, FString* OutErrorText)
+{
+	typedef UCS2CHAR PayloadCharType;
+	FMemoryReaderView Reader(InTCHARPayload);
+	TSharedRef<TJsonReader<PayloadCharType>> JsonReader = TJsonReader<PayloadCharType>::Create(&Reader);
+
+	EJsonNotation Notation;
+
+	FString ErrorText;
+	// The payload should be an object
+	JsonReader->ReadNext(Notation);
+	if (Notation != EJsonNotation::ObjectStart)
+	{
+		ErrorText = TEXT("Expected json object.");
+	}
+
+	bool bHasEncounteredError = false;
+	bool bIsInRequestsArray = false;
+	bool bIsDoneRequestsArray = false;
+	const TCHAR* const RequestsIdentifier = TEXT("Requests");
+
+	// Mark the start/end of the param object in the payload
+	while (!bHasEncounteredError && JsonReader->ReadNext(Notation) && !bIsDoneRequestsArray)
+	{
+		switch (Notation)
+		{
+		case EJsonNotation::ArrayStart:
+			if (JsonReader->GetIdentifier() == RequestsIdentifier)
+			{
+				bIsInRequestsArray = true;
+			}
+			break;
+
+		case EJsonNotation::ArrayEnd:
+			if (JsonReader->GetIdentifier() == RequestsIdentifier)
+			{
+				bIsDoneRequestsArray = true;
+			}
+			break;
+
+		case EJsonNotation::ObjectStart:
+			if (bIsInRequestsArray && JsonReader->GetIdentifier() == TEXT("Parameters"))
+			{
+				FBlockDelimiters& Delimiters = OutStructParameters.Emplace_GetRef();
+				Delimiters.BlockStart = Reader.Tell() - sizeof(PayloadCharType);
+				if (JsonReader->SkipObject())
+				{
+					Delimiters.BlockEnd = Reader.Tell();
+				}
+				else
+				{
+					ErrorText = FString::Printf(TEXT("%s object improperly formatted."), *JsonReader->GetIdentifier());
+				}
+			}
+			break;
+
+		case EJsonNotation::Error:
+			bHasEncounteredError = true;
+			ErrorText = JsonReader->GetErrorMessage();
+			break;
+
+		default:
+			// Ignore any other fields
+			break;
+		}
+	}
+
+	if (!ErrorText.IsEmpty())
+	{
+		UE_LOG(LogRemoteControl, Error, TEXT("Web Remote Control deserialization error: %s"), *ErrorText);
+		if (OutErrorText)
+		{
+			*OutErrorText = MoveTemp(ErrorText);
+		}
+		return false;
+	}
+
+	return true;
+}
+
 bool WebRemoteControlInternalUtils::ValidateContentType(const FHttpServerRequest& InRequest, FString InContentType, const FHttpResultCallback& InCompleteCallback)
 {
 	FString ErrorText;
diff --git a/Engine/Plugins/VirtualProduction/RemoteControl/Source/WebRemoteControl/Private/WebRemoteControlInternalUtils.h b/Engine/Plugins/VirtualProduction/RemoteControl/Source/WebRemoteControl/Private/WebRemoteControlInternalUtils.h
--- a/Engine/Plugins/VirtualProduction/RemoteControl/Source/WebRemoteControl/Private/WebRemoteControlInternalUtils.h
+++ b/Engine/Plugins/VirtualProduction/RemoteControl/Source/WebRemoteControl/Private/WebRemoteControlInternalUtils.h
@@ -14,7 +14,7 @@
 #include "RemoteControlRequest.h"
 #include "Templates/UnrealTypeTraits.h"
 #include "WebRemoteControlUtils.h"
-
+#include "WebSocketMessageHandler.h"
 
 namespace RemotePayloadSerializer
 {
@@ -148,6 +148,15 @@ namespace WebRemoteControlInternalUtils
 	 * @return Whether the delimiters were able to be found.
 	 */
 	UE_NODISCARD bool GetBatchRequestStructDelimiters(TConstArrayView<uint8> InTCHARPayload, TMap<int32, FBlockDelimiters>& OutStructParameters, FString* OutErrorText = nullptr);
+
+	/**
+	 * Get the struct delimiters for all the batched WebSocket requests.
+	 * @param InTCHARPayload The json payload to deserialize.
+	 * @param OutStructParameters Delimiters for each request.
+	 * @param OutErrorText If set, the string pointer will be populated with an error message on error.
+	 * @return Whether the delimiters were able to be found.
+	 */
+	UE_NODISCARD bool GetBatchWebSocketRequestStructDelimiters(TConstArrayView<uint8> InTCHARPayload, TArray<FBlockDelimiters>& OutStructParameters, FString* OutErrorText = nullptr);
 	
 	/**
 	 * Specialization of DeserializeRequestPayload that handles Batch requests.
@@ -182,7 +191,7 @@ namespace WebRemoteControlInternalUtils
 			if (InCompleteCallback)
 			{
 				TUniquePtr<FHttpServerResponse> Response = CreateHttpResponse();
-				CreateUTF8ErrorMessage(TEXT("Unable to deserialize request."), Response->Body);
+				CreateUTF8ErrorMessage(TEXT("Unable to get struct delimiters for batch request."), Response->Body);
 				(*InCompleteCallback)(MoveTemp(Response));
 			}
 			return false;
@@ -200,6 +209,67 @@ namespace WebRemoteControlInternalUtils
 		return true;
 	}
 
+	/**
+	 * Specialization of DeserializeRequestPayload that handles WebSocket batch requests.
+	 * This will populate the TCHARBody of all the wrapped requests.
+	 * @param InTCHARPayload The json payload to deserialize.
+	 * @param InCompleteCallback The callback to call error.
+	 * @param The structure to serialize using the request's content.
+	 * @return Whether the deserialization was successful.
+	 *
+	 * @note InCompleteCallback will be called with an appropriate http response if the deserialization fails.
+	 */
+	template <>
+	UE_NODISCARD inline bool DeserializeRequestPayload(TConstArrayView<uint8> InTCHARPayload, const FHttpResultCallback* InCompleteCallback, FRCWebSocketBatchRequest& OutDeserializedRequest)
+	{
+		FMemoryReaderView Reader(InTCHARPayload);
+		FJsonStructDeserializerBackend DeserializerBackend(Reader);
+
+		if (!FStructDeserializer::Deserialize(&OutDeserializedRequest, *FRCWebSocketBatchRequest::StaticStruct(), DeserializerBackend, FStructDeserializerPolicies()))
+		{
+			if (InCompleteCallback)
+			{
+				TUniquePtr<FHttpServerResponse> Response = CreateHttpResponse();
+				CreateUTF8ErrorMessage(TEXT("Unable to deserialize request."), Response->Body);
+				(*InCompleteCallback)(MoveTemp(Response));
+			}
+			return false;
+		}
+
+		TArray<FBlockDelimiters> Delimiters;
+		if (!GetBatchWebSocketRequestStructDelimiters(InTCHARPayload, Delimiters, nullptr))
+		{
+			if (InCompleteCallback)
+			{
+				TUniquePtr<FHttpServerResponse> Response = CreateHttpResponse();
+				CreateUTF8ErrorMessage(TEXT("Unable to get struct delimiters for batch request."), Response->Body);
+				(*InCompleteCallback)(MoveTemp(Response));
+			}
+			return false;
+		}
+
+		if (Delimiters.Num() != OutDeserializedRequest.Requests.Num())
+		{
+			if (InCompleteCallback)
+			{
+				TUniquePtr<FHttpServerResponse> Response = CreateHttpResponse();
+				CreateUTF8ErrorMessage(TEXT("Batch request delimiters did not match number of requests."), Response->Body);
+				(*InCompleteCallback)(MoveTemp(Response));
+			}
+			return false;
+		}
+
+		for (int32 RequestIndex = 0; RequestIndex < OutDeserializedRequest.Requests.Num(); ++RequestIndex)
+		{
+			FRCWebSocketRequest& Request = OutDeserializedRequest.Requests[RequestIndex];
+			FBlockDelimiters& ParametersDelimeters = Delimiters[RequestIndex];
+			Request.GetParameterDelimiters(FRCWebSocketRequest::ParametersFieldLabel()) = MoveTemp(ParametersDelimeters);
+			Request.TCHARBody = InTCHARPayload.Slice(ParametersDelimeters.BlockStart, ParametersDelimeters.BlockEnd - ParametersDelimeters.BlockStart);
+		}
+
+		return true;
+	}
+
 	/**
 	 * Adds a header indicating that the request is a wrapped request that originated from the engine itself.
 	 */
@@ -283,4 +353,64 @@ namespace WebRemoteControlInternalUtils
 		SerializerBackendType SerializerBackend(Writer);
 		FStructSerializer::Serialize(Struct.GetStructMemory(), *(UScriptStruct*)Struct.GetStruct(), SerializerBackend, FStructSerializerPolicies());
 	}
+
+	/**
+	 * Modify a property as specified by a remote request.
+	 * @param Property The property to modify.
+	 * @param Request The request to modify the property, containing additional information about how to modify it.
+	 * @param Payload The payload from which to deserialize property data.
+	 * @param ClientId The ID of the client that sent this request.
+	 * @param WebSocketHandler The WebSocket handler that will be notified of this remote change.
+	 */
+	template <typename RequestType>
+	bool ModifyPropertyUsingPayload(FRemoteControlProperty& Property, const RequestType& Request, const TArrayView<uint8>& Payload, const FGuid& ClientId, FWebSocketMessageHandler& WebSocketHandler)
+	{
+		FRCObjectReference ObjectRef;
+
+		// Replace PropertyValue with the underlying property name.
+		TArray<uint8> NewPayload;
+		RemotePayloadSerializer::ReplaceFirstOccurence(Payload, TEXT("PropertyValue"), Property.FieldName.ToString(), NewPayload);
+
+		// Then deserialize the payload onto all the bound objects.
+		FMemoryReader NewPayloadReader(NewPayload);
+		FRCJsonStructDeserializerBackend Backend(NewPayloadReader);
+
+		ObjectRef.Property = Property.GetProperty();
+		ObjectRef.Access = Request.GenerateTransaction ? ERCAccess::WRITE_TRANSACTION_ACCESS : ERCAccess::WRITE_ACCESS;
+
+		bool bSuccess = true;
+
+		Property.EnableEditCondition();
+
+		for (UObject* Object : Property.GetBoundObjects())
+		{
+			IRemoteControlModule::Get().ResolveObjectProperty(ObjectRef.Access, Object, Property.FieldPathInfo.ToString(), ObjectRef);
+
+			// Notify the handler before the change to ensure that the notification triggered by PostEditChange is ignored by the handler 
+			// if the client does not want remote change notifications.
+			if (ClientId.IsValid())
+			{
+				// Don't manually trigger a property change modification if this request gets converted to a function call.
+				if (ObjectRef.IsValid() && !RemoteControlPropertyUtilities::FindSetterFunction(ObjectRef.Property.Get(), ObjectRef.Object->GetClass()))
+				{
+					WebSocketHandler.NotifyPropertyChangedRemotely(ClientId, Property.GetOwner()->GetPresetId(), Property.GetId());
+				}
+			}
+
+			if (Request.ResetToDefault)
+			{
+				// set interception flag as an extra argument {}
+				constexpr bool bAllowIntercept = true;
+				bSuccess &= IRemoteControlModule::Get().ResetObjectProperties(ObjectRef, bAllowIntercept);
+			}
+			else
+			{
+				NewPayloadReader.Seek(0);
+				// Set a ERCPayloadType and TCHARBody in order to follow the replication path
+				bSuccess &= IRemoteControlModule::Get().SetObjectProperties(ObjectRef, Backend, ERCPayloadType::Json, NewPayload, Request.Operation);
+			}
+		}
+
+		return bSuccess;
+	}
 }
diff --git a/Engine/Plugins/VirtualProduction/RemoteControl/Source/WebRemoteControl/Private/WebSocketMessageHandler.cpp b/Engine/Plugins/VirtualProduction/RemoteControl/Source/WebRemoteControl/Private/WebSocketMessageHandler.cpp
--- a/Engine/Plugins/VirtualProduction/RemoteControl/Source/WebRemoteControl/Private/WebSocketMessageHandler.cpp
+++ b/Engine/Plugins/VirtualProduction/RemoteControl/Source/WebRemoteControl/Private/WebSocketMessageHandler.cpp
@@ -21,6 +21,7 @@
 #endif
 
 static TAutoConsoleVariable<int32> CVarWebRemoteControlFramesBetweenPropertyNotifications(TEXT("WebControl.FramesBetweenPropertyNotifications"), 5, TEXT("The number of frames between sending batches of property notifications."));
+const int64 FWebSocketMessageHandler::DefaultSequenceNumber = -1;
 
 namespace WebSocketMessageHandlerStructUtils
 {
@@ -37,6 +38,7 @@ namespace WebSocketMessageHandlerStructUtils
 	FName Prop_PresetName= "PresetName";
 	FName Prop_PresetId = "PresetId";
 	FName Prop_ChangedFields = "ChangedFields";
+	FName Prop_SequenceNumber = "SequenceNumber";
 
 	FName Struct_ActorPropertyValue= "WEBRC_ActorPropertyValue";
 	FName Prop_PropertyName = "PropertyName";
@@ -78,7 +80,8 @@ namespace WebSocketMessageHandlerStructUtils
 		{ 
 			Prop_PresetId,
 			Prop_PresetName,
-			Prop_Type
+			Prop_Type,
+			Prop_SequenceNumber
 		};
 
 		Args.ArrayProperties.Emplace(Prop_ChangedFields, PropertyValueStruct);
@@ -149,7 +152,7 @@ namespace WebSocketMessageHandlerStructUtils
 		return StructOnScope;
 	}
 
-	FStructOnScope CreatePresetFieldsChangedStructOnScope(const URemoteControlPreset* Preset, const TArray<FStructOnScope>& PropertyValuesOnScope)
+	FStructOnScope CreatePresetFieldsChangedStructOnScope(const URemoteControlPreset* Preset, const TArray<FStructOnScope>& PropertyValuesOnScope, int64 SequenceNumber)
 	{
 		UScriptStruct* PropertyValueStruct = (UScriptStruct*)PropertyValuesOnScope[0].GetStruct();
 		check(PropertyValueStruct);
@@ -160,6 +163,7 @@ namespace WebSocketMessageHandlerStructUtils
 		SetStringPropertyValue(Prop_Type, FieldsChangedOnScope, TEXT("PresetFieldsChanged"));
 		SetStringPropertyValue(Prop_PresetName, FieldsChangedOnScope, *Preset->GetFName().ToString());
 		SetStringPropertyValue(Prop_PresetId, FieldsChangedOnScope, *Preset->GetPresetId().ToString());
+		SetStringPropertyValue(Prop_SequenceNumber, FieldsChangedOnScope, FString::Printf(TEXT("%lld"), SequenceNumber));
 		SetStructArrayPropertyValue(Prop_ChangedFields, FieldsChangedOnScope, PropertyValuesOnScope);
 
 		return FieldsChangedOnScope;
@@ -251,41 +255,35 @@ void FWebSocketMessageHandler::RegisterRoutes(FWebRemoteControlModule* WebRemote
 #endif
 	
 	// WebSocket routes
-	TUniquePtr<FRemoteControlWebsocketRoute> PresetRegisterRoute = MakeUnique<FRemoteControlWebsocketRoute>(
+	RegisterRoute(WebRemoteControl, MakeUnique<FRemoteControlWebsocketRoute>(
 		TEXT("Subscribe to events emitted by a Remote Control Preset"),
 		TEXT("preset.register"),
 		FWebSocketMessageDelegate::CreateRaw(this, &FWebSocketMessageHandler::HandleWebSocketPresetRegister)
-		);
-
-	WebRemoteControl->RegisterWebsocketRoute(*PresetRegisterRoute);
-	Routes.Emplace(MoveTemp(PresetRegisterRoute));
+	));
 
-	TUniquePtr<FRemoteControlWebsocketRoute> PresetUnregisterRoute = MakeUnique<FRemoteControlWebsocketRoute>(
+	RegisterRoute(WebRemoteControl, MakeUnique<FRemoteControlWebsocketRoute>(
 		TEXT("Unsubscribe to events emitted by a Remote Control Preset"),
 		TEXT("preset.unregister"),
 		FWebSocketMessageDelegate::CreateRaw(this, &FWebSocketMessageHandler::HandleWebSocketPresetUnregister)
-		);
+	));
 
-	WebRemoteControl->RegisterWebsocketRoute(*PresetUnregisterRoute);
-	Routes.Emplace(MoveTemp(PresetUnregisterRoute));
-
-	TUniquePtr<FRemoteControlWebsocketRoute> ActorRegisterRoute = MakeUnique<FRemoteControlWebsocketRoute>(
+	RegisterRoute(WebRemoteControl, MakeUnique<FRemoteControlWebsocketRoute>(
 		TEXT("Subscribe to events emitted when actors of a particular type are added to/deleted from/renamed in the editor world"),
 		TEXT("actors.register"),
 		FWebSocketMessageDelegate::CreateRaw(this, &FWebSocketMessageHandler::HandleWebSocketActorRegister)
-		);
-
-	WebRemoteControl->RegisterWebsocketRoute(*ActorRegisterRoute);
-	Routes.Emplace(MoveTemp(ActorRegisterRoute));
+	));
 
-	TUniquePtr<FRemoteControlWebsocketRoute> ActorUnregisterRoute = MakeUnique<FRemoteControlWebsocketRoute>(
+	RegisterRoute(WebRemoteControl, MakeUnique<FRemoteControlWebsocketRoute>(
 		TEXT("Unsubscribe to events emitted when actors of a particular type are added to/deleted from/renamed in the editor world"),
 		TEXT("actors.unregister"),
 		FWebSocketMessageDelegate::CreateRaw(this, &FWebSocketMessageHandler::HandleWebSocketActorUnregister)
-		);
+	));
 
-	WebRemoteControl->RegisterWebsocketRoute(*ActorUnregisterRoute);
-	Routes.Emplace(MoveTemp(ActorUnregisterRoute));
+	RegisterRoute(WebRemoteControl, MakeUnique<FRemoteControlWebsocketRoute>(
+		TEXT("Modify the value of of a property exposed on a preset"),
+		TEXT("preset.property.modify"),
+		FWebSocketMessageDelegate::CreateRaw(this, &FWebSocketMessageHandler::HandleWebSocketPresetModifyProperty)
+	));
 }
 
 void FWebSocketMessageHandler::UnregisterRoutes(FWebRemoteControlModule* WebRemoteControl)
@@ -311,6 +309,14 @@ void FWebSocketMessageHandler::UnregisterRoutes(FWebRemoteControlModule* WebRemo
 	}
 }
 
+void FWebSocketMessageHandler::RegisterRoute(FWebRemoteControlModule* WebRemoteControl, TUniquePtr<FRemoteControlWebsocketRoute> Route)
+{
+	checkSlow(WebRemoteControl);
+
+	WebRemoteControl->RegisterWebsocketRoute(*Route);
+	Routes.Emplace(MoveTemp(Route));
+}
+
 void FWebSocketMessageHandler::RegisterActorHandlers()
 {
 #if WITH_EDITOR
@@ -502,6 +508,38 @@ void FWebSocketMessageHandler::HandleWebSocketActorUnregister(const FRemoteContr
 	UnregisterClientForActorClass(WebSocketMessage.ClientId, ActorClass);
 }
 
+void FWebSocketMessageHandler::HandleWebSocketPresetModifyProperty(const FRemoteControlWebSocketMessage& WebSocketMessage)
+{
+	FRCWebSocketPresetSetPropertyBody Body;
+	if (!WebRemoteControlInternalUtils::DeserializeRequestPayload(WebSocketMessage.RequestPayload, nullptr, Body))
+	{
+		return;
+	}
+
+	URemoteControlPreset* Preset = IRemoteControlModule::Get().ResolvePreset(Body.PresetName);
+	if (Preset == nullptr)
+	{
+		return;
+	}
+
+	const FGuid PropertyId = Preset->GetExposedEntityId(Body.PropertyLabel);
+	TSharedPtr<FRemoteControlProperty> RemoteControlProperty = Preset->GetExposedEntity<FRemoteControlProperty>(PropertyId).Pin();
+
+	if (!RemoteControlProperty.IsValid())
+	{
+		return;
+	}
+
+	WebRemoteControlInternalUtils::ModifyPropertyUsingPayload(*RemoteControlProperty.Get(), Body, WebSocketMessage.RequestPayload, ActingClientId, *this);
+
+	// Update the sequence number for this client
+	int64& SequenceNumber = ClientSequenceNumbers.FindOrAdd(ActingClientId, DefaultSequenceNumber);
+	if (SequenceNumber < Body.SequenceNumber)
+	{
+		SequenceNumber = Body.SequenceNumber;
+	}
+}
+
 void FWebSocketMessageHandler::ProcessChangedProperties()
 {
 	//Go over each property that were changed for each preset
@@ -528,8 +566,11 @@ void FWebSocketMessageHandler::ProcessChangedProperties()
 			// property class. See UE-139683
 			for (const FGuid& Id : ClientToEventsPair.Value)
 			{
+				const int64* ClientSequenceNumber = ClientSequenceNumbers.Find(Id);
+				const int64 SequenceNumber = ClientSequenceNumber ? *ClientSequenceNumber : DefaultSequenceNumber;
+
 				TArray<uint8> WorkingBuffer;
-				if (ClientToEventsPair.Value.Num() && WritePropertyChangeEventPayload(Preset, { Id }, WorkingBuffer))
+				if (ClientToEventsPair.Value.Num() && WritePropertyChangeEventPayload(Preset, { Id }, SequenceNumber, WorkingBuffer))
 				{
 					TArray<uint8> Payload;
 					WebRemoteControlUtils::ConvertToUTF8(WorkingBuffer, Payload);
@@ -780,6 +821,7 @@ void FWebSocketMessageHandler::OnConnectionClosedCallback(FGuid ClientId)
 
 	/** Remove this client's config. */
 	ClientConfigMap.Remove(ClientId);
+	ClientSequenceNumbers.Remove(ClientId);
 }
 
 void FWebSocketMessageHandler::OnEndFrame()
@@ -1114,7 +1156,7 @@ bool FWebSocketMessageHandler::ShouldProcessEventForPreset(const FGuid& PresetId
 	return PresetNotificationMap.Contains(PresetId) && PresetNotificationMap[PresetId].Num() > 0;
 }
 
-bool FWebSocketMessageHandler::WritePropertyChangeEventPayload(URemoteControlPreset* InPreset, const TSet<FGuid>& InModifiedPropertyIds, TArray<uint8>& OutBuffer)
+bool FWebSocketMessageHandler::WritePropertyChangeEventPayload(URemoteControlPreset* InPreset, const TSet<FGuid>& InModifiedPropertyIds, int64 InSequenceNumber, TArray<uint8>& OutBuffer)
 {
 	bool bHasProperty = false;
 
@@ -1137,7 +1179,7 @@ bool FWebSocketMessageHandler::WritePropertyChangeEventPayload(URemoteControlPre
 
 	if (PropValuesOnScope.Num())
 	{
-		FStructOnScope FieldsChangedEventOnScope = WebSocketMessageHandlerStructUtils::CreatePresetFieldsChangedStructOnScope(InPreset, PropValuesOnScope);
+		FStructOnScope FieldsChangedEventOnScope = WebSocketMessageHandlerStructUtils::CreatePresetFieldsChangedStructOnScope(InPreset, PropValuesOnScope, InSequenceNumber);
 
 		FMemoryWriter Writer(OutBuffer);
 		WebRemoteControlInternalUtils::SerializeStructOnScope(FieldsChangedEventOnScope, Writer);
diff --git a/Engine/Plugins/VirtualProduction/RemoteControl/Source/WebRemoteControl/Private/WebSocketMessageHandler.h b/Engine/Plugins/VirtualProduction/RemoteControl/Source/WebRemoteControl/Private/WebSocketMessageHandler.h
--- a/Engine/Plugins/VirtualProduction/RemoteControl/Source/WebRemoteControl/Private/WebSocketMessageHandler.h
+++ b/Engine/Plugins/VirtualProduction/RemoteControl/Source/WebRemoteControl/Private/WebSocketMessageHandler.h
@@ -78,6 +78,9 @@ private:
 	};
 
 
+	/** Register a WebSocket route */
+	void RegisterRoute(FWebRemoteControlModule* WebRemoteControl, TUniquePtr<FRemoteControlWebsocketRoute> Route);
+
 	/** Register handlers for actors being added/deleted (must happen after engine init). */
 	void RegisterActorHandlers();
 	
@@ -93,6 +96,12 @@ private:
 	/** Handles registration to callbacks for creation/destruction/rename of a given actor type */
 	void HandleWebSocketActorUnregister(const FRemoteControlWebSocketMessage& WebSocketMessage);
 
+	/** Handles multiple messages batched into a single request */
+	void HandleWebSocketBatchRequest(const FRemoteControlWebSocketMessage& WebSocketMessage);
+
+	/** Handles property modification for a given preset */
+	void HandleWebSocketPresetModifyProperty(const FRemoteControlWebSocketMessage& WebSocketMessage);
+
 	//Preset callbacks
 	void OnPresetExposedPropertiesModified(URemoteControlPreset* Owner, const TSet<FGuid>& ModifiedPropertyIds);
 	void OnPropertyExposed(URemoteControlPreset* Owner,  const FGuid& EntityId);
@@ -165,7 +174,7 @@ private:
 	/**
 	 * Write the provided list of events to a buffer.
 	 */
-	bool WritePropertyChangeEventPayload(URemoteControlPreset* InPreset, const TSet<FGuid>& InModifiedPropertyIds, TArray<uint8>& OutBuffer);
+	bool WritePropertyChangeEventPayload(URemoteControlPreset* InPreset, const TSet<FGuid>& InModifiedPropertyIds, int64 InSequenceNumber, TArray<uint8>& OutBuffer);
 
 	/**
 	 * Write the provided list of actor modifications to a buffer.
@@ -220,6 +229,9 @@ private:
 
 private:
 
+	/** Default sequence number for a client that hasn't reported one yet. */
+	static const int64 DefaultSequenceNumber;
+
 	/** Map type from class to Guids of clients listening for changes to actors of that class. */
 	typedef TMap<TWeakObjectPtr<UClass>, FWatchedClassData, FDefaultSetAllocator, TWeakObjectPtrMapKeyFuncs<TWeakObjectPtr<UClass>, FWatchedClassData>> FActorNotificationMap;
 
@@ -250,6 +262,9 @@ private:
 	/** Holds client-specific config if any. */
 	TMap<FGuid, FRCClientConfig> ClientConfigMap;
 
+	/** The largest sequence number received from each client. */
+	TMap<FGuid, int64> ClientSequenceNumbers;
+
 	/** Properties that changed for a frame, per preset.  */
 	TMap<FGuid, TMap<FGuid, TSet<FGuid>>> PerFrameModifiedProperties;
 
