[WebRemoteControl] Fix sequence number in property change events always being -1

diff --git a/Engine/Plugins/VirtualProduction/RemoteControl/Source/WebRemoteControl/Private/WebSocketMessageHandler.cpp b/Engine/Plugins/VirtualProduction/RemoteControl/Source/WebRemoteControl/Private/WebSocketMessageHandler.cpp
--- a/Engine/Plugins/VirtualProduction/RemoteControl/Source/WebRemoteControl/Private/WebSocketMessageHandler.cpp
+++ b/Engine/Plugins/VirtualProduction/RemoteControl/Source/WebRemoteControl/Private/WebSocketMessageHandler.cpp
@@ -530,10 +530,10 @@ void FWebSocketMessageHandler::HandleWebSocketPresetModifyProperty(const FRemote
 		return;
 	}
 
-	WebRemoteControlInternalUtils::ModifyPropertyUsingPayload(*RemoteControlProperty.Get(), Body, WebSocketMessage.RequestPayload, ActingClientId, *this);
+	WebRemoteControlInternalUtils::ModifyPropertyUsingPayload(*RemoteControlProperty.Get(), Body, WebSocketMessage.RequestPayload, WebSocketMessage.ClientId, *this);
 
 	// Update the sequence number for this client
-	int64& SequenceNumber = ClientSequenceNumbers.FindOrAdd(ActingClientId, DefaultSequenceNumber);
+	int64& SequenceNumber = ClientSequenceNumbers.FindOrAdd(WebSocketMessage.ClientId, DefaultSequenceNumber);
 	if (SequenceNumber < Body.SequenceNumber)
 	{
 		SequenceNumber = Body.SequenceNumber;
@@ -566,7 +566,7 @@ void FWebSocketMessageHandler::ProcessChangedProperties()
 			// property class. See UE-139683
 			for (const FGuid& Id : ClientToEventsPair.Value)
 			{
-				const int64* ClientSequenceNumber = ClientSequenceNumbers.Find(Id);
+				const int64* ClientSequenceNumber = ClientSequenceNumbers.Find(ClientToEventsPair.Key);
 				const int64 SequenceNumber = ClientSequenceNumber ? *ClientSequenceNumber : DefaultSequenceNumber;
 
 				TArray<uint8> WorkingBuffer;
