[EpicStageApp][WebRemoteControl] Automatically clean up transient preset when stage app disconnects

diff --git a/Engine/Plugins/VirtualProduction/RemoteControl/Source/RemoteControl/Private/RemoteControlModule.cpp b/Engine/Plugins/VirtualProduction/RemoteControl/Source/RemoteControl/Private/RemoteControlModule.cpp
--- a/Engine/Plugins/VirtualProduction/RemoteControl/Source/RemoteControl/Private/RemoteControlModule.cpp
+++ b/Engine/Plugins/VirtualProduction/RemoteControl/Source/RemoteControl/Private/RemoteControlModule.cpp
@@ -1229,6 +1229,26 @@ URemoteControlPreset* FRemoteControlModule::CreateTransientPreset()
 }
 
 bool FRemoteControlModule::DestroyTransientPreset(FName PresetName)
+{
+	if (URemoteControlPreset* Preset = ResolvePreset(PresetName))
+	{
+		return DestroyTransientPreset(Preset);
+	}
+
+	return false;
+}
+
+bool FRemoteControlModule::DestroyTransientPreset(const FGuid& PresetId)
+{
+	if (URemoteControlPreset* Preset = ResolvePreset(PresetId))
+	{
+		return DestroyTransientPreset(Preset);
+	}
+
+	return false;
+}
+
+bool FRemoteControlModule::IsPresetTransient(FName PresetName) const
 {
 	if (URemoteControlPreset* Preset = ResolvePreset(PresetName))
 	{
@@ -1236,19 +1256,22 @@ bool FRemoteControlModule::DestroyTransientPreset(FName PresetName)
 		{
 			if (TransientAssetData.GetAsset() == Preset)
 			{
-				// This call will also remove the asset from TransientPresets
-				OnAssetRemoved(TransientAssetData);
-				
-				FAssetRegistryModule::AssetDeleted(Preset);
-
-				if (UPackage* Package = Preset->GetPackage())
-				{
-					Package->MarkAsGarbage();
-				}
+				return true;
+			}
+		}
+	}
 
-				Preset->ClearFlags(RF_Public | RF_Standalone);
-				Preset->MarkAsGarbage();
+	return false;
+}
 
+bool FRemoteControlModule::IsPresetTransient(const FGuid& PresetId) const
+{
+	if (URemoteControlPreset* Preset = ResolvePreset(PresetId))
+	{
+		for (const FAssetData& TransientAssetData : TransientPresets)
+		{
+			if (TransientAssetData.GetAsset() == Preset)
+			{
 				return true;
 			}
 		}
@@ -1444,6 +1467,32 @@ void FRemoteControlModule::OnAssetRenamed(const FAssetData& AssetData, const FSt
 	CachedPresetsByName.FindOrAdd(AssetData.AssetName).AddUnique(AssetData);
 }
 
+bool FRemoteControlModule::DestroyTransientPreset(URemoteControlPreset* Preset)
+{
+	for (const FAssetData& TransientAssetData : TransientPresets)
+	{
+		if (TransientAssetData.GetAsset() == Preset)
+		{
+			// This call will also remove the asset from TransientPresets
+			OnAssetRemoved(TransientAssetData);
+
+			FAssetRegistryModule::AssetDeleted(Preset);
+
+			if (UPackage* Package = Preset->GetPackage())
+			{
+				Package->MarkAsGarbage();
+			}
+
+			Preset->ClearFlags(RF_Public | RF_Standalone);
+			Preset->MarkAsGarbage();
+
+			return true;
+		}
+	}
+
+	return false;
+}
+
 bool FRemoteControlModule::PropertyModificationShouldUseSetter(UObject* Object, FProperty* Property)
 {
 	if (!Property || !Object)
diff --git a/Engine/Plugins/VirtualProduction/RemoteControl/Source/RemoteControl/Private/RemoteControlModule.h b/Engine/Plugins/VirtualProduction/RemoteControl/Source/RemoteControl/Private/RemoteControlModule.h
--- a/Engine/Plugins/VirtualProduction/RemoteControl/Source/RemoteControl/Private/RemoteControlModule.h
+++ b/Engine/Plugins/VirtualProduction/RemoteControl/Source/RemoteControl/Private/RemoteControlModule.h
@@ -39,6 +39,9 @@ public:
 	virtual URemoteControlPreset* ResolvePreset(const FGuid& PresetId) const override;
 	virtual URemoteControlPreset* CreateTransientPreset() override;
 	virtual bool DestroyTransientPreset(FName PresetName) override;
+	virtual bool DestroyTransientPreset(const FGuid& PresetId) override;
+	virtual bool IsPresetTransient(FName PresetName) const override;
+	virtual bool IsPresetTransient(const FGuid& PresetId) const override;
 	virtual void GetPresets(TArray<TSoftObjectPtr<URemoteControlPreset>>& OutPresets) const override;
 	virtual void GetPresetAssets(TArray<FAssetData>& OutPresetAssets, bool bIncludeTransient = true) const override;
 	virtual const TMap<FName, FEntityMetadataInitializer>& GetDefaultMetadataInitializers() const override;
@@ -59,6 +62,9 @@ private:
 	void OnAssetRemoved(const FAssetData& AssetData);
 	void OnAssetRenamed(const FAssetData& AssetData, const FString&);
 
+	/** Destroy a transient preset using an object reference. */
+	bool DestroyTransientPreset(URemoteControlPreset* Preset);
+
 	/** Determines if a property modification should use a setter or default to deserializing directly onto an object. */
 	static bool PropertyModificationShouldUseSetter(UObject* Object, FProperty* Property);
 
diff --git a/Engine/Plugins/VirtualProduction/RemoteControl/Source/RemoteControl/Public/IRemoteControlModule.h b/Engine/Plugins/VirtualProduction/RemoteControl/Source/RemoteControl/Public/IRemoteControlModule.h
--- a/Engine/Plugins/VirtualProduction/RemoteControl/Source/RemoteControl/Public/IRemoteControlModule.h
+++ b/Engine/Plugins/VirtualProduction/RemoteControl/Source/RemoteControl/Public/IRemoteControlModule.h
@@ -332,12 +332,33 @@ public:
 	virtual URemoteControlPreset* CreateTransientPreset() = 0;
 
 	/**
-	 * Destroy a transient preset.
+	 * Destroy a transient preset using its name.
 	 * @arg PresetName name of the preset to destroy.
 	 * @return true if a transient preset with that name existed and was destroyed.
 	 */
 	virtual bool DestroyTransientPreset(FName PresetName) = 0;
 
+	/**
+	 * Destroy a transient preset using its id.
+     * @arg PresetId id of the preset to destroy.
+	 * @return true if a transient preset with that name existed and was destroyed.
+	 */
+	virtual bool DestroyTransientPreset(const FGuid& PresetId) = 0;
+
+	/**
+	 * Check whether a preset is transient using its name.
+	 * @arg PresetName name of the preset to check.
+	 * @return true if the preset exists and is transient.
+	 */
+	virtual bool IsPresetTransient(FName PresetName) const = 0;
+
+	/**
+	 * Check whether a preset is transient using its id.
+     * @arg PresetId id of the preset to check.
+	 * @return true if the preset exists and is transient.
+	 */
+	virtual bool IsPresetTransient(const FGuid& PresetId) const = 0;
+
 	/**
 	 * Get all the presets currently registered with the module.
 	 */
diff --git a/Engine/Plugins/VirtualProduction/RemoteControl/Source/WebRemoteControl/Private/RemoteControlRequest.h b/Engine/Plugins/VirtualProduction/RemoteControl/Source/WebRemoteControl/Private/RemoteControlRequest.h
--- a/Engine/Plugins/VirtualProduction/RemoteControl/Source/WebRemoteControl/Private/RemoteControlRequest.h
+++ b/Engine/Plugins/VirtualProduction/RemoteControl/Source/WebRemoteControl/Private/RemoteControlRequest.h
@@ -608,6 +608,31 @@ struct FRCWebSocketPresetRegisterBody : public FRCRequest
 	bool IgnoreRemoteChanges = false;
 };
 
+/**
+ * Holds a request made via websocket to automatically destroy a transient preset when the calling client disconnects.
+ */
+USTRUCT()
+struct FRCWebSocketTransientPresetAutoDestroyBody : public FRCRequest
+{
+	GENERATED_BODY()
+
+	FRCWebSocketTransientPresetAutoDestroyBody()
+	{
+		AddStructParameter(ParametersFieldLabel());
+	}
+
+	/**
+	 * Get the label for the property value struct.
+	 */
+	static FString ParametersFieldLabel() { return TEXT("Parameters"); }
+
+	/**
+	 * Name of the transient preset to mark for automatic destruction.
+	 */
+	UPROPERTY()
+	FString PresetName;
+};
+
 /**
  * Holds a request made via websocket to register for spawn/destroy events about a given actor type.
  */
diff --git a/Engine/Plugins/VirtualProduction/RemoteControl/Source/WebRemoteControl/Private/WebSocketMessageHandler.cpp b/Engine/Plugins/VirtualProduction/RemoteControl/Source/WebRemoteControl/Private/WebSocketMessageHandler.cpp
--- a/Engine/Plugins/VirtualProduction/RemoteControl/Source/WebRemoteControl/Private/WebSocketMessageHandler.cpp
+++ b/Engine/Plugins/VirtualProduction/RemoteControl/Source/WebRemoteControl/Private/WebSocketMessageHandler.cpp
@@ -267,6 +267,12 @@ void FWebSocketMessageHandler::RegisterRoutes(FWebRemoteControlModule* WebRemote
 		FWebSocketMessageDelegate::CreateRaw(this, &FWebSocketMessageHandler::HandleWebSocketPresetUnregister)
 	));
 
+	RegisterRoute(WebRemoteControl, MakeUnique<FRemoteControlWebsocketRoute>(
+		TEXT("Register a transient preset to be automatically destroyed when the calling client disconnects from WebSocket. If multiple clients call this, it will be destroyed once all the clients disconnect."),
+		TEXT("preset.transient.autodestroy"),
+		FWebSocketMessageDelegate::CreateRaw(this, &FWebSocketMessageHandler::HandleWebSocketTransientPresetAutoDestroy)
+		));
+
 	RegisterRoute(WebRemoteControl, MakeUnique<FRemoteControlWebsocketRoute>(
 		TEXT("Subscribe to events emitted when actors of a particular type are added to/deleted from/renamed in the editor world"),
 		TEXT("actors.register"),
@@ -444,6 +450,39 @@ void FWebSocketMessageHandler::HandleWebSocketPresetUnregister(const FRemoteCont
 	}
 }
 
+void FWebSocketMessageHandler::HandleWebSocketTransientPresetAutoDestroy(const FRemoteControlWebSocketMessage& WebSocketMessage)
+{
+	FRCWebSocketTransientPresetAutoDestroyBody Body;
+	if (!WebRemoteControlInternalUtils::DeserializeRequestPayload(WebSocketMessage.RequestPayload, nullptr, Body))
+	{
+		return;
+	}
+
+	URemoteControlPreset* Preset = nullptr;
+	IRemoteControlModule& RemoteControl = IRemoteControlModule::Get();
+
+	bool bIsTransient;
+	FGuid PresetId;
+
+	if (FGuid::ParseExact(Body.PresetName, EGuidFormats::Digits, PresetId))
+	{
+		Preset = RemoteControl.ResolvePreset(PresetId);
+		bIsTransient = RemoteControl.IsPresetTransient(PresetId);
+	}
+	else
+	{
+		Preset = RemoteControl.ResolvePreset(*Body.PresetName);
+		bIsTransient = RemoteControl.IsPresetTransient(*Body.PresetName);
+	}
+
+	if (!bIsTransient || !Preset)
+	{
+		return;
+	}
+
+	TransientPresetAutoDestroyClients.FindOrAdd(Preset->GetPresetId()).Add(WebSocketMessage.ClientId);
+}
+
 void FWebSocketMessageHandler::HandleWebSocketActorRegister(const FRemoteControlWebSocketMessage& WebSocketMessage)
 {
 	FRCWebSocketActorRegisterBody Body;
@@ -819,6 +858,23 @@ void FWebSocketMessageHandler::OnConnectionClosedCallback(FGuid ClientId)
 		UnregisterClientForActorClass(ClientId, WatchedClass.Get());
 	}
 
+	IRemoteControlModule& RemoteControl = IRemoteControlModule::Get();
+
+	// Clean up transient presets registered to auto-destroy for this client
+	for (auto Iter = TransientPresetAutoDestroyClients.CreateIterator(); Iter; ++Iter)
+	{
+		const FGuid& PresetId = Iter.Key();
+		TArray<FGuid>& ClientList = Iter.Value();
+
+		ClientList.Remove(ClientId);
+		if (ClientList.Num() == 0)
+		{
+			RemoteControl.DestroyTransientPreset(PresetId);
+			PresetNotificationMap.Remove(PresetId);
+			Iter.RemoveCurrent();
+		}
+	}
+
 	/** Remove this client's config. */
 	ClientConfigMap.Remove(ClientId);
 	ClientSequenceNumbers.Remove(ClientId);
diff --git a/Engine/Plugins/VirtualProduction/RemoteControl/Source/WebRemoteControl/Private/WebSocketMessageHandler.h b/Engine/Plugins/VirtualProduction/RemoteControl/Source/WebRemoteControl/Private/WebSocketMessageHandler.h
--- a/Engine/Plugins/VirtualProduction/RemoteControl/Source/WebRemoteControl/Private/WebSocketMessageHandler.h
+++ b/Engine/Plugins/VirtualProduction/RemoteControl/Source/WebRemoteControl/Private/WebSocketMessageHandler.h
@@ -90,6 +90,9 @@ private:
 	/** Handles unregistration to callbacks to a given preset */
 	void HandleWebSocketPresetUnregister(const FRemoteControlWebSocketMessage& WebSocketMessage);
 
+	/** Handles unregistration to callbacks to a given preset */
+	void HandleWebSocketTransientPresetAutoDestroy(const FRemoteControlWebSocketMessage& WebSocketMessage);
+
 	/** Handles registration to callbacks for creation/destruction/rename of a given actor type */
 	void HandleWebSocketActorRegister(const FRemoteControlWebSocketMessage& WebSocketMessage);
 
@@ -300,6 +303,9 @@ private:
 
 	/** Presets that had their layout modified for a frame. */
 	TSet<FGuid> PerFrameModifiedPresetLayouts;
+
+	/** Map from transient preset ID to clients which, when all disconnected, will automatically destroy the preset. */
+	TMap<FGuid, TArray<FGuid>> TransientPresetAutoDestroyClients;
 	
 	/** Holds the ID of the client currently making a request. Used to prevent sending back notifications to it. */
 	const FGuid& ActingClientId;
