RC initial protocol binding value from current value (for all types)

diff --git a/Engine/Plugins/VirtualProduction/RemoteControl/Source/RemoteControlProtocolWidgets/Private/ViewModels/ProtocolBindingViewModel.cpp b/Engine/Plugins/VirtualProduction/RemoteControl/Source/RemoteControlProtocolWidgets/Private/ViewModels/ProtocolBindingViewModel.cpp
--- a/Engine/Plugins/VirtualProduction/RemoteControl/Source/RemoteControlProtocolWidgets/Private/ViewModels/ProtocolBindingViewModel.cpp
+++ b/Engine/Plugins/VirtualProduction/RemoteControl/Source/RemoteControlProtocolWidgets/Private/ViewModels/ProtocolBindingViewModel.cpp
@@ -286,28 +286,28 @@ void FProtocolBindingViewModel::AddDefaultRangeMappings()
 				int64 IntMin = RemoteControlTypeUtilities::GetDefaultRangeValueMin<int64>(NumericProperty);
 				int64 IntMax = RemoteControlTypeUtilities::GetDefaultRangeValueMax<int64>(NumericProperty);
 
-		// fixup typename according to typesize, ie. it can be a UInt32 but a typesize of 2 makes its a UInt16
-		if(RangePropertyTypeName == NAME_UInt32Property && RangePropertySize > 0)
-		{
-			if(RangePropertySize == sizeof(uint8))
-			{
+				// fixup typename according to typesize, ie. it can be a UInt32 but a typesize of 2 makes its a UInt16
+				if(RangePropertyTypeName == NAME_UInt32Property && RangePropertySize > 0)
+				{
+					if(RangePropertySize == sizeof(uint8))
+					{
 						IntMin = RemoteControlTypeUtilities::GetDefaultRangeValueMin<uint8>(NumericProperty);
 						IntMax = RemoteControlTypeUtilities::GetDefaultRangeValueMax<uint8>(NumericProperty);
-				RangePropertyTypeName = NAME_ByteProperty;
-			}
-			else if(RangePropertySize == sizeof(uint16))
-			{
+						RangePropertyTypeName = NAME_ByteProperty;
+					}
+					else if(RangePropertySize == sizeof(uint16))
+					{
 						IntMin = RemoteControlTypeUtilities::GetDefaultRangeValueMin<uint16>(NumericProperty);
 						IntMax = RemoteControlTypeUtilities::GetDefaultRangeValueMax<uint16>(NumericProperty);
-				RangePropertyTypeName = NAME_UInt16Property; 
-			}
-			else if(RangePropertySize == sizeof(uint64))
-			{
+						RangePropertyTypeName = NAME_UInt16Property; 
+					}
+					else if(RangePropertySize == sizeof(uint64))
+					{
 						IntMin = RemoteControlTypeUtilities::GetDefaultRangeValueMin<uint64>(NumericProperty);
 						IntMax = RemoteControlTypeUtilities::GetDefaultRangeValueMax<uint64>(NumericProperty);
-				RangePropertyTypeName = NAME_UInt64Property;
-			}
-		}
+						RangePropertyTypeName = NAME_UInt64Property;
+					}
+				}
 
 				if(RangePropertyTypeName == NAME_ByteProperty)
 				{
@@ -420,10 +420,10 @@ void FProtocolBindingViewModel::AddDefaultRangeMappings()
 		else
 		{
 			if (MappingPropertyTypeName == NAME_BoolProperty)
-				{
+            {
 				MinItem->SetOutputValue<bool>(RemoteControlTypeUtilities::GetDefaultMappingValueMin<bool>(MappingProperty));
 				MaxItem->SetOutputValue<bool>(RemoteControlTypeUtilities::GetDefaultMappingValueMax<bool>(MappingProperty));
-				}
+            }
 
 			else if (MappingPropertyTypeName == NAME_StrProperty)
 			{
@@ -438,11 +438,11 @@ void FProtocolBindingViewModel::AddDefaultRangeMappings()
 				MaxItem->SetOutputValue<Type>(RemoteControlTypeUtilities::GetDefaultMappingValueMax<Type>(MappingProperty));
 			}
 			else if (MappingPropertyTypeName == NAME_TextProperty)
-				{
+            {
 				using Type = TRCTypeNameToType<NAME_TextProperty>::Type;
 				MinItem->SetOutputValue<Type>(RemoteControlTypeUtilities::GetDefaultMappingValueMin<Type>(MappingProperty));
 				MaxItem->SetOutputValue<Type>(RemoteControlTypeUtilities::GetDefaultMappingValueMax<Type>(MappingProperty));
-				}
+            }
 
 			else if (MappingPropertyTypeName == NAME_Vector)
 			{
@@ -450,17 +450,17 @@ void FProtocolBindingViewModel::AddDefaultRangeMappings()
 				MaxItem->SetOutputValue<FVector>(RemoteControlTypeUtilities::GetDefaultMappingValueMax<FVector>(MappingProperty));
 			}
 			else if (MappingPropertyTypeName == NAME_Vector2D)
-				{
+			{
 				using Type = TRCTypeNameToType<NAME_Vector2D>::Type;
 				MinItem->SetOutputValue<Type>(RemoteControlTypeUtilities::GetDefaultMappingValueMin<Type>(MappingProperty));
 				MaxItem->SetOutputValue<Type>(RemoteControlTypeUtilities::GetDefaultMappingValueMax<Type>(MappingProperty));
-				}
+			}
 			else if (MappingPropertyTypeName == NAME_Vector4)
-				{
+			{
 				using Type = TRCTypeNameToType<NAME_Vector4>::Type;
 				MinItem->SetOutputValue<Type>(RemoteControlTypeUtilities::GetDefaultMappingValueMin<Type>(MappingProperty));
 				MaxItem->SetOutputValue<Type>(RemoteControlTypeUtilities::GetDefaultMappingValueMax<Type>(MappingProperty));
-				}
+			}
 			else if (MappingPropertyTypeName == NAME_Rotator)
 			{
 				using Type = TRCTypeNameToType<NAME_Rotator>::Type;
@@ -472,9 +472,9 @@ void FProtocolBindingViewModel::AddDefaultRangeMappings()
 				using Type = TRCTypeNameToType<NAME_Color>::Type;
 				MinItem->SetOutputValue<Type>(RemoteControlTypeUtilities::GetDefaultMappingValueMin<Type>(MappingProperty));
 				MaxItem->SetOutputValue<Type>(RemoteControlTypeUtilities::GetDefaultMappingValueMax<Type>(MappingProperty));
-		}
+			}
 			else if (MappingPropertyTypeName == NAME_LinearColor)
-		{
+			{
 				using Type = TRCTypeNameToType<NAME_LinearColor>::Type;
 				MinItem->SetOutputValue<Type>(RemoteControlTypeUtilities::GetDefaultMappingValueMin<Type>(MappingProperty));
 				MaxItem->SetOutputValue<Type>(RemoteControlTypeUtilities::GetDefaultMappingValueMax<Type>(MappingProperty));
@@ -500,8 +500,11 @@ void FProtocolBindingViewModel::AddDefaultRangeMappings()
 			else
 			{
 				UE_LOG(LogRemoteControlProtocolWidgets, Warning, TEXT("AddDefaultRangeMappings type: %s for default value resolution."), *MappingPropertyTypeName.ToString());
-			}	
+			}
 		}
+		
+		// (Overrides above behavior) - get initial "min" range item value from current property value
+		MinItem->CopyFromCurrentPropertyValue();
 	}
 }
 
diff --git a/Engine/Plugins/VirtualProduction/RemoteControl/Source/RemoteControlProtocolWidgets/Private/ViewModels/ProtocolRangeViewModel.cpp b/Engine/Plugins/VirtualProduction/RemoteControl/Source/RemoteControlProtocolWidgets/Private/ViewModels/ProtocolRangeViewModel.cpp
--- a/Engine/Plugins/VirtualProduction/RemoteControl/Source/RemoteControlProtocolWidgets/Private/ViewModels/ProtocolRangeViewModel.cpp
+++ b/Engine/Plugins/VirtualProduction/RemoteControl/Source/RemoteControlProtocolWidgets/Private/ViewModels/ProtocolRangeViewModel.cpp
@@ -232,7 +232,7 @@ void FProtocolRangeViewModel::CopyFromCurrentPropertyValue() const
 {
 	check(IsValid());
 
-	// @note: if this fails, amend IsValid above to account for it
+	// @note: if this fails, ammend IsValid above to account for it
 	const TSharedPtr<const FRemoteControlField> ExposedEntity = GetPreset()->GetExposedEntity<FRemoteControlField>(GetBinding()->GetPropertyId()).Pin();
 	FRCFieldPathInfo FieldPathInfo = ExposedEntity->FieldPathInfo;
 
