Remote Control UI 2.0 refresh

1. Add toolbar to RC preset

2. Add search bar / filter in RC preset

3. Add action buttons in RC preset toolbar

diff --git a/Engine/Plugins/VirtualProduction/RemoteControl/Source/RemoteControlUI/Private/AssetEditor/RemoteControlPresetEditorToolkit.cpp b/Engine/Plugins/VirtualProduction/RemoteControl/Source/RemoteControlUI/Private/AssetEditor/RemoteControlPresetEditorToolkit.cpp
--- a/Engine/Plugins/VirtualProduction/RemoteControl/Source/RemoteControlUI/Private/AssetEditor/RemoteControlPresetEditorToolkit.cpp
+++ b/Engine/Plugins/VirtualProduction/RemoteControl/Source/RemoteControlUI/Private/AssetEditor/RemoteControlPresetEditorToolkit.cpp
@@ -83,7 +83,7 @@ void FRemoteControlPresetEditorToolkit::RegisterTabSpawners(const TSharedRef<FTa
 	InTabManager->RegisterTabSpawner(FRemoteControlUIModule::RemoteControlPanelTabName, FOnSpawnTab::CreateSP(this, &FRemoteControlPresetEditorToolkit::HandleTabManagerSpawnPanelTab))
 		.SetDisplayName(LOCTEXT("RemoteControlPanelMainTab", "Remote Control Panel"))
 		.SetGroup(WorkspaceMenuCategory.ToSharedRef())
-		.SetIcon(FSlateIcon(FEditorStyle::GetStyleSetName(), "LevelEditor.GameSettings.Small"));
+		.SetIcon(FSlateIcon(FEditorStyle::GetStyleSetName(), "LevelEditor.Tabs.Details"));
 
 	InTabManager->RegisterTabSpawner(FRemoteControlUIModule::EntityDetailsTabName, FOnSpawnTab::CreateSP(this, &FRemoteControlPresetEditorToolkit::HandleTabManagerSpawnDetailsTab))
 		.SetDisplayName(LOCTEXT("RemoteControlPanelDetailsTab", "Entity Details"))
@@ -133,7 +133,9 @@ TSharedRef<SDockTab> FRemoteControlPresetEditorToolkit::HandleTabManagerSpawnPan
 {
 	check(Args.GetTabId() == FRemoteControlUIModule::RemoteControlPanelTabName);
 
+	/*
 	FText TabName;
+	
 	if (Preset)
 	{
 		TabName = FText::FromString(Preset->GetName());
@@ -142,9 +144,10 @@ TSharedRef<SDockTab> FRemoteControlPresetEditorToolkit::HandleTabManagerSpawnPan
 	{
 		TabName = LOCTEXT("ControlPanelLabel", "Control Panel");
 	}
+	*/
 
 	return SNew(SDockTab)
-		.Label(TabName)
+		.Label(LOCTEXT("RCPanelTabName", "Remote Control"))
 		.TabColorScale(GetTabColorScale())	
 		[
 			PanelTab.ToSharedRef()
@@ -184,6 +187,17 @@ void FRemoteControlPresetEditorToolkit::InvokePanelTab()
 	{
 		Tab->SetContent(PanelTab.ToSharedRef());
 		Tab->SetOnTabClosed(SDockTab::FOnTabClosedCallback::CreateStatic(&Local::OnRemoteControlPresetClosed, TWeakPtr<IAssetEditorInstance>(SharedThis(this))));
+	
+		// Force the parent window to have minimum width.
+		TSharedPtr<SWindow> Window = Tab->GetParentWindow();
+		if (Window.IsValid())
+		{
+			FWindowSizeLimits SizeLimits = Window->GetSizeLimits();
+
+			SizeLimits.SetMinWidth(SRemoteControlPanel::MinimumPanelWidth);
+
+			Window->SetSizeLimits(SizeLimits);
+		}
 	}
 }
 
diff --git a/Engine/Plugins/VirtualProduction/RemoteControl/Source/RemoteControlUI/Private/Commands/RemoteControlCommands.cpp b/Engine/Plugins/VirtualProduction/RemoteControl/Source/RemoteControlUI/Private/Commands/RemoteControlCommands.cpp
--- a/Engine/Plugins/VirtualProduction/RemoteControl/Source/RemoteControlUI/Private/Commands/RemoteControlCommands.cpp
+++ b/Engine/Plugins/VirtualProduction/RemoteControl/Source/RemoteControlUI/Private/Commands/RemoteControlCommands.cpp
@@ -20,8 +20,22 @@ FRemoteControlCommands::FRemoteControlCommands()
 void FRemoteControlCommands::RegisterCommands()
 {
 	// Toggle Edit Mode
-
 	UI_COMMAND(ToggleEditMode, "Toggle Editing", "Toggles edit mode in remote control panel which enables user to modify the exposed entities.", EUserInterfaceActionType::Check, FInputChord(EModifierKey::Control, EKeys::E));
+
+	// Save Preset
+	UI_COMMAND(SavePreset, "Save", "Saves this preset", EUserInterfaceActionType::Button, FInputChord(EModifierKey::Control, EKeys::S));
+
+	// Find Preset
+	UI_COMMAND(FindPresetInContentBrowser, "Browse to Preset", "Browses to the associated preset and selects it in the most recently used Content Browser (summoning one if necessary)", EUserInterfaceActionType::Button, FInputChord(EModifierKey::Control, EKeys::B));
+
+	// Toggle Expose Functions
+	UI_COMMAND(ToggleExposeFunctions, "Functions", "Expose a function from available list of options.", EUserInterfaceActionType::Button, FInputChord());
+
+	// Toggle Protocol Mappings
+	UI_COMMAND(ToggleProtocolMappings, "Mappings", "View list of protocols mapped to active selection.", EUserInterfaceActionType::Button, FInputChord());
+
+	// Toggle Logic Editor
+	UI_COMMAND(ToggleLogicEditor, "Logic", "View the logic applied to active selection.", EUserInterfaceActionType::Button, FInputChord());
 }
 
 #undef LOCTEXT_NAMESPACE
diff --git a/Engine/Plugins/VirtualProduction/RemoteControl/Source/RemoteControlUI/Private/Commands/RemoteControlCommands.h b/Engine/Plugins/VirtualProduction/RemoteControl/Source/RemoteControlUI/Private/Commands/RemoteControlCommands.h
--- a/Engine/Plugins/VirtualProduction/RemoteControl/Source/RemoteControlUI/Private/Commands/RemoteControlCommands.h
+++ b/Engine/Plugins/VirtualProduction/RemoteControl/Source/RemoteControlUI/Private/Commands/RemoteControlCommands.h
@@ -25,4 +25,29 @@ public:
 	 * Holds the information about UI Command that toggles edit mode in RC Panel.
 	 */
 	TSharedPtr<FUICommandInfo> ToggleEditMode;
+
+	/**
+	 * Holds the information about UI Command that finds the actively edited preset in the Content Browser.
+	 */
+	TSharedPtr<FUICommandInfo> FindPresetInContentBrowser;
+
+	/**
+	 * Holds the information about UI Command that saves the actively edited preset.
+	 */
+	TSharedPtr<FUICommandInfo> SavePreset;
+
+	/**
+	 * Holds the information about UI Command that brings up a panel that facilitates exposing functions using RC API.
+	 */
+	TSharedPtr<FUICommandInfo> ToggleExposeFunctions;
+	
+	/**
+	 * Holds the information about UI Command that brings up a panel which holds the active protocol mappings.
+	 */
+	TSharedPtr<FUICommandInfo> ToggleProtocolMappings;
+	
+	/**
+	 * Holds the information about UI Command that brings up a panel which enables the RC Logical Behaviour.
+	 */
+	TSharedPtr<FUICommandInfo> ToggleLogicEditor;
 };
diff --git a/Engine/Plugins/VirtualProduction/RemoteControl/Source/RemoteControlUI/Private/UI/SRemoteControlPanel.cpp b/Engine/Plugins/VirtualProduction/RemoteControl/Source/RemoteControlUI/Private/UI/SRemoteControlPanel.cpp
--- a/Engine/Plugins/VirtualProduction/RemoteControl/Source/RemoteControlUI/Private/UI/SRemoteControlPanel.cpp
+++ b/Engine/Plugins/VirtualProduction/RemoteControl/Source/RemoteControlUI/Private/UI/SRemoteControlPanel.cpp
@@ -6,6 +6,7 @@
 #include "AssetRegistryModule.h"
 #include "ClassViewerFilter.h"
 #include "ClassViewerModule.h"
+#include "Commands/RemoteControlCommands.h"
 #include "Editor.h"
 #include "Editor/EditorPerformanceSettings.h"
 #include "EditorFontGlyphs.h"
@@ -14,6 +15,7 @@
 #include "FileHelpers.h"
 #include "Layout/Visibility.h"
 #include "Input/Reply.h"
+#include "Interfaces/IMainFrameModule.h"
 #include "IRemoteControlModule.h"
 #include "IRemoteControlProtocolWidgetsModule.h"
 #include "ISettingsModule.h"
@@ -42,14 +44,17 @@
 #include "SRCPanelExposedActor.h"
 #include "SRCPanelExposedField.h"
 #include "SRCPanelTreeNode.h"
+#include "Styling/ToolBarStyle.h"
 #include "Subsystems/Subsystem.h"
 #include "SWarningOrErrorBox.h"
 #include "Templates/SharedPointer.h"
 #include "Templates/SubclassOf.h"
 #include "Templates/UnrealTypeTraits.h"
 #include "Toolkits/IToolkitHost.h"
+#include "ToolMenus.h"
 #include "Widgets/DeclarativeSyntaxSupport.h"
 #include "Widgets/Docking/SDockTab.h"
+#include "Widgets/Images/SLayeredImage.h"
 #include "Widgets/Input/SButton.h"
 #include "Widgets/Input/SCheckBox.h"
 #include "Widgets/Layout/SBorder.h"
@@ -62,6 +67,9 @@
 
 #define LOCTEXT_NAMESPACE "RemoteControlPanel"
 
+const FName SRemoteControlPanel::DefaultRemoteControlPanelToolBarName("RemoteControlPanel.DefaultToolBar");
+const FName SRemoteControlPanel::AuxiliaryRemoteControlPanelToolBarName("RemoteControlPanel.AuxiliaryToolBar");
+const float SRemoteControlPanel::MinimumPanelWidth = 640.f;
 
 namespace RemoteControlPanelUtils
 {
@@ -92,6 +100,140 @@ namespace RemoteControlPanelUtils
 	}
 }
 
+/**
+ * UI representation of a auto resizing button.
+ */
+class SAutoResizeButton : public SCompoundWidget
+{
+	SLATE_BEGIN_ARGS(SAutoResizeButton)
+		: _ForceSmallIcons(false)
+		, _UICommand(nullptr)
+		, _IconOverride()
+	{}
+
+		SLATE_ATTRIBUTE(bool, ForceSmallIcons)
+
+		/** UI Command must be mapped to MainFrame Command List. */
+		SLATE_ARGUMENT(TSharedPtr<const FUICommandInfo>, UICommand)
+
+		SLATE_ARGUMENT(TAttribute<FSlateIcon>, IconOverride)
+
+	SLATE_END_ARGS()
+
+public:
+
+	void Construct(const FArguments& InArgs)
+	{
+		bForceSmallIcons = InArgs._ForceSmallIcons;
+
+		UICommand = InArgs._UICommand;
+
+		IconOverride = InArgs._IconOverride;
+
+		// Mimic Toolbar button style
+		const FToolBarStyle& ToolBarStyle = FAppStyle::Get().GetWidgetStyle<FToolBarStyle>("AssetEditorToolbar");
+
+		// Get the label & tooltip from the UI Command.
+		const TAttribute<FText> ActualLabel = UICommand.IsValid() ? UICommand->GetLabel() : FText::GetEmpty();
+		const TAttribute<FText> ActualToolTip = UICommand.IsValid() ? UICommand->GetDescription() : FText::GetEmpty();
+
+		// If we were supplied an image than go ahead and use that, otherwise we use a null widget
+		TSharedRef<SLayeredImage> IconWidget =
+			SNew(SLayeredImage)
+			.ColorAndOpacity(this, &SAutoResizeButton::GetIconForegroundColor)
+			.Visibility(EVisibility::HitTestInvisible)
+			.Image(this, &SAutoResizeButton::GetIconBrush);
+
+		ChildSlot
+		.Padding(ToolBarStyle.ButtonPadding.Left, 0.f, ToolBarStyle.ButtonPadding.Right, 0.f)
+			[
+				SNew(SButton)
+				.ContentPadding(ToolBarStyle.ButtonPadding)
+				.ButtonStyle(&ToolBarStyle.ButtonStyle)
+				.OnClicked(this, &SAutoResizeButton::OnClicked)
+				.ToolTipText(ActualToolTip)
+				.Content()
+				[
+					SNew(SHorizontalBox)
+
+					// Icon Widget
+					+ SHorizontalBox::Slot()
+					.AutoWidth()
+					.VAlign(VAlign_Center)
+					.HAlign(HAlign_Center)
+					[
+						IconWidget
+					]
+					
+					// Label Text
+					+ SHorizontalBox::Slot()
+					.AutoWidth()
+					.Padding(ToolBarStyle.LabelPadding)
+					.VAlign(VAlign_Center)
+					[
+						SNew(STextBlock)
+						.Visibility(this, &SAutoResizeButton::GetLabelVisibility)
+						.Text(ActualLabel)
+						.TextStyle(&ToolBarStyle.LabelStyle)	// Smaller font for tool tip labels
+					]
+				]
+			];
+	}
+
+private:
+	/** Called by Slate to determine whether labels are visible */
+	EVisibility GetLabelVisibility() const
+	{
+		const bool bUseSmallIcons = bForceSmallIcons.IsSet() ? bForceSmallIcons.Get() : false;
+
+		return bUseSmallIcons ? EVisibility::Collapsed : EVisibility::Visible;
+	}
+
+	/** Gets the icon brush for the toolbar block widget */
+	const FSlateBrush* GetIconBrush() const
+	{
+		const FSlateIcon ActionIcon = UICommand.IsValid() ?  UICommand->GetIcon() : FSlateIcon();
+		const FSlateIcon& ActualIcon = IconOverride.IsSet() ? IconOverride.Get() : ActionIcon;
+
+		return ActualIcon.GetIcon();
+	}
+
+	FSlateColor GetIconForegroundColor() const
+	{
+		// If any brush has a tint, don't assume it should be subdued
+		const FSlateBrush* Brush = GetIconBrush();
+		if (Brush && Brush->TintColor != FLinearColor::White)
+		{
+			return FLinearColor::White;
+		}
+
+		return FSlateColor::UseForeground();
+	}
+
+	/**
+	 * Called by Slate when this tool bar button's button is clicked
+	 */
+	FReply OnClicked()
+	{
+		IMainFrameModule& MainFrame = FModuleManager::Get().LoadModuleChecked<IMainFrameModule>("MainFrame");
+
+		TSharedPtr<FUICommandList> ActionList = MainFrame.GetMainFrameCommandBindings();
+
+		if (ActionList.IsValid() && UICommand.IsValid())
+		{
+			ActionList->ExecuteAction(UICommand.ToSharedRef());
+		}
+
+		return FReply::Handled();
+	}
+
+	TAttribute<bool> bForceSmallIcons;
+
+	TSharedPtr<const FUICommandInfo> UICommand;
+
+	TAttribute<FSlateIcon> IconOverride;
+};
+
 void SRemoteControlPanel::Construct(const FArguments& InArgs, URemoteControlPreset* InPreset, TSharedPtr<IToolkitHost> InToolkitHost)
 {
 	OnEditModeChange = InArgs._OnEditModeChange;
@@ -99,11 +241,68 @@ void SRemoteControlPanel::Construct(const FArguments& InArgs, URemoteControlPres
 	WidgetRegistry = MakeShared<FRCPanelWidgetRegistry>();
 	ToolkitHost = InToolkitHost;
 
-	UpdateRebindButtonVisibility();
-
 	TArray<TSharedRef<SWidget>> ExtensionWidgets;
 	FRemoteControlUIModule::Get().GetExtensionGenerators().Broadcast(ExtensionWidgets);
 
+	BindRemoteControlCommands();
+
+	// Edit Mode
+	AddToolbarWidget(SNew(SCheckBox)
+		.Style(FRemoteControlPanelStyle::Get(), "RemoteControlPanel.Switch")
+		.ToolTipText(LOCTEXT("EditModeTooltip", "Toggle Editing (Ctrl + E)"))
+		.IsChecked_Lambda([this]() { return this->bIsInEditMode ? ECheckBoxState::Checked : ECheckBoxState::Unchecked; })
+		.Padding(4.f)
+		.OnCheckStateChanged(this, &SRemoteControlPanel::OnEditModeCheckboxToggle)
+		.Content()
+		[
+			SNew(STextBlock)
+			.Text(LOCTEXT("EditModeLabel", "Edit"))
+		]
+	);
+
+	// Extension Generators
+	for (const TSharedRef<SWidget>& Widget : ExtensionWidgets)
+	{
+		AddToolbarWidget(Widget);
+	}
+
+	// Settings
+	AddToolbarWidget(SNew(SButton)
+		.ContentPadding(2.0f)
+		.ButtonStyle(FAppStyle::Get(), "FlatButton")
+		.TextStyle(FRemoteControlPanelStyle::Get(), "RemoteControlPanel.Button.TextStyle")
+		.OnClicked_Raw(this, &SRemoteControlPanel::OnClickSettingsButton)
+		.ToolTipText(LOCTEXT("OpenRemoteControlSettings", "Open Remote Control settings."))
+		[
+			SNew(SHorizontalBox)
+
+			+SHorizontalBox::Slot()
+			.HAlign(HAlign_Center)
+			.VAlign(VAlign_Center)
+			.AutoWidth()
+			.Padding(0.f, 2.f, 4.f, 2.f)
+			[
+				SNew(SImage)
+				.Image(FAppStyle::Get().GetBrush("Icons.Toolbar.Settings"))
+				.ColorAndOpacity(FSlateColor::UseForeground())
+			]
+
+			+ SHorizontalBox::Slot()
+			.HAlign(HAlign_Fill)
+			.VAlign(VAlign_Center)
+			.FillWidth(1.f)
+			.Padding(4.f, 2.f, 0.f, 2.f)
+			[
+				SNew(STextBlock)
+				.Text(LOCTEXT("SettingsLabel", "Settings"))
+			]
+		]);
+
+	GenerateToolbar();
+	GenerateAuxiliaryToolbar();
+
+	UpdateRebindButtonVisibility();
+
 	TSharedPtr<SHorizontalBox> TopExtensions;
 
 	EntityProtocolDetails = SNew(SBox);
@@ -137,6 +336,80 @@ void SRemoteControlPanel::Construct(const FArguments& InArgs, URemoteControlPres
 	[
 		SNew(SVerticalBox)
 		+ SVerticalBox::Slot()
+		.Padding(3.f)
+		.AutoHeight()
+		[
+			ToolbarWidgetContent.ToSharedRef()
+		]
+		+ SVerticalBox::Slot()
+		.HAlign(HAlign_Fill)
+		.AutoHeight()
+		[
+			// Separator
+			SNew(SSeparator)
+			.SeparatorImage(FAppStyle::Get().GetBrush("Separator"))
+			.Thickness(5.f)
+			.Orientation(EOrientation::Orient_Horizontal)
+		]
+		+ SVerticalBox::Slot()
+		.Padding(3.f)
+		.AutoHeight()
+		[
+			// Auxiliary Toolbar
+			SNew(SHorizontalBox)
+
+			// Search Box
+			+ SHorizontalBox::Slot()
+			.Padding(5.f, 3.f)
+			.VAlign(VAlign_Center)
+			.FillWidth(1.f)
+			[
+				SAssignNew(SearchBoxPtr, SSearchBox)
+			]
+
+			// Filters
+			+ SHorizontalBox::Slot()
+			.VAlign(VAlign_Center)
+			.AutoWidth()
+			.Padding(5.f, 3.f, 3.f, 3.f)
+			[
+				SNew(SComboButton)
+				.ComboButtonStyle(&FAppStyle::Get().GetWidgetStyle<FComboButtonStyle>("SimpleComboButtonWithIcon"))
+				.ForegroundColor(FSlateColor::UseStyle())
+				.ToolTipText(LOCTEXT("Filters_Tooltip", "Filter options for the RC Preset."))
+				.HasDownArrow(true)
+				.AddMetaData<FTagMetaData>(FTagMetaData(TEXT("RemoteControlFiltersCombo")))
+				.ContentPadding(FMargin(1.f, 0.f))
+				.ButtonContent()
+				[
+					SNew(SImage)
+					.Image(FAppStyle::Get().GetBrush("Icons.Filter"))
+					.ColorAndOpacity(FSlateColor::UseForeground())
+				]
+			]
+
+			// Separator
+			+SHorizontalBox::Slot()
+			.HAlign(HAlign_Center)
+			.VAlign(VAlign_Fill)
+			.AutoWidth()
+			.Padding(3.f, 3.f, 0.f, 3.f)
+			[
+				SNew(SSeparator)
+				.SeparatorImage(FAppStyle::Get().GetBrush("Separator"))
+				.Thickness(2.f)
+				.Orientation(EOrientation::Orient_Vertical)
+			]
+
+			// Mini Toolbar Widget
+			+ SHorizontalBox::Slot()
+			.AutoWidth()
+			.Padding(3.f)
+			[
+				AuxiliaryToolbarWidgetContent.ToSharedRef()
+			]
+		]
+		+ SVerticalBox::Slot()
 		.Padding(5.0f, 5.0f)
 		.AutoHeight()
 		[
@@ -376,19 +649,6 @@ void SRemoteControlPanel::Construct(const FArguments& InArgs, URemoteControlPres
 		]
 	];
 
-	for (const TSharedRef<SWidget>& Widget : ExtensionWidgets)
-	{
-		// We want to insert the widgets before the edit mode buttons.
-		constexpr int32 NumEditModeWidgets = 2;
-		const int32 ExtensionsPosition = ExtensionWidgets.Num() - NumEditModeWidgets;
-		TopExtensions->InsertSlot(ExtensionsPosition)
-			.VAlign(VAlign_Center)
-			.AutoWidth()
-			[
-				Widget
-			];
-	}
-
 	RegisterEvents();
 	CacheLevelClasses();
 	Refresh();
@@ -396,6 +656,7 @@ void SRemoteControlPanel::Construct(const FArguments& InArgs, URemoteControlPres
 
 SRemoteControlPanel::~SRemoteControlPanel()
 {
+	UnbindRemoteControlCommands();
 	UnregisterEvents();
 
 	// Clear the log
@@ -847,6 +1108,58 @@ void SRemoteControlPanel::OnMapChange(uint32)
 	Refresh();
 }
 
+void SRemoteControlPanel::BindRemoteControlCommands()
+{
+	const FRemoteControlCommands& Commands = FRemoteControlCommands::Get();
+
+	IMainFrameModule& MainFrame = FModuleManager::Get().LoadModuleChecked<IMainFrameModule>("MainFrame");
+
+	FUICommandList& ActionList = *MainFrame.GetMainFrameCommandBindings();
+
+	ActionList.MapAction(
+		Commands.SavePreset,
+		FExecuteAction::CreateSP(this, &SRemoteControlPanel::SaveAsset_Execute),
+		FCanExecuteAction::CreateSP(this, &SRemoteControlPanel::CanSaveAsset));
+
+	ActionList.MapAction(
+		Commands.FindPresetInContentBrowser,
+		FExecuteAction::CreateSP(this, &SRemoteControlPanel::FindInContentBrowser_Execute),
+		FCanExecuteAction::CreateSP(this, &SRemoteControlPanel::CanFindInContentBrowser));
+
+	ActionList.MapAction(
+		Commands.ToggleExposeFunctions,
+		FExecuteAction::CreateLambda([this] { OnToggleExposeFunctions(); }),
+		FCanExecuteAction::CreateLambda([] { return true; }));
+
+	ActionList.MapAction(
+		Commands.ToggleProtocolMappings,
+		FExecuteAction::CreateLambda([this] { OnToggleProtocolMappings(); }),
+		FCanExecuteAction::CreateLambda([] { return true; }));
+
+	ActionList.MapAction(
+		Commands.ToggleLogicEditor,
+		FExecuteAction::CreateLambda([this] { OnToggleLogicEditor(); }),
+		FCanExecuteAction::CreateLambda([] { return true; }));
+}
+
+void SRemoteControlPanel::UnbindRemoteControlCommands()
+{
+	if (FModuleManager::Get().IsModuleLoaded("MainFrame"))
+	{
+		const FRemoteControlCommands& Commands = FRemoteControlCommands::Get();
+
+		IMainFrameModule& MainFrame = FModuleManager::Get().LoadModuleChecked<IMainFrameModule>("MainFrame");
+
+		FUICommandList& ActionList = *MainFrame.GetMainFrameCommandBindings();
+
+		ActionList.UnmapAction(Commands.SavePreset);
+		ActionList.UnmapAction(Commands.FindPresetInContentBrowser);
+		ActionList.UnmapAction(Commands.ToggleExposeFunctions);
+		ActionList.UnmapAction(Commands.ToggleProtocolMappings);
+		ActionList.UnmapAction(Commands.ToggleLogicEditor);
+	}
+}
+
 void SRemoteControlPanel::RegisterEvents()
 {
 	FAssetRegistryModule& AssetRegistryModule = FModuleManager::LoadModuleChecked<FAssetRegistryModule>(TEXT("AssetRegistry"));
@@ -904,6 +1217,15 @@ void SRemoteControlPanel::Refresh()
 	EntityList->Refresh();
 }
 
+void SRemoteControlPanel::AddToolbarWidget(TSharedRef<SWidget> Widget)
+{
+	ToolbarWidgets.AddUnique(Widget);
+}
+
+void SRemoteControlPanel::RemoveAllToolbarWidgets()
+{
+	ToolbarWidgets.Empty();
+}
 
 void SRemoteControlPanel::Unexpose(const FRCExposesPropertyArgs& InPropertyArgs)
 {
@@ -1199,6 +1521,281 @@ void SRemoteControlPanel::OnMaterialCompiled(UMaterialInterface* MaterialInterfa
 	Refresh();
 }
 
+void SRemoteControlPanel::RegisterDefaultToolBar()
+{
+	UToolMenus* ToolMenus = UToolMenus::Get();
+	if (!ToolMenus->IsMenuRegistered(DefaultRemoteControlPanelToolBarName))
+	{
+		UToolMenu* ToolbarBuilder = ToolMenus->RegisterMenu(DefaultRemoteControlPanelToolBarName, NAME_None, EMultiBoxType::SlimHorizontalToolBar);
+
+#if 0
+		ToolbarBuilder->StyleName = "AssetEditorToolbar";
+#endif
+		{
+			FToolMenuSection& AssetSection = ToolbarBuilder->AddSection("Asset");
+			AssetSection.AddEntry(FToolMenuEntry::InitToolBarButton(FRemoteControlCommands::Get().SavePreset, TAttribute<FText>(), TAttribute<FText>(), FSlateIcon(FAppStyle::GetAppStyleSetName(), TEXT("Icons.Save"))));
+			AssetSection.AddEntry(FToolMenuEntry::InitToolBarButton(FRemoteControlCommands::Get().FindPresetInContentBrowser, LOCTEXT("FindInContentBrowserButton", "Browse"), TAttribute<FText>(), FSlateIcon(FAppStyle::GetAppStyleSetName(), TEXT("LevelEditor.OpenContentBrowser"))));
+			AssetSection.AddSeparator("Common");
+		}
+	}
+}
+
+void SRemoteControlPanel::GenerateToolbar()
+{
+	RegisterDefaultToolBar();
+
+	ToolbarWidgetContent = SNew(SBorder)
+		.BorderImage(FAppStyle::Get().GetBrush("Brushes.Panel"))
+		[
+			SNullWidget::NullWidget
+		];
+
+	UToolMenus* ToolMenus = UToolMenus::Get();
+	IMainFrameModule& MainFrame = FModuleManager::Get().LoadModuleChecked<IMainFrameModule>("MainFrame");
+
+	UToolMenu* GeneratedToolbar = ToolMenus->FindMenu(DefaultRemoteControlPanelToolBarName);
+	
+	GeneratedToolbar->Context = FToolMenuContext(MainFrame.GetMainFrameCommandBindings());
+	
+	TSharedRef<class SWidget> ToolBarWidget = ToolMenus->GenerateWidget(GeneratedToolbar);
+
+	TSharedRef<SWidget> MiscWidgets = SNullWidget::NullWidget;
+
+	if (ToolbarWidgets.Num() > 0)
+	{
+		TSharedRef<SHorizontalBox> MiscWidgetsHBox = SNew(SHorizontalBox);
+
+		for (int32 WidgetIdx = 0; WidgetIdx < ToolbarWidgets.Num(); ++WidgetIdx)
+		{
+			MiscWidgetsHBox->AddSlot()
+				.HAlign(HAlign_Center)
+				.VAlign(VAlign_Fill)
+				.AutoWidth()
+				[
+					SNew(SSeparator)
+					.SeparatorImage(FAppStyle::Get().GetBrush("Separator"))
+					.Thickness(2.f)
+					.Orientation(EOrientation::Orient_Vertical)
+				];
+
+			MiscWidgetsHBox->AddSlot()
+				.Padding(5.f, 0.f)
+				.AutoWidth()
+				.VAlign(VAlign_Center)
+				[
+					ToolbarWidgets[WidgetIdx]
+				];
+		}
+
+		MiscWidgets = MiscWidgetsHBox;
+	}
+
+	Toolbar =
+		SNew(SHorizontalBox)
+		+ SHorizontalBox::Slot()
+		.AutoWidth()
+		[
+			ToolBarWidget
+		]
+		+ SHorizontalBox::Slot()
+		.Padding(5.f, 0.f)
+		.HAlign(HAlign_Fill)
+		.VAlign(VAlign_Center)
+		.FillWidth(1.f)
+		[
+			SNew(STextBlock)
+			.Text(this, &SRemoteControlPanel::HandlePresetName)
+		]
+		+ SHorizontalBox::Slot()
+		.Padding(5.f, 0.f)
+		.HAlign(HAlign_Fill)
+		.VAlign(VAlign_Center)
+		.AutoWidth()
+		[
+			SNew(SCheckBox)
+			.IsChecked_Lambda([]() { return FRemoteControlLogger::Get().IsEnabled() ? ECheckBoxState::Checked : ECheckBoxState::Unchecked; })
+			.OnCheckStateChanged(this, &SRemoteControlPanel::OnLogCheckboxToggle)
+			.Padding(FMargin(5.f, 0.f))
+			[
+				SNew(STextBlock)
+				.Text(LOCTEXT("ShowLogLabel", "Show Log"))
+				.Justification(ETextJustify::Center)
+			]
+		]
+		+ SHorizontalBox::Slot()
+		.HAlign(HAlign_Right)
+		.AutoWidth()
+		[
+			MiscWidgets
+		];
+
+	if (ToolbarWidgetContent.IsValid())
+	{
+		ToolbarWidgetContent->SetContent(Toolbar.ToSharedRef());
+	}
+}
+
+void SRemoteControlPanel::RegisterAuxiliaryToolBar()
+{
+	UToolMenus* ToolMenus = UToolMenus::Get();
+	if (!ToolMenus->IsMenuRegistered(AuxiliaryRemoteControlPanelToolBarName))
+	{
+		UToolMenu* ToolbarBuilder = ToolMenus->RegisterMenu(AuxiliaryRemoteControlPanelToolBarName, NAME_None, EMultiBoxType::SlimHorizontalToolBar);
+
+#if 0
+		ToolbarBuilder->StyleName = "AssetEditorToolbar";
+#endif
+		{
+			FToolMenuSection& ToolsSection = ToolbarBuilder->AddSection("Tools");
+
+			ToolsSection.AddEntry(FToolMenuEntry::InitWidget("Functions"
+				, SNew(SAutoResizeButton)
+					.UICommand(FRemoteControlCommands::Get().ToggleExposeFunctions)
+					.ForceSmallIcons_Static(SRemoteControlPanel::ShouldForceSmallIcons)
+					.IconOverride(FSlateIcon(FAppStyle::GetAppStyleSetName(), TEXT("GraphEditor.Function_24x")))
+				, LOCTEXT("FunctionsLabel", "Functions")
+			)
+			);
+			
+			ToolsSection.AddEntry(FToolMenuEntry::InitWidget("Mappings"
+				, SNew(SAutoResizeButton)
+					.UICommand(FRemoteControlCommands::Get().ToggleProtocolMappings)
+					.ForceSmallIcons_Static(SRemoteControlPanel::ShouldForceSmallIcons)
+					.IconOverride(FSlateIcon(FEditorStyle::GetStyleSetName(), "LevelEditor.Tabs.StatsViewer"))
+				, LOCTEXT("MappingsLabel", "Mappings")
+			)
+			);
+			
+			ToolsSection.AddEntry(FToolMenuEntry::InitWidget("Logic"
+				, SNew(SAutoResizeButton)
+					.UICommand(FRemoteControlCommands::Get().ToggleLogicEditor)
+					.ForceSmallIcons_Static(SRemoteControlPanel::ShouldForceSmallIcons)
+					.IconOverride(FSlateIcon(FAppStyle::GetAppStyleSetName(), TEXT("GraphEditor.StateMachine_24x")))
+				, LOCTEXT("LogicLabel", "Logic")
+			)
+			);
+		}
+	}
+}
+
+void SRemoteControlPanel::GenerateAuxiliaryToolbar()
+{
+	RegisterAuxiliaryToolBar();
+
+	AuxiliaryToolbarWidgetContent = SNew(SBorder)
+		.BorderImage(FAppStyle::Get().GetBrush("Brushes.Panel"))
+		[
+			SNullWidget::NullWidget
+		];
+
+	UToolMenus* ToolMenus = UToolMenus::Get();
+	IMainFrameModule& MainFrame = FModuleManager::Get().LoadModuleChecked<IMainFrameModule>("MainFrame");
+
+	UToolMenu* GeneratedToolbar = ToolMenus->FindMenu(AuxiliaryRemoteControlPanelToolBarName);
+
+	GeneratedToolbar->Context = FToolMenuContext(MainFrame.GetMainFrameCommandBindings());
+
+	TSharedRef<class SWidget> ToolBarWidget = ToolMenus->GenerateWidget(GeneratedToolbar);
+
+	AuxiliaryToolbar =
+		SNew(SHorizontalBox)
+		+ SHorizontalBox::Slot()
+		.AutoWidth()
+		[
+			ToolBarWidget
+		];
+
+	if (AuxiliaryToolbarWidgetContent.IsValid())
+	{
+		AuxiliaryToolbarWidgetContent->SetContent(AuxiliaryToolbar.ToSharedRef());
+	}
+}
 
+FText SRemoteControlPanel::HandlePresetName() const
+{
+	if (Preset)
+	{
+		return FText::FromString(Preset->GetName());
+	}
+	
+	return FText::GetEmpty();
+}
+
+bool SRemoteControlPanel::CanSaveAsset() const
+{
+	return Preset.IsValid();
+}
+
+void SRemoteControlPanel::SaveAsset_Execute() const
+{
+	if (Preset.IsValid())
+	{
+		TArray<UPackage*> PackagesToSave;
+
+		if (!Preset->IsAsset())
+		{
+			// Log an invalid object but don't try to save it
+			UE_LOG(LogRemoteControl, Log, TEXT("Invalid object to save: %s"), (Preset.IsValid()) ? *Preset->GetFullName() : TEXT("Null Object"));
+		}
+		else
+		{
+			PackagesToSave.Add(Preset->GetOutermost());
+		}
+
+		constexpr bool bCheckDirtyOnAssetSave = false;
+		constexpr bool bPromptToSave = false;
+
+		FEditorFileUtils::PromptForCheckoutAndSave(PackagesToSave, bCheckDirtyOnAssetSave, bPromptToSave);
+	}
+}
+
+bool SRemoteControlPanel::CanFindInContentBrowser() const
+{
+	return Preset.IsValid();
+}
+
+void SRemoteControlPanel::FindInContentBrowser_Execute() const
+{
+	if (Preset.IsValid())
+	{
+		TArray<UObject*> ObjectsToSyncTo;
+
+		ObjectsToSyncTo.Add(Preset.Get());
+
+		GEditor->SyncBrowserToObjects(ObjectsToSyncTo);
+	}
+}
+
+bool SRemoteControlPanel::ShouldForceSmallIcons()
+{
+	// Find the DockTab that houses this RemoteControlPreset widget in it.
+	FLevelEditorModule& LevelEditorModule = FModuleManager::GetModuleChecked<FLevelEditorModule>("LevelEditor");
+	TSharedPtr<FTabManager> EditorTabManager = LevelEditorModule.GetLevelEditorTabManager();
+
+	if (TSharedPtr<SDockTab> Tab = EditorTabManager->FindExistingLiveTab(FRemoteControlUIModule::RemoteControlPanelTabName))
+	{
+		if (TSharedPtr<SWindow> Window = Tab->GetParentWindow())
+		{
+			const FVector2D& ActualWindowSize = Window->GetSizeInScreen() / Window->GetDPIScaleFactor();
+
+			// Need not to check for less than the minimum value as user can never go beyond that limit while resizing the parent window.
+			return ActualWindowSize.X == SRemoteControlPanel::MinimumPanelWidth ? true : false;
+		}
+	}
+
+	return false;
+}
+
+void SRemoteControlPanel::OnToggleExposeFunctions() const
+{
+}
+
+void SRemoteControlPanel::OnToggleProtocolMappings() const
+{
+}
+
+void SRemoteControlPanel::OnToggleLogicEditor() const
+{
+}
 
 #undef LOCTEXT_NAMESPACE /*RemoteControlPanel*/
diff --git a/Engine/Plugins/VirtualProduction/RemoteControl/Source/RemoteControlUI/Private/UI/SRemoteControlPanel.h b/Engine/Plugins/VirtualProduction/RemoteControl/Source/RemoteControlUI/Private/UI/SRemoteControlPanel.h
--- a/Engine/Plugins/VirtualProduction/RemoteControl/Source/RemoteControlUI/Private/UI/SRemoteControlPanel.h
+++ b/Engine/Plugins/VirtualProduction/RemoteControl/Source/RemoteControlUI/Private/UI/SRemoteControlPanel.h
@@ -30,6 +30,7 @@ struct SRCPanelTreeNode;
 class SRCPanelFunctionPicker;
 class SRemoteControlPanel;
 class SRCPanelExposedEntitiesList;
+class SSearchBox;
 class STextBlock;
 class URemoteControlPreset;
 
@@ -112,8 +113,16 @@ public:
 	/** Re-create the sections of the panel. */
 	void Refresh();
 
+	/** Adds or removes widgets from the default toolbar in this asset editor */
+	void AddToolbarWidget(TSharedRef<SWidget> Widget);
+	void RemoveAllToolbarWidgets();
+
 private:
 
+	//~ Remote Control Commands
+	void BindRemoteControlCommands();
+	void UnbindRemoteControlCommands();
+
 	/** Register editor events needed to handle reloading objects and blueprint libraries. */
 	void RegisterEvents();
 	/** Unregister editor events */
@@ -205,8 +214,44 @@ private:
 
 	/** Handle triggers each time when any material has been recompiled */
 	void OnMaterialCompiled(class UMaterialInterface* MaterialInterface);
+	
+	/** Registers default tool bar */
+	static void RegisterDefaultToolBar();
+
+	/** Makes a default asset editing toolbar */
+	void GenerateToolbar();
+	
+	/** Registers Auxiliary tool bar */
+	static void RegisterAuxiliaryToolBar();
+
+	/** Makes a Auxiliary asset editing toolbar */
+	void GenerateAuxiliaryToolbar();
+
+	FText HandlePresetName() const;
+
+	/** Called to test if "Save" should be enabled for this asset */
+	bool CanSaveAsset() const;
+
+	/** Called when "Save" is clicked for this asset */
+	void SaveAsset_Execute() const;
+
+	/** Called to test if "Find in Content Browser" should be enabled for this asset */
+	bool CanFindInContentBrowser() const;
+
+	/** Called when "Find in Content Browser" is clicked for this asset */
+	void FindInContentBrowser_Execute() const;
+
+	static bool ShouldForceSmallIcons();
+
+	void OnToggleExposeFunctions() const;
+
+	void OnToggleProtocolMappings() const;
+
+	void OnToggleLogicEditor() const;
 
 private:
+	static const FName DefaultRemoteControlPanelToolBarName;
+	static const FName AuxiliaryRemoteControlPanelToolBarName;
 	/** Holds the preset asset. */
 	TStrongObjectPtr<URemoteControlPreset> Preset;
 	/** Whether the panel is in edit mode. */
@@ -245,4 +290,19 @@ private:
 	FTimerHandle NextTickTimerHandle;
 	/** The toolkit that hosts this panel. */
 	TWeakPtr<IToolkitHost> ToolkitHost;
+	/** Asset Editor Default Toolbar */
+	TSharedPtr<SWidget> Toolbar;
+	/** Asset Editor Auxiliary Toolbar */
+	TSharedPtr<SWidget> AuxiliaryToolbar;
+	/** The widget that will house the default Toolbar widget */
+	TSharedPtr<SBorder> ToolbarWidgetContent;
+	/** The widget that will house the Auxiliary Toolbar widget */
+	TSharedPtr<SBorder> AuxiliaryToolbarWidgetContent;
+	/** Additional widgets to be added to the toolbar */
+	TArray<TSharedRef<SWidget>> ToolbarWidgets;
+	/** The text box used to search for tags. */
+	TSharedPtr<SSearchBox> SearchBoxPtr;
+
+public:
+	static const float MinimumPanelWidth;
 };
diff --git a/Engine/Plugins/VirtualProduction/RemoteControl/Source/RemoteControlUI/RemoteControlUI.Build.cs b/Engine/Plugins/VirtualProduction/RemoteControl/Source/RemoteControlUI/RemoteControlUI.Build.cs
--- a/Engine/Plugins/VirtualProduction/RemoteControl/Source/RemoteControlUI/RemoteControlUI.Build.cs
+++ b/Engine/Plugins/VirtualProduction/RemoteControl/Source/RemoteControlUI/RemoteControlUI.Build.cs
@@ -34,6 +34,7 @@ public class RemoteControlUI : ModuleRules
 				"SceneOutliner",
 				"Slate",
 				"SlateCore",
+				"ToolMenus",
 				"ToolWidgets",
 				"UnrealEd"
 			}
diff --git a/Engine/Plugins/VirtualProduction/RemoteControlWebInterface/Source/RemoteControlWebInterface/Private/RCWebInterfaceCustomizations.cpp b/Engine/Plugins/VirtualProduction/RemoteControlWebInterface/Source/RemoteControlWebInterface/Private/RCWebInterfaceCustomizations.cpp
--- a/Engine/Plugins/VirtualProduction/RemoteControlWebInterface/Source/RemoteControlWebInterface/Private/RCWebInterfaceCustomizations.cpp
+++ b/Engine/Plugins/VirtualProduction/RemoteControlWebInterface/Source/RemoteControlWebInterface/Private/RCWebInterfaceCustomizations.cpp
@@ -243,31 +243,59 @@ void FRCWebInterfaceCustomizations::GeneratePanelExtensions(TArray<TSharedRef<SW
 		[
 			SNew(SButton)
 			.ButtonStyle(FEditorStyle::Get(), "FlatButton")
+			.ToolTipText_Lambda([this]()
+				{
+					switch (WebApp->GetStatus())
+					{
+					case FRemoteControlWebInterfaceProcess::EStatus::Error: return LOCTEXT("AppErrorTooltip", "An error occurred when launching the web app.\r\nClick to change the Web App's settings.");
+					case FRemoteControlWebInterfaceProcess::EStatus::Running: return LOCTEXT("AppRunningTooltip", "Open the web app in a web browser.");
+					case FRemoteControlWebInterfaceProcess::EStatus::Stopped: return LOCTEXT("AppStoppedTooltip", "The web app is not currently running.\r\nClick to change the Web App's settings.");
+					default: return FText();
+					}
+				})
 			.OnClicked_Raw(this, &FRCWebInterfaceCustomizations::OpenWebApp)
 			[
-				SNew(STextBlock)
-				.ToolTipText_Lambda([this]()
-					{
-						switch (WebApp->GetStatus())
+				SNew(SHorizontalBox)
+
+				+SHorizontalBox::Slot()
+				.HAlign(HAlign_Center)
+				.VAlign(VAlign_Center)
+				.AutoWidth()
+				.Padding(0.f, 2.f, 4.f, 2.f)
+				[
+					SNew(STextBlock)
+					.TextStyle(FEditorStyle::Get(), "NormalText.Important")
+					.Font(FEditorStyle::Get().GetFontStyle("FontAwesome.14"))
+					.Text_Lambda([this]()
 						{
-						case FRemoteControlWebInterfaceProcess::EStatus::Error: return LOCTEXT("AppErrorTooltip", "An error occurred when launching the web app.\r\nClick to change the Web App's settings.");
-						case FRemoteControlWebInterfaceProcess::EStatus::Running: return LOCTEXT("AppRunningTooltip", "Open the web app in a web browser.");
-						case FRemoteControlWebInterfaceProcess::EStatus::Stopped: return LOCTEXT("AppStoppedTooltip", "The web app is not currently running.\r\nClick to change the Web App's settings.");
-						default: return FText();
-						}
-					})
-				.TextStyle(FEditorStyle::Get(), "NormalText.Important")
-				.Font(FEditorStyle::Get().GetFontStyle("FontAwesome.10"))
-				.Text_Lambda([this]()
+							switch (WebApp->GetStatus())
+							{
+							case FRemoteControlWebInterfaceProcess::EStatus::Error: return FEditorFontGlyphs::Exclamation_Triangle;
+							case FRemoteControlWebInterfaceProcess::EStatus::Running: return FEditorFontGlyphs::External_Link;
+							case FRemoteControlWebInterfaceProcess::EStatus::Stopped: return FEditorFontGlyphs::Exclamation_Triangle;
+							default: return FEditorFontGlyphs::Exclamation_Triangle;
+							}
+						})
+				]
+
+				+ SHorizontalBox::Slot()
+				.HAlign(HAlign_Fill)
+				.VAlign(VAlign_Center)
+				.FillWidth(1.f)
+				.Padding(4.f, 2.f, 0.f, 2.f)
+				[
+					SNew(STextBlock)
+					.Text_Lambda([this]()
 					{
 						switch (WebApp->GetStatus())
 						{
-						case FRemoteControlWebInterfaceProcess::EStatus::Error: return FEditorFontGlyphs::Exclamation_Triangle;
-						case FRemoteControlWebInterfaceProcess::EStatus::Running: return FEditorFontGlyphs::External_Link;
-						case FRemoteControlWebInterfaceProcess::EStatus::Stopped: return FEditorFontGlyphs::Exclamation_Triangle;
-						default: return FEditorFontGlyphs::Exclamation_Triangle;
+							case FRemoteControlWebInterfaceProcess::EStatus::Error: return LOCTEXT("ErrorText", "Error");
+							case FRemoteControlWebInterfaceProcess::EStatus::Running: return LOCTEXT("RunningText", "Web App");
+							case FRemoteControlWebInterfaceProcess::EStatus::Stopped: return LOCTEXT("StoppedText", "Stopped");
+							default: return LOCTEXT("UnknownErrorText", "Unknown Error");
 						}
 					})
+				]
 			]
 		]
 	);
