[EpicStageApp] Control lightcard color from app

diff --git a/Engine/Plugins/VirtualProduction/RemoteControl/Source/RemoteControl/Private/RemoteControlFunctionLibrary.cpp b/Engine/Plugins/VirtualProduction/RemoteControl/Source/RemoteControl/Private/RemoteControlFunctionLibrary.cpp
--- a/Engine/Plugins/VirtualProduction/RemoteControl/Source/RemoteControl/Private/RemoteControlFunctionLibrary.cpp
+++ b/Engine/Plugins/VirtualProduction/RemoteControl/Source/RemoteControl/Private/RemoteControlFunctionLibrary.cpp
@@ -54,4 +54,65 @@ bool URemoteControlFunctionLibrary::ExposeActor(URemoteControlPreset* Preset, AA
 		return Preset->ExposeActor(Actor, {Args.DisplayName, RemoteControlFunctionLibrary::GetOrCreateGroup(Preset, Args.GroupName)}).IsValid();
 	}
 	return false;
-}
\ No newline at end of file
+}
+
+bool URemoteControlFunctionLibrary::ApplyColorWheelDelta(UObject* TargetObject, const FString& PropertyName, FColorWheelColor DeltaValue, FColorWheelColor ReferenceColor)
+{
+	if (!TargetObject)
+	{
+		return false;
+	}
+
+	FProperty* Property = TargetObject->GetClass()->FindPropertyByName(FName(*PropertyName));
+	if (const FStructProperty* ColorProperty = CastField<FStructProperty>(Property))
+	{
+		FLinearColor Color;
+		ColorProperty->GetValue_InContainer(TargetObject, &Color);
+
+		// Convert to HSV
+		Color = Color.LinearRGBToHSV();
+
+		FVector2D Position;
+
+		if (Color.B > UE_DOUBLE_KINDA_SMALL_NUMBER)
+		{
+			// Determine direction as a unit vector based on the calculated hue
+			const double HueRadians = FMath::DegreesToRadians(Color.R);
+			Position = FVector2D(FMath::Cos(HueRadians), FMath::Sin(HueRadians));
+			
+			// Multiply the unit vector by saturation to determine the current position in the color wheel
+			Position *= Color.G;
+		}
+		else
+		{
+			// Color's value is too low to determine the hue and saturation. Fall back to the reference color's position
+			Position = ReferenceColor.Position;
+		}
+
+		// Apply the delta to the position, then convert back to hue and saturation
+		Position += FVector2D(DeltaValue.Position.X, DeltaValue.Position.Y);
+		Color.R = FMath::Fmod(FMath::RadiansToDegrees(FMath::Atan2(Position.Y, Position.X)) + 360.0, 360.0);
+		Color.G = FMath::Clamp(Position.Length(), 0.0, 1.0);
+
+		// Apply the value and alpha changes directly
+		Color.B = FMath::Clamp(Color.B + DeltaValue.Value, 0.0, 1.0);
+		Color.A = FMath::Clamp(Color.A + DeltaValue.Alpha, 0.0, 1.0);
+
+		Color = Color.HSVToLinearRGB();
+
+#if WITH_EDITOR
+		TargetObject->PreEditChange(Property);
+#endif
+
+		ColorProperty->SetValue_InContainer(TargetObject, &Color);
+
+#if WITH_EDITOR
+		FPropertyChangedEvent ChangeEvent(Property, EPropertyChangeType::ValueSet);
+		TargetObject->PostEditChangeProperty(ChangeEvent);
+#endif
+
+		return true;
+	}
+
+	return false;
+}
diff --git a/Engine/Plugins/VirtualProduction/RemoteControl/Source/RemoteControl/Public/RemoteControlFunctionLibrary.h b/Engine/Plugins/VirtualProduction/RemoteControl/Source/RemoteControl/Public/RemoteControlFunctionLibrary.h
--- a/Engine/Plugins/VirtualProduction/RemoteControl/Source/RemoteControl/Public/RemoteControlFunctionLibrary.h
+++ b/Engine/Plugins/VirtualProduction/RemoteControl/Source/RemoteControl/Public/RemoteControlFunctionLibrary.h
@@ -29,11 +29,37 @@ struct FRemoteControlOptionalExposeArgs
 	FString GroupName;
 };
 
+/** A color as represented by a position on a color wheel. */
+USTRUCT(Blueprintable)
+struct FColorWheelColor
+{
+	GENERATED_BODY()
+
+	/**
+	 * The position on the unit circle. Magnitude should be in range [0, 1].
+	 */
+	UPROPERTY()
+	FVector2D Position;
+
+	/**
+	 * The color's value component (as in HSV).
+	 */
+	UPROPERTY()
+	double Value;
+
+	/**
+	 * The color's alpha component.
+	 */
+	UPROPERTY()
+	double Alpha;
+};
+
 UCLASS()
 class URemoteControlFunctionLibrary : public UBlueprintFunctionLibrary
 {
 	GENERATED_BODY()
 	
+public:
 	/**
 	 * Expose a property in a remote control preset.
 	 * @param Preset the preset to expose the property in.
@@ -65,4 +91,15 @@ class URemoteControlFunctionLibrary : public UBlueprintFunctionLibrary
 	 */
 	UFUNCTION(BlueprintCallable, Category = "RemoteControlPreset")
 	static bool ExposeActor(URemoteControlPreset* Preset, AActor* Actor, FRemoteControlOptionalExposeArgs Args);
+
+	/**
+	 * Add/subtract from the value of an FLinearColor property using a delta value based on color wheel coordinates.
+	 * @param TargetObject the object that contains the property to modify.
+	 * @param PropertyName the name of the property to modify.
+	 * @param DeltaValue the amount to change the color by.
+	 * @param ReferenceColor if the color's current position on the wheel is ambiguous as calculated from RGB values (e.g. black), use this reference color's position instead.
+	 * @return Whether the operation was successful.
+	 */
+	UFUNCTION(BlueprintCallable, Category = "Remote Control Color")
+	static bool ApplyColorWheelDelta(UObject* TargetObject, const FString& PropertyName, FColorWheelColor DeltaValue, FColorWheelColor ReferenceColor);
 };
\ No newline at end of file
diff --git a/Engine/Plugins/VirtualProduction/RemoteControl/Source/WebRemoteControl/Private/RemoteControlRequest.h b/Engine/Plugins/VirtualProduction/RemoteControl/Source/WebRemoteControl/Private/RemoteControlRequest.h
--- a/Engine/Plugins/VirtualProduction/RemoteControl/Source/WebRemoteControl/Private/RemoteControlRequest.h
+++ b/Engine/Plugins/VirtualProduction/RemoteControl/Source/WebRemoteControl/Private/RemoteControlRequest.h
@@ -715,3 +715,24 @@ struct FRCWebSocketPresetSetPropertyBody : public FRCRequest
 	int64 SequenceNumber = -1;
 };
 
+
+/**
+ * Holds a request made via websocket to call an exposed function on an object.
+ */
+USTRUCT()
+struct FRCWebSocketCallBody : public FRCCallRequest
+{
+	GENERATED_BODY()
+
+	FRCWebSocketCallBody()
+		: FRCCallRequest()
+	{
+	}
+
+	/**
+	 * The sequence number of this change. The highest sequence number received from this client will be
+	 * sent back to the client in future PresetFieldsChanged events.
+	 */
+	UPROPERTY()
+	int64 SequenceNumber = -1;
+};
diff --git a/Engine/Plugins/VirtualProduction/RemoteControl/Source/WebRemoteControl/Private/WebSocketMessageHandler.cpp b/Engine/Plugins/VirtualProduction/RemoteControl/Source/WebRemoteControl/Private/WebSocketMessageHandler.cpp
--- a/Engine/Plugins/VirtualProduction/RemoteControl/Source/WebRemoteControl/Private/WebSocketMessageHandler.cpp
+++ b/Engine/Plugins/VirtualProduction/RemoteControl/Source/WebRemoteControl/Private/WebSocketMessageHandler.cpp
@@ -290,6 +290,12 @@ void FWebSocketMessageHandler::RegisterRoutes(FWebRemoteControlModule* WebRemote
 		TEXT("preset.property.modify"),
 		FWebSocketMessageDelegate::CreateRaw(this, &FWebSocketMessageHandler::HandleWebSocketPresetModifyProperty)
 	));
+
+	RegisterRoute(WebRemoteControl, MakeUnique<FRemoteControlWebsocketRoute>(
+		TEXT("Call a function on an object"),
+		TEXT("object.call"),
+		FWebSocketMessageDelegate::CreateRaw(this, &FWebSocketMessageHandler::HandleWebSocketFunctionCall)
+	));
 }
 
 void FWebSocketMessageHandler::UnregisterRoutes(FWebRemoteControlModule* WebRemoteControl)
@@ -569,14 +575,59 @@ void FWebSocketMessageHandler::HandleWebSocketPresetModifyProperty(const FRemote
 		return;
 	}
 
+	UpdateSequenceNumber(WebSocketMessage.ClientId, Body.SequenceNumber);
+
 	WebRemoteControlInternalUtils::ModifyPropertyUsingPayload(*RemoteControlProperty.Get(), Body, WebSocketMessage.RequestPayload, WebSocketMessage.ClientId, *this);
+}
 
-	// Update the sequence number for this client
-	int64& SequenceNumber = ClientSequenceNumbers.FindOrAdd(WebSocketMessage.ClientId, DefaultSequenceNumber);
-	if (SequenceNumber < Body.SequenceNumber)
+void FWebSocketMessageHandler::HandleWebSocketFunctionCall(const FRemoteControlWebSocketMessage& WebSocketMessage)
+{
+	FRCWebSocketCallBody Body;
+	if (!WebRemoteControlInternalUtils::DeserializeRequestPayload(WebSocketMessage.RequestPayload, nullptr, Body))
 	{
-		SequenceNumber = Body.SequenceNumber;
+		return;
 	}
+
+	FRCCall Call;
+	if (IRemoteControlModule::Get().ResolveCall(Body.ObjectPath, Body.FunctionName, Call.CallRef, nullptr))
+	{
+		// Initialize the param struct with default parameters
+		Call.bGenerateTransaction = Body.GenerateTransaction;
+		Call.ParamStruct = FStructOnScope(Call.CallRef.Function.Get());
+
+		// If some parameters were provided, deserialize them
+		const FBlockDelimiters& ParametersDelimiters = Body.GetStructParameters().FindChecked(FRCCallRequest::ParametersLabel());
+		if (ParametersDelimiters.BlockStart > 0)
+		{
+			// Extract function parameters from the request payload
+			const int64 DelimitersSize = ParametersDelimiters.GetBlockSize();
+
+			TArray<uint8> ParamPayload;
+			ParamPayload.SetNumUninitialized(DelimitersSize);
+
+			const uint8* DataStart = &WebSocketMessage.RequestPayload[ParametersDelimiters.BlockStart];
+			FMemory::Memcpy(ParamPayload.GetData(), DataStart, DelimitersSize);
+
+			// Using the extracted data, set up a reader and attempt to deserialize
+			FMemoryReader Reader(ParamPayload);
+			FRCJsonStructDeserializerBackend Backend(Reader);
+
+			if (!FStructDeserializer::Deserialize((void*)Call.ParamStruct.GetStructMemory(), *const_cast<UStruct*>(Call.ParamStruct.GetStruct()), Backend, FStructDeserializerPolicies()))
+			{
+				// Failed to deserialize parameters
+				return;
+			}
+		}
+	}
+
+	if (!Call.IsValid())
+	{
+		return;
+	}
+
+	UpdateSequenceNumber(WebSocketMessage.ClientId, Body.SequenceNumber);
+
+	IRemoteControlModule::Get().InvokeCall(Call);
 }
 
 void FWebSocketMessageHandler::ProcessChangedProperties()
@@ -1573,3 +1624,13 @@ void FWebSocketMessageHandler::UnregisterClientForActorClass(const FGuid& Client
 		}
 	}
 }
+
+void FWebSocketMessageHandler::UpdateSequenceNumber(const FGuid& ClientId, int64 NewSequenceNumber)
+{
+	// Update the sequence number for this client
+	int64& StoredSequenceNumber = ClientSequenceNumbers.FindOrAdd(ClientId, DefaultSequenceNumber);
+	if (StoredSequenceNumber < NewSequenceNumber)
+	{
+		StoredSequenceNumber = NewSequenceNumber;
+	}
+}
diff --git a/Engine/Plugins/VirtualProduction/RemoteControl/Source/WebRemoteControl/Private/WebSocketMessageHandler.h b/Engine/Plugins/VirtualProduction/RemoteControl/Source/WebRemoteControl/Private/WebSocketMessageHandler.h
--- a/Engine/Plugins/VirtualProduction/RemoteControl/Source/WebRemoteControl/Private/WebSocketMessageHandler.h
+++ b/Engine/Plugins/VirtualProduction/RemoteControl/Source/WebRemoteControl/Private/WebSocketMessageHandler.h
@@ -99,12 +99,12 @@ private:
 	/** Handles registration to callbacks for creation/destruction/rename of a given actor type */
 	void HandleWebSocketActorUnregister(const FRemoteControlWebSocketMessage& WebSocketMessage);
 
-	/** Handles multiple messages batched into a single request */
-	void HandleWebSocketBatchRequest(const FRemoteControlWebSocketMessage& WebSocketMessage);
-
 	/** Handles property modification for a given preset */
 	void HandleWebSocketPresetModifyProperty(const FRemoteControlWebSocketMessage& WebSocketMessage);
 
+	/** Handles calling a Blueprint function */
+	void HandleWebSocketFunctionCall(const FRemoteControlWebSocketMessage& WebSocketMessage);
+
 	//Preset callbacks
 	void OnPresetExposedPropertiesModified(URemoteControlPreset* Owner, const TSet<FGuid>& ModifiedPropertyIds);
 	void OnPropertyExposed(URemoteControlPreset* Owner,  const FGuid& EntityId);
@@ -230,6 +230,11 @@ private:
 	 */
 	void UnregisterClientForActorClass(const FGuid& ClientId, TSubclassOf<AActor> ActorClass);
 
+	/**
+	 * Update the sequence number for a client when a new one is received.
+	 */
+	void UpdateSequenceNumber(const FGuid& ClientId, int64 NewSequenceNumber);
+
 private:
 
 	/** Default sequence number for a client that hasn't reported one yet. */
