[EpicStageApp][WebRemoteControl] Change ApplyColorWheelDelta to take structs as refs

diff --git a/Engine/Plugins/VirtualProduction/RemoteControl/Source/RemoteControl/Private/RemoteControlFunctionLibrary.cpp b/Engine/Plugins/VirtualProduction/RemoteControl/Source/RemoteControl/Private/RemoteControlFunctionLibrary.cpp
--- a/Engine/Plugins/VirtualProduction/RemoteControl/Source/RemoteControl/Private/RemoteControlFunctionLibrary.cpp
+++ b/Engine/Plugins/VirtualProduction/RemoteControl/Source/RemoteControl/Private/RemoteControlFunctionLibrary.cpp
@@ -56,7 +56,7 @@ bool URemoteControlFunctionLibrary::ExposeActor(URemoteControlPreset* Preset, AA
 	return false;
 }
 
-bool URemoteControlFunctionLibrary::ApplyColorWheelDelta(UObject* TargetObject, const FString& PropertyName, FColorWheelColor DeltaValue, FColorWheelColor ReferenceColor)
+bool URemoteControlFunctionLibrary::ApplyColorWheelDelta(UObject* TargetObject, const FString& PropertyName, const FColorWheelColor& DeltaValue, const FColorWheelColor& ReferenceColor)
 {
 	if (!TargetObject)
 	{
diff --git a/Engine/Plugins/VirtualProduction/RemoteControl/Source/RemoteControl/Public/RemoteControlFunctionLibrary.h b/Engine/Plugins/VirtualProduction/RemoteControl/Source/RemoteControl/Public/RemoteControlFunctionLibrary.h
--- a/Engine/Plugins/VirtualProduction/RemoteControl/Source/RemoteControl/Public/RemoteControlFunctionLibrary.h
+++ b/Engine/Plugins/VirtualProduction/RemoteControl/Source/RemoteControl/Public/RemoteControlFunctionLibrary.h
@@ -101,5 +101,5 @@ public:
 	 * @return Whether the operation was successful.
 	 */
 	UFUNCTION(BlueprintCallable, Category = "Remote Control Color")
-	static bool ApplyColorWheelDelta(UObject* TargetObject, const FString& PropertyName, FColorWheelColor DeltaValue, FColorWheelColor ReferenceColor);
+	static bool ApplyColorWheelDelta(UObject* TargetObject, const FString& PropertyName, const FColorWheelColor& DeltaValue, const FColorWheelColor& ReferenceColor);
 };
\ No newline at end of file
