Fix uninitialized variables in remote control

diff --git a/Engine/Plugins/VirtualProduction/RemoteControl/Source/RemoteControl/Public/RemoteControlFunctionLibrary.h b/Engine/Plugins/VirtualProduction/RemoteControl/Source/RemoteControl/Public/RemoteControlFunctionLibrary.h
--- a/Engine/Plugins/VirtualProduction/RemoteControl/Source/RemoteControl/Public/RemoteControlFunctionLibrary.h
+++ b/Engine/Plugins/VirtualProduction/RemoteControl/Source/RemoteControl/Public/RemoteControlFunctionLibrary.h
@@ -39,19 +39,19 @@ struct FColorWheelColor
 	 * The position on the unit circle. Magnitude should be in range [0, 1].
 	 */
 	UPROPERTY()
-	FVector2D Position;
+	FVector2D Position = FVector2D::ZeroVector;
 
 	/**
 	 * The color's value component (as in HSV).
 	 */
 	UPROPERTY()
-	double Value;
+	double Value = 0.0;
 
 	/**
 	 * The color's alpha component.
 	 */
 	UPROPERTY()
-	double Alpha;
+	double Alpha = 0.0;
 };
 
 UCLASS()
