RC PropertyWatcher/FText crash

diff --git a/Engine/Plugins/VirtualProduction/RemoteControl/Source/RemoteControl/Private/RemoteControlPreset.cpp b/Engine/Plugins/VirtualProduction/RemoteControl/Source/RemoteControl/Private/RemoteControlPreset.cpp
--- a/Engine/Plugins/VirtualProduction/RemoteControl/Source/RemoteControl/Private/RemoteControlPreset.cpp
+++ b/Engine/Plugins/VirtualProduction/RemoteControl/Source/RemoteControl/Private/RemoteControlPreset.cpp
@@ -2320,6 +2320,7 @@ void URemoteControlPreset::FRCPropertyWatcher::SetLastFrameValue(const FRCFieldR
 	
 	const void* NewValueAddress = ResolvedData.Field->ContainerPtrToValuePtr<void>(ResolvedData.ContainerAddress);
 	LastFrameValue.SetNumUninitialized(ResolvedData.Field->GetSize());
+	ResolvedData.Field->InitializeValue(LastFrameValue.GetData());
 	ResolvedData.Field->CopyCompleteValue(LastFrameValue.GetData(), NewValueAddress);
 }
 
