RC protocol mapping undo/redo crash fix

diff --git a/Engine/Plugins/VirtualProduction/RemoteControl/Source/RemoteControl/Private/RemoteControlProtocolBinding.cpp b/Engine/Plugins/VirtualProduction/RemoteControl/Source/RemoteControl/Private/RemoteControlProtocolBinding.cpp
--- a/Engine/Plugins/VirtualProduction/RemoteControl/Source/RemoteControl/Private/RemoteControlProtocolBinding.cpp
+++ b/Engine/Plugins/VirtualProduction/RemoteControl/Source/RemoteControl/Private/RemoteControlProtocolBinding.cpp
@@ -423,8 +423,9 @@ namespace EntityInterpolation
 	}
 }
 
-FRemoteControlProtocolMapping::FRemoteControlProtocolMapping(FProperty* InProperty, uint8 InRangeValueSize)
-	: Id(FGuid::NewGuid())
+// Optionally supplying the MappingId is used by the undo system
+FRemoteControlProtocolMapping::FRemoteControlProtocolMapping(FProperty* InProperty, uint8 InRangeValueSize, const FGuid& InMappingId)
+	: Id(InMappingId)
 {
 	if (FBoolProperty* BoolProperty = CastField<FBoolProperty>(InProperty))
 	{
@@ -798,8 +799,9 @@ TArray<FRemoteControlProtocolEntity::FRangeMappingData> FRemoteControlProtocolEn
 	return RangeMappingBuffers;
 }
 
-FRemoteControlProtocolBinding::FRemoteControlProtocolBinding(const FName InProtocolName, const FGuid& InPropertyId, TSharedPtr<TStructOnScope<FRemoteControlProtocolEntity>> InRemoteControlProtocolEntityPtr)
-	: Id(FGuid::NewGuid())
+// Optionally supplying the BindingId is used by the undo system
+FRemoteControlProtocolBinding::FRemoteControlProtocolBinding(const FName InProtocolName, const FGuid& InPropertyId, TSharedPtr<TStructOnScope<FRemoteControlProtocolEntity>> InRemoteControlProtocolEntityPtr, const FGuid& InBindingId)
+	: Id(InBindingId)
 	, ProtocolName(InProtocolName)
 	, PropertyId(InPropertyId)
 	, RemoteControlProtocolEntityPtr(InRemoteControlProtocolEntityPtr)
diff --git a/Engine/Plugins/VirtualProduction/RemoteControl/Source/RemoteControl/Public/RemoteControlProtocolBinding.h b/Engine/Plugins/VirtualProduction/RemoteControl/Source/RemoteControl/Public/RemoteControlProtocolBinding.h
--- a/Engine/Plugins/VirtualProduction/RemoteControl/Source/RemoteControl/Public/RemoteControlProtocolBinding.h
+++ b/Engine/Plugins/VirtualProduction/RemoteControl/Source/RemoteControl/Public/RemoteControlProtocolBinding.h
@@ -54,7 +54,7 @@ struct REMOTECONTROL_API FRemoteControlProtocolMapping
 
 public:
 	FRemoteControlProtocolMapping() = default;
-	FRemoteControlProtocolMapping(FProperty* InProperty, uint8 InRangeValueSize);
+	FRemoteControlProtocolMapping(FProperty* InProperty, uint8 InRangeValueSize, const FGuid& InMappingId = FGuid::NewGuid());
 
 	bool operator==(const FRemoteControlProtocolMapping& InProtocolMapping) const;
 	bool operator==(FGuid InProtocolMappingId) const;
@@ -535,7 +535,7 @@ struct REMOTECONTROL_API FRemoteControlProtocolBinding
 
 public:
 	FRemoteControlProtocolBinding() = default;
-	FRemoteControlProtocolBinding(const FName InProtocolName, const FGuid& InPropertyId, TSharedPtr<TStructOnScope<FRemoteControlProtocolEntity>> InRemoteControlProtocolEntityPtr);
+	FRemoteControlProtocolBinding(const FName InProtocolName, const FGuid& InPropertyId, TSharedPtr<TStructOnScope<FRemoteControlProtocolEntity>> InRemoteControlProtocolEntityPtr, const FGuid& InBindingId = FGuid::NewGuid());
 
 	bool operator==(const FRemoteControlProtocolBinding& InProtocolBinding) const;
 	bool operator==(FGuid InProtocolBindingId) const;
diff --git a/Engine/Plugins/VirtualProduction/RemoteControl/Source/RemoteControlProtocolWidgets/Private/ViewModels/ProtocolBindingViewModel.cpp b/Engine/Plugins/VirtualProduction/RemoteControl/Source/RemoteControlProtocolWidgets/Private/ViewModels/ProtocolBindingViewModel.cpp
--- a/Engine/Plugins/VirtualProduction/RemoteControl/Source/RemoteControlProtocolWidgets/Private/ViewModels/ProtocolBindingViewModel.cpp
+++ b/Engine/Plugins/VirtualProduction/RemoteControl/Source/RemoteControlProtocolWidgets/Private/ViewModels/ProtocolBindingViewModel.cpp
@@ -22,16 +22,27 @@ namespace Commands
 	{
 		FGuid EntityId;
 		FGuid BindingId;
-		FGuid RangeId;	
-		TArray<FGuid> RangeIds; // Specify if there are multiple to remove
+		FGuid RangeId;
+		TArray<FGuid> RangeIds; // Specify if there are multiple to add or remove
 	};
+
+	static FGuid AddRangeMappingInternal(URemoteControlPreset* InPreset, const FAddRemoveRangeArgs& InArgs);
+	static bool RemoveRangeMappingInternal(URemoteControlPreset* InPreset, const FAddRemoveRangeArgs& InArgs);
 	
 	static FGuid AddRangeMappingInternal(URemoteControlPreset* InPreset, const FAddRemoveRangeArgs& InArgs)
 	{
-		if(TSharedPtr<FRemoteControlProperty> RCProperty = InPreset->GetExposedEntity<FRemoteControlProperty>(InArgs.EntityId).Pin())
+		TArray<FGuid> NewRangeIds;
+		NewRangeIds.Reserve(1);
+		
+		if(const TSharedPtr<FRemoteControlProperty> RCProperty = InPreset->GetExposedEntity<FRemoteControlProperty>(InArgs.EntityId).Pin())
 		{
 			FRemoteControlProtocolBinding* ProtocolBinding = RCProperty->ProtocolBindings.FindByHash(GetTypeHash(InArgs.BindingId), InArgs.BindingId);
-			check(ProtocolBinding);
+			
+			// If undone, this can be nullptr
+			if(!ProtocolBinding)
+			{
+				return FGuid();
+			}
 
 			const TSharedPtr<TStructOnScope<FRemoteControlProtocolEntity>> EntityPtr = ProtocolBinding->GetRemoteControlProtocolEntityPtr();
 
@@ -41,30 +52,63 @@ namespace Commands
 				// @note: This only adds the same number of range mappings back, excluding their parameters
 				for (int32 Idx = 0; Idx < InArgs.RangeIds.Num(); ++Idx)
 				{
-					const FRemoteControlProtocolMapping RangesData(RCProperty->GetProperty(), (*EntityPtr)->GetRangePropertySize());
+					FGuid RangeId = InArgs.RangeIds[Idx];
+					RangeId = RangeId.IsValid() ? RangeId : FGuid::NewGuid();
+					
+					const FRemoteControlProtocolMapping RangesData(RCProperty->GetProperty(), (*EntityPtr)->GetRangePropertySize(), RangeId);
 					ProtocolBinding->AddMapping(RangesData);
+					NewRangeIds.Add(RangesData.GetId());
 				}
-
-				return FGuid::NewGuid(); // Multiple added, but only 1 id returned, so just return a valid one
 			}
 			else
 			{
-			const FRemoteControlProtocolMapping RangesData(RCProperty->GetProperty(), (*EntityPtr)->GetRangePropertySize());
+				const FRemoteControlProtocolMapping RangesData(RCProperty->GetProperty(), (*EntityPtr)->GetRangePropertySize(), InArgs.RangeId.IsValid() ? InArgs.RangeId : FGuid::NewGuid());
 				ProtocolBinding->AddMapping(RangesData);
 
-			return RangesData.GetId();
-		}
+				NewRangeIds.Add(RangesData.GetId());
+			}
+
+			using FCommandChange = TRemoteControlProtocolCommandChange<Commands::FAddRemoveRangeArgs>;
+			using FOnChange = FCommandChange::FOnUndoRedoDelegate;
+			if(GUndo)
+			{
+				// Set or re-set the RangeId' (will be same if it was already set), so undos below remove or add the correct ranges
+				FAddRemoveRangeArgs Args = InArgs;
+				Args.RangeId = NewRangeIds[0];
+				Args.RangeIds = NewRangeIds;
+				
+				// What to do on undo and redo
+				GUndo->StoreUndo(InPreset,
+					MakeUnique<FCommandChange>(
+						InPreset,
+						MoveTemp(Args),
+						FOnChange::CreateLambda([](URemoteControlPreset* InPreset, const Commands::FAddRemoveRangeArgs& InChangeArgs)
+						{
+							Commands::AddRangeMappingInternal(InPreset, InChangeArgs);
+						}),
+						FOnChange::CreateLambda([](URemoteControlPreset* InPreset, const Commands::FAddRemoveRangeArgs& InChangeArgs)
+						{
+							Commands::RemoveRangeMappingInternal(InPreset, InChangeArgs);
+						})
+					)
+				);
+			}
 		}
-		
-		return FGuid();
+
+		return NewRangeIds[0]; // If multiple added, only return first as caller only needs a valid guid
 	}
 
 	static bool RemoveRangeMappingInternal(URemoteControlPreset* InPreset, const FAddRemoveRangeArgs& InArgs)
 	{
-		if(TSharedPtr<FRemoteControlProperty> RCProperty = InPreset->GetExposedEntity<FRemoteControlProperty>(InArgs.EntityId).Pin())
+		if(const TSharedPtr<FRemoteControlProperty> RCProperty = InPreset->GetExposedEntity<FRemoteControlProperty>(InArgs.EntityId).Pin())
 		{
 			FRemoteControlProtocolBinding* ProtocolBinding = RCProperty->ProtocolBindings.FindByHash(GetTypeHash(InArgs.BindingId), InArgs.BindingId);
-			check(ProtocolBinding);
+
+			// If undone, this can be nullptr
+			if(!ProtocolBinding)
+			{
+				return false;
+			}
 
 			int32 NumRemoved = 0;
 			if (InArgs.RangeId.IsValid())
@@ -77,6 +121,27 @@ namespace Commands
 				NumRemoved = 1; // ClearMapping doesn't return a num, so set to 1 to indicate success
 			}
 
+			using FCommandChange = TRemoteControlProtocolCommandChange<Commands::FAddRemoveRangeArgs>;
+			using FOnChange = FCommandChange::FOnUndoRedoDelegate;
+			if(GUndo)
+			{
+				// What to do on undo and redo
+				GUndo->StoreUndo(InPreset,
+					MakeUnique<FCommandChange>(
+						InPreset,
+						CopyTemp(InArgs),
+						FOnChange::CreateLambda([](URemoteControlPreset* InPreset, const Commands::FAddRemoveRangeArgs& InChangeArgs)
+						{
+							Commands::RemoveRangeMappingInternal(InPreset, InChangeArgs);
+						}),
+						FOnChange::CreateLambda([](URemoteControlPreset* InPreset, const Commands::FAddRemoveRangeArgs& InChangeArgs)
+						{
+							Commands::AddRangeMappingInternal(InPreset, InChangeArgs);
+						})
+					)
+				);
+			}
+
 			return NumRemoved > 0;
 		}
 		
@@ -120,7 +185,7 @@ FProtocolBindingViewModel::~FProtocolBindingViewModel()
 
 void FProtocolBindingViewModel::Initialize()
 {
-	FRemoteControlProtocolBinding* Binding = GetBinding();
+	const FRemoteControlProtocolBinding* Binding = GetBinding();
 	// May be stale as a result of an undo deleting it.
 	if(Binding)
 	{
@@ -146,21 +211,6 @@ FGuid FProtocolBindingViewModel::AddRangeMapping()
 
 	TSharedPtr<FProtocolRangeViewModel>& NewRangeViewModel = AddRangeMappingInternal();
 
-	using FCommandChange = TRemoteControlProtocolCommandChange<Commands::FAddRemoveRangeArgs>;
-	using FOnChange = FCommandChange::FOnUndoRedoDelegate;
-	if(GUndo)
-	{
-		// What to do on undo and redo
-		GUndo->StoreUndo(Preset.Get(),
-            MakeUnique<FCommandChange>(
-                Preset.Get(),
-                Commands::FAddRemoveRangeArgs{ParentViewModel.Pin()->GetId(), GetId(), NewRangeViewModel->GetId()},
-                FOnChange::CreateLambda([](URemoteControlPreset* InPreset, const Commands::FAddRemoveRangeArgs& InChangeArgs){ Commands::AddRangeMappingInternal(InPreset, InChangeArgs); }),
-                FOnChange::CreateLambda([](URemoteControlPreset* InPreset, const Commands::FAddRemoveRangeArgs& InChangeArgs){ Commands::RemoveRangeMappingInternal(InPreset, InChangeArgs); })
-            )
-        );
-	}
-
 	OnRangeMappingAddedDelegate.Broadcast(NewRangeViewModel.ToSharedRef());	
 
 	return NewRangeViewModel->GetId();
@@ -189,21 +239,6 @@ void FProtocolBindingViewModel::RemoveRangeMapping(const FGuid& InId)
 	check(GUndo != nullptr);
 	GetPreset()->MarkPackageDirty();
 
-	using FCommandChange = TRemoteControlProtocolCommandChange<Commands::FAddRemoveRangeArgs>;
-	using FOnChange = FCommandChange::FOnUndoRedoDelegate;
-	if(GUndo)
-	{
-		// What to do on undo and redo
-		GUndo->StoreUndo(Preset.Get(),
-            MakeUnique<FCommandChange>(
-                Preset.Get(),
-                Commands::FAddRemoveRangeArgs{ParentViewModel.Pin()->GetId(), GetId(), InId},
-                FOnChange::CreateLambda([](URemoteControlPreset* InPreset, const Commands::FAddRemoveRangeArgs& InChangeArgs){ Commands::RemoveRangeMappingInternal(InPreset, InChangeArgs); }),
-                FOnChange::CreateLambda([](URemoteControlPreset* InPreset, const Commands::FAddRemoveRangeArgs& InChangeArgs){ Commands::AddRangeMappingInternal(InPreset, InChangeArgs); })
-            )
-        );
-	}
-
 	if(Commands::RemoveRangeMappingInternal(Preset.Get(), Commands::FAddRemoveRangeArgs{ParentViewModel.Pin()->GetId(), GetId(), InId}))
 	{
 		const int32 NumItemsRemoved = Ranges.RemoveAll([&](const TSharedPtr<FProtocolRangeViewModel>& InRange)
@@ -212,7 +247,9 @@ void FProtocolBindingViewModel::RemoveRangeMapping(const FGuid& InId)
         });
 
 		if(NumItemsRemoved <= 0)
+		{
 			return;
+		}
 
 		OnRangeMappingRemovedDelegate.Broadcast(InId);
 	}
@@ -239,28 +276,15 @@ void FProtocolBindingViewModel::RemoveAllRangeMappings()
 		return InRange->GetId();
 	});
 
-	using FCommandChange = TRemoteControlProtocolCommandChange<Commands::FAddRemoveRangeArgs>;
-	using FOnChange = FCommandChange::FOnUndoRedoDelegate;
-	if (GUndo)
-	{
-		// What to do on undo and redo
-		GUndo->StoreUndo(Preset.Get(),
-						MakeUnique<FCommandChange>(
-							Preset.Get(),
-							Commands::FAddRemoveRangeArgs{ParentViewModel.Pin()->GetId(), GetId(), {}, AllRangeIds},
-							FOnChange::CreateLambda([](URemoteControlPreset* InPreset, const Commands::FAddRemoveRangeArgs& InChangeArgs) { Commands::RemoveRangeMappingInternal(InPreset, InChangeArgs); }),
-							FOnChange::CreateLambda([](URemoteControlPreset* InPreset, const Commands::FAddRemoveRangeArgs& InChangeArgs) { Commands::AddRangeMappingInternal(InPreset, InChangeArgs); })
-						)
-		);
-	}
-
 	if (Commands::RemoveRangeMappingInternal(Preset.Get(), Commands::FAddRemoveRangeArgs{ParentViewModel.Pin()->GetId(), GetId(), {}, AllRangeIds}))
 	{
 		const int32 NumItemsRemoved = Ranges.Num();
 		Ranges.Empty();
 
 		if (NumItemsRemoved <= 0)
+		{
 			return;
+		}
 
 		OnRangeMappingsRemovedDelegate.Broadcast();
 	}
@@ -270,8 +294,8 @@ void FProtocolBindingViewModel::AddDefaultRangeMappings()
 {
 	check(IsValid());
 
-	TSharedPtr<FProtocolRangeViewModel> MinItem = AddRangeMappingInternal();
-	TSharedPtr<FProtocolRangeViewModel> MaxItem = AddRangeMappingInternal();
+	const TSharedPtr<FProtocolRangeViewModel> MinItem = AddRangeMappingInternal();
+	const TSharedPtr<FProtocolRangeViewModel> MaxItem = AddRangeMappingInternal();
 
 	FName RangePropertyTypeName = GetBinding()->GetRemoteControlProtocolEntityPtr()->Get()->GetRangePropertyName();
 	FProperty* RangeProperty = GetProtocol()->GetRangeInputTemplateProperty();
@@ -350,6 +374,7 @@ void FProtocolBindingViewModel::AddDefaultRangeMappings()
 	{
 		MappingPropertyTypeName = StructProperty->Struct->GetFName();
 	}
+
 	// Mapping (output)
 	{
 		if(const FNumericProperty* NumericProperty = CastField<FNumericProperty>(MappingProperty))
@@ -512,7 +537,7 @@ FRemoteControlProtocolBinding* FProtocolBindingViewModel::GetBinding() const
 {
 	check(IsValid());
 
-	if (TSharedPtr<FRemoteControlProperty> RCProperty = Preset->GetExposedEntity<FRemoteControlProperty>(PropertyId).Pin())
+	if (const TSharedPtr<FRemoteControlProperty> RCProperty = Preset->GetExposedEntity<FRemoteControlProperty>(PropertyId).Pin())
 	{
 		return RCProperty->ProtocolBindings.FindByHash(GetTypeHash(BindingId), BindingId);
 	}
@@ -523,6 +548,12 @@ FRemoteControlProtocolBinding* FProtocolBindingViewModel::GetBinding() const
 
 TSharedPtr<IRemoteControlProtocol> FProtocolBindingViewModel::GetProtocol() const
 {
+	// If undo removes the binding, this will occur
+	if(!GetBinding())
+	{
+		return nullptr;		
+	}
+	
 	return IRemoteControlProtocolModule::Get().GetProtocolByName(GetBinding()->GetProtocolName());
 }
 
@@ -553,7 +584,7 @@ bool FProtocolBindingViewModel::IsValid(FText& OutMessage)
 	check(IsValid());
 
 	EValidity Result = EValidity::Ok;
-	FString* AdditionalMessage = nullptr;
+	const FString* AdditionalMessage = nullptr;
 
 	if (Ranges.Num() <= 1)
 	{
@@ -575,7 +606,7 @@ bool FProtocolBindingViewModel::IsValid(FText& OutMessage)
 	// Continue if ok
 	if (Result == EValidity::Ok)
 	{
-		FRemoteControlProtocolBinding* ProtocolBinding = GetBinding();
+		const FRemoteControlProtocolBinding* ProtocolBinding = GetBinding();
 		// Was probably deleted
 		if(ProtocolBinding == nullptr)
 		{
@@ -589,7 +620,7 @@ bool FProtocolBindingViewModel::IsValid(FText& OutMessage)
 		// Iterate over each exposed field
 		for (TWeakPtr<FRemoteControlField> ExposedFieldWeakPtr : GetPreset()->GetExposedEntities<FRemoteControlField>())
 		{
-			TSharedPtr<FRemoteControlField> ExposedField = ExposedFieldWeakPtr.Pin();
+			const TSharedPtr<FRemoteControlField> ExposedField = ExposedFieldWeakPtr.Pin();
 
 			// Skip self
 			if (ExposedField->GetId() == PropertyId)
diff --git a/Engine/Plugins/VirtualProduction/RemoteControl/Source/RemoteControlProtocolWidgets/Private/ViewModels/ProtocolEntityViewModel.cpp b/Engine/Plugins/VirtualProduction/RemoteControl/Source/RemoteControlProtocolWidgets/Private/ViewModels/ProtocolEntityViewModel.cpp
--- a/Engine/Plugins/VirtualProduction/RemoteControl/Source/RemoteControlProtocolWidgets/Private/ViewModels/ProtocolEntityViewModel.cpp
+++ b/Engine/Plugins/VirtualProduction/RemoteControl/Source/RemoteControlProtocolWidgets/Private/ViewModels/ProtocolEntityViewModel.cpp
@@ -25,20 +25,49 @@ namespace Commands
 		FGuid BindingId;
 	};
 
+	static TSharedPtr<FRemoteControlProtocolBinding> AddProtocolInternal(URemoteControlPreset* InPreset, const FAddRemoveProtocolArgs& InArgs);
+	static FName RemoveProtocolInternal(URemoteControlPreset* InPreset, const FAddRemoveProtocolArgs& InArgs);
+
 	static TSharedPtr<FRemoteControlProtocolBinding> AddProtocolInternal(URemoteControlPreset* InPreset, const FAddRemoveProtocolArgs& InArgs)
 	{
 		const TSharedPtr<IRemoteControlProtocol> Protocol = IRemoteControlProtocolModule::Get().GetProtocolByName(InArgs.ProtocolName);
 		if (Protocol.IsValid())
 		{
-			if (TSharedPtr<FRemoteControlProperty> RCProperty = InPreset->GetExposedEntity<FRemoteControlProperty>(InArgs.EntityId).Pin())
+			if (const TSharedPtr<FRemoteControlProperty> RCProperty = InPreset->GetExposedEntity<FRemoteControlProperty>(InArgs.EntityId).Pin())
 			{
 				const FRemoteControlProtocolEntityPtr RemoteControlProtocolEntityPtr = Protocol->CreateNewProtocolEntity(RCProperty->GetProperty(), InPreset, InArgs.EntityId);
 
-				FRemoteControlProtocolBinding ProtocolBinding(InArgs.ProtocolName, InArgs.EntityId, RemoteControlProtocolEntityPtr);
+				FRemoteControlProtocolBinding ProtocolBinding(InArgs.ProtocolName, InArgs.EntityId, RemoteControlProtocolEntityPtr, InArgs.BindingId.IsValid() ? InArgs.BindingId : FGuid::NewGuid());
 				Protocol->Bind(ProtocolBinding.GetRemoteControlProtocolEntityPtr());
 				RCProperty->ProtocolBindings.Emplace(MoveTemp(ProtocolBinding));
 
-				return MakeShared<FRemoteControlProtocolBinding>(ProtocolBinding);
+				TSharedPtr<FRemoteControlProtocolBinding> NewBinding = MakeShared<FRemoteControlProtocolBinding>(ProtocolBinding);
+
+				using FCommandChange = TRemoteControlProtocolCommandChange<Commands::FAddRemoveProtocolArgs>;
+				using FOnChange = FCommandChange::FOnUndoRedoDelegate;
+				if (GUndo)
+				{
+					FAddRemoveProtocolArgs Args = InArgs;
+					// Set or re-set the BindingId (will be same if it was already set), so undos below remove or add the correct binding
+					Args.BindingId = NewBinding->GetId();
+					
+					GUndo->StoreUndo(InPreset,
+						MakeUnique<FCommandChange>(
+							InPreset,
+							MoveTemp(Args),
+							FOnChange::CreateLambda([](URemoteControlPreset* InPreset, const Commands::FAddRemoveProtocolArgs& InChangeArgs)
+							{
+								Commands::AddProtocolInternal(InPreset, InChangeArgs);
+							}),
+							FOnChange::CreateLambda([](URemoteControlPreset* InPreset, const Commands::FAddRemoveProtocolArgs& InChangeArgs)
+							{
+								Commands::RemoveProtocolInternal(InPreset, InChangeArgs);
+							})
+						)
+					);
+				}
+
+				return NewBinding;
 			}
 		}
 
@@ -47,13 +76,41 @@ namespace Commands
 
 	static FName RemoveProtocolInternal(URemoteControlPreset* InPreset, const FAddRemoveProtocolArgs& InArgs)
 	{
-		if (TSharedPtr<FRemoteControlProperty> RCProperty = InPreset->GetExposedEntity<FRemoteControlProperty>(InArgs.EntityId).Pin())
+		if (const TSharedPtr<FRemoteControlProperty> RCProperty = InPreset->GetExposedEntity<FRemoteControlProperty>(InArgs.EntityId).Pin())
 		{
-			FRemoteControlProtocolBinding* FoundElement = RCProperty->ProtocolBindings.FindByHash(GetTypeHash(InArgs.BindingId), InArgs.BindingId);
+			const FRemoteControlProtocolBinding* FoundElement = RCProperty->ProtocolBindings.FindByHash(GetTypeHash(InArgs.BindingId), InArgs.BindingId);
 			if (FoundElement)
 			{
 				RCProperty->ProtocolBindings.RemoveByHash(GetTypeHash(InArgs.BindingId), InArgs.BindingId);
-				return FoundElement->GetProtocolName();
+
+				const FName RemovedProtocolName = FoundElement->GetProtocolName();
+				if (RemovedProtocolName != NAME_None)
+				{
+					using FCommandChange = TRemoteControlProtocolCommandChange<Commands::FAddRemoveProtocolArgs>;
+					using FOnChange = FCommandChange::FOnUndoRedoDelegate;
+					if (GUndo)
+					{
+						FAddRemoveProtocolArgs Args = InArgs;
+						Args.ProtocolName = RemovedProtocolName;
+						
+						GUndo->StoreUndo(InPreset,
+							MakeUnique<FCommandChange>(
+								InPreset,
+								MoveTemp(Args),
+								FOnChange::CreateLambda([](URemoteControlPreset* InPreset, const Commands::FAddRemoveProtocolArgs& InChangeArgs)
+								{
+									Commands::RemoveProtocolInternal(InPreset, InChangeArgs);
+								}),
+								FOnChange::CreateLambda([](URemoteControlPreset* InPreset, const Commands::FAddRemoveProtocolArgs& InChangeArgs)
+								{
+									Commands::AddProtocolInternal(InPreset, InChangeArgs);
+								})
+							)
+						);
+					}
+				}
+				
+				return RemovedProtocolName;
 			}
 		}
 
@@ -99,7 +156,7 @@ void FProtocolEntityViewModel::Initialize()
 	Preset->OnEntityUnexposed().AddSP(this, &FProtocolEntityViewModel::OnEntityUnexposed);
 
 	Bindings.Empty(Bindings.Num());
-	if (TSharedPtr<FRemoteControlProperty> RCProperty = Preset->GetExposedEntity<FRemoteControlProperty>(PropertyId).Pin())
+	if (const TSharedPtr<FRemoteControlProperty> RCProperty = Preset->GetExposedEntity<FRemoteControlProperty>(PropertyId).Pin())
 	{
 		Property = RCProperty->GetProperty();
 		for(FRemoteControlProtocolBinding& Binding : RCProperty->ProtocolBindings)
@@ -108,7 +165,13 @@ void FProtocolEntityViewModel::Initialize()
 			// Supporting plugin needs to be loaded/protocol available
 			if (Protocol.IsValid())
 			{
-				Bindings.Emplace(FProtocolBindingViewModel::Create(SharedThis(this), MakeShared<FRemoteControlProtocolBinding>(Binding)));
+				const TSharedPtr<FProtocolBindingViewModel>& BindingViewModel = Bindings.Add_GetRef(FProtocolBindingViewModel::Create(SharedThis(this), MakeShared<FRemoteControlProtocolBinding>(Binding)));
+				// Probably result of "redo" action
+				if(BindingViewModel->GetRanges().IsEmpty())
+				{
+					BindingViewModel->AddDefaultRangeMappings();
+				}
+
 				Protocol->Bind(Binding.GetRemoteControlProtocolEntityPtr());
 			}
 		}
@@ -157,23 +220,10 @@ TSharedPtr<FProtocolBindingViewModel> FProtocolEntityViewModel::AddBinding(const
 	FScopedTransaction Transaction(LOCTEXT("AddProtocolBinding", "Add Protocol Binding"));
 	Preset->MarkPackageDirty();
 
+	const FGuid CommandId = FGuid::NewGuid();
 	const TSharedPtr<FRemoteControlProtocolBinding> ProtocolBinding = Commands::AddProtocolInternal(Preset.Get(), Commands::FAddRemoveProtocolArgs{GetId(), InProtocolName});
 	if (ProtocolBinding.IsValid())
 	{
-		using FCommandChange = TRemoteControlProtocolCommandChange<Commands::FAddRemoveProtocolArgs>;
-		using FOnChange = FCommandChange::FOnUndoRedoDelegate;
-		if (GUndo)
-		{
-			GUndo->StoreUndo(Preset.Get(),
-                MakeUnique<FCommandChange>(
-                    Preset.Get(),
-                    Commands::FAddRemoveProtocolArgs{GetId(), InProtocolName, ProtocolBinding->GetId()},
-                    FOnChange::CreateLambda([](URemoteControlPreset* InPreset, const Commands::FAddRemoveProtocolArgs& InChangeArgs){ Commands::AddProtocolInternal(InPreset, InChangeArgs); }),
-                    FOnChange::CreateLambda([](URemoteControlPreset* InPreset, const Commands::FAddRemoveProtocolArgs& InChangeArgs){ Commands::RemoveProtocolInternal(InPreset, InChangeArgs); })
-                )
-            );
-		}
-
 		TSharedPtr<FProtocolBindingViewModel>& NewBindingViewModel = Bindings.Add_GetRef(FProtocolBindingViewModel::Create(AsShared(), ProtocolBinding.ToSharedRef()));
 		NewBindingViewModel->AddDefaultRangeMappings();
 		
@@ -192,23 +242,7 @@ void FProtocolEntityViewModel::RemoveBinding(const FGuid& InBindingId)
 	FScopedTransaction Transaction(LOCTEXT("RemoveProtocolBinding", "Remove Protocol Binding"));
 	Preset->MarkPackageDirty();
 
-	const FName RemovedProtocolName = Commands::RemoveProtocolInternal(Preset.Get(), Commands::FAddRemoveProtocolArgs{GetId(), NAME_None, InBindingId});
-	if (RemovedProtocolName != NAME_None)
-	{
-		using FCommandChange = TRemoteControlProtocolCommandChange<Commands::FAddRemoveProtocolArgs>;
-		using FOnChange = FCommandChange::FOnUndoRedoDelegate;
-		if (GUndo)
-		{
-			GUndo->StoreUndo(Preset.Get(),
-                MakeUnique<FCommandChange>(
-                    Preset.Get(),
-                    Commands::FAddRemoveProtocolArgs{GetId(), RemovedProtocolName, InBindingId},
-                    FOnChange::CreateLambda([](URemoteControlPreset* InPreset, const Commands::FAddRemoveProtocolArgs& InChangeArgs){ Commands::RemoveProtocolInternal(InPreset, InChangeArgs); }),
-                    FOnChange::CreateLambda([](URemoteControlPreset* InPreset, const Commands::FAddRemoveProtocolArgs& InChangeArgs){ Commands::AddProtocolInternal(InPreset, InChangeArgs); })
-                )
-            );
-		}
-	}
+	Commands::RemoveProtocolInternal(Preset.Get(), Commands::FAddRemoveProtocolArgs{GetId(), NAME_None, InBindingId});
 
 	const int32 NumItemsRemoved = Bindings.RemoveAll([&](const TSharedPtr<FProtocolBindingViewModel>& InBinding)
     {
diff --git a/Engine/Plugins/VirtualProduction/RemoteControl/Source/RemoteControlProtocolWidgets/Private/ViewModels/ProtocolRangeViewModel.cpp b/Engine/Plugins/VirtualProduction/RemoteControl/Source/RemoteControlProtocolWidgets/Private/ViewModels/ProtocolRangeViewModel.cpp
--- a/Engine/Plugins/VirtualProduction/RemoteControl/Source/RemoteControlProtocolWidgets/Private/ViewModels/ProtocolRangeViewModel.cpp
+++ b/Engine/Plugins/VirtualProduction/RemoteControl/Source/RemoteControlProtocolWidgets/Private/ViewModels/ProtocolRangeViewModel.cpp
@@ -73,7 +73,7 @@ void FProtocolRangeViewModel::UpdateInputValueRange()
 	check(IsValid());
 
 	// Early out if not yet initialized or stale
-	if(!InputProxyPropertyContainer.IsValid())
+	if(!InputProxyPropertyContainer.IsValid() || !ParentViewModel.IsValid())
 	{
 		return;	
 	}
@@ -208,6 +208,8 @@ void FProtocolRangeViewModel::SetOutputData(const FRCFieldResolvedData& InResolv
 
 FProperty* FProtocolRangeViewModel::GetInputProperty() const
 {
+	check(ParentViewModel.IsValid());
+	
 	return ParentViewModel.Pin()->GetProtocol()->GetRangeInputTemplateProperty();
 }
 
@@ -313,6 +315,13 @@ void FProtocolRangeViewModel::PostUndo(bool bSuccess)
 void FProtocolRangeViewModel::OnParentChanged()
 {
 	check(IsValid());
+	
+	// Undo deleted the parent binding
+	if(!GetBinding())
+	{
+		return;
+	}
+	
 	UpdateInputValueRange();
 	OnChangedDelegate.Broadcast();
 }
diff --git a/Engine/Plugins/VirtualProduction/RemoteControl/Source/RemoteControlProtocolWidgets/Private/Widgets/SRCProtocolRange.cpp b/Engine/Plugins/VirtualProduction/RemoteControl/Source/RemoteControlProtocolWidgets/Private/Widgets/SRCProtocolRange.cpp
--- a/Engine/Plugins/VirtualProduction/RemoteControl/Source/RemoteControlProtocolWidgets/Private/Widgets/SRCProtocolRange.cpp
+++ b/Engine/Plugins/VirtualProduction/RemoteControl/Source/RemoteControlProtocolWidgets/Private/Widgets/SRCProtocolRange.cpp
@@ -23,15 +23,8 @@ void SRCProtocolRange::Construct(const FArguments& InArgs, const TSharedRef<STab
 	PrimaryColumnSizeData = InArgs._PrimaryColumnSizeData;
 	SecondaryColumnSizeData = InArgs._SecondaryColumnSizeData;
 
-	ViewModel->OnChanged().AddLambda([this]()
-	{
-		// May not have been created yet
-		if(InputPropertyView.IsValid())
-		{
-			InputPropertyView->Refresh();	
-		}
-	});
-
+	ViewModel->OnChanged().AddSP(this, &SRCProtocolRange::OnViewModelChanged);
+	
 	const TSharedPtr<SWidget> LeftWidget =
 	SNew(SHorizontalBox)
 	+ SHorizontalBox::Slot()
@@ -270,4 +263,13 @@ void SRCProtocolRange::OnOutputProxyChanged(const FPropertyChangedEvent& InEvent
 	OutputPropertyView->Refresh();
 }
 
+void SRCProtocolRange::OnViewModelChanged() const
+{
+	// May not have been created yet
+	if(InputPropertyView.IsValid())
+	{
+		InputPropertyView->Refresh();	
+	}
+}
+
 #undef LOCTEXT_NAMESPACE
diff --git a/Engine/Plugins/VirtualProduction/RemoteControl/Source/RemoteControlProtocolWidgets/Private/Widgets/SRCProtocolRange.h b/Engine/Plugins/VirtualProduction/RemoteControl/Source/RemoteControlProtocolWidgets/Private/Widgets/SRCProtocolRange.h
--- a/Engine/Plugins/VirtualProduction/RemoteControl/Source/RemoteControlProtocolWidgets/Private/Widgets/SRCProtocolRange.h
+++ b/Engine/Plugins/VirtualProduction/RemoteControl/Source/RemoteControlProtocolWidgets/Private/Widgets/SRCProtocolRange.h
@@ -55,7 +55,10 @@ private:
 
 	/** Applies proxy data to actual. */
 	void OnOutputProxyChanged(const FPropertyChangedEvent& InEvent);
-
+	
+	/** Called whenever the viewmodel changes. */
+	void OnViewModelChanged() const;
+	
 protected:
 	/** Represents an individual range binding for a protocol binding. */
 	TSharedPtr<FProtocolRangeViewModel> ViewModel;
