Remote Control Logic.

Runtime part is ready for review

diff --git a/Engine/Plugins/VirtualProduction/RemoteControl/RemoteControl.uplugin b/Engine/Plugins/VirtualProduction/RemoteControl/RemoteControl.uplugin
--- a/Engine/Plugins/VirtualProduction/RemoteControl/RemoteControl.uplugin
+++ b/Engine/Plugins/VirtualProduction/RemoteControl/RemoteControl.uplugin
@@ -14,13 +14,17 @@
 	"IsBetaVersion": true,
 	"Installed": false,
 	"Plugins": [
-		{
+		{ 
 			"Name": "WebSocketNetworking",
 			"Enabled": true
 		},
 		{
 			"Name": "RemoteControlInterception",
 			"Enabled": true
+		},
+		{
+			"Name": "StructUtils",
+			"Enabled": true
 		}
 	],
 	"Modules": [
@@ -29,6 +33,11 @@
 			"Type": "Runtime",
 			"LoadingPhase": "Default"
 		},
+		{
+			"Name": "RemoteControlLogic",
+			"Type": "Runtime",
+			"LoadingPhase": "Default"
+		},
 		{
 			"Name": "WebRemoteControl",
 			"Type": "Runtime",
diff --git a/Engine/Plugins/VirtualProduction/RemoteControl/Source/RemoteControl/Private/RCVirtualProperty.cpp b/Engine/Plugins/VirtualProduction/RemoteControl/Source/RemoteControl/Private/RCVirtualProperty.cpp
new file mode 100644
--- /dev/null
+++ b/Engine/Plugins/VirtualProduction/RemoteControl/Source/RemoteControl/Private/RCVirtualProperty.cpp
@@ -0,0 +1,543 @@
+﻿// Copyright Epic Games, Inc. All Rights Reserved.
+
+#include "RCVirtualProperty.h"
+
+#include "StructSerializer.h"
+#include "UObject/TextProperty.h"
+#include "UObject/StructOnScope.h"
+
+#include "RCVirtualPropertyContainer.h"
+
+const FProperty* URCVirtualPropertyBase::GetProperty() const
+{
+	if (const FPropertyBagPropertyDesc* PropertyBagPropertyDesc = GetBagPropertyDesc())
+	{
+		return PropertyBagPropertyDesc->CachedProperty;
+	}
+
+	return nullptr;
+}
+
+FProperty* URCVirtualPropertyBase::GetProperty()
+{
+	if (const FPropertyBagPropertyDesc* PropertyBagPropertyDesc = GetBagPropertyDesc())
+	{
+		return const_cast<FProperty*>(PropertyBagPropertyDesc->CachedProperty);
+	}
+
+	return nullptr;
+}
+
+const uint8* URCVirtualPropertyBase::GetContainerPtr() const
+{
+	if (const FInstancedPropertyBag* InstancedPropertyBag = GetPropertyBagInstance())
+	{
+		return InstancedPropertyBag->GetMutableValue().GetMemory();
+	}
+	
+	return GetContainerPtr();
+}
+
+uint8* URCVirtualPropertyBase::GetContainerPtr()
+{
+	if (const FInstancedPropertyBag* InstancedPropertyBag = GetPropertyBagInstance())
+	{
+		return InstancedPropertyBag->GetMutableValue().GetMutableMemory();
+	}
+	
+	return nullptr;
+}
+
+const uint8* URCVirtualPropertyBase::GetValuePtr() const
+{
+	if (const uint8* ContainerPtr = GetContainerPtr())
+	{
+		if (const FProperty* Property = GetProperty())
+		{
+			return Property->ContainerPtrToValuePtr<uint8>(ContainerPtr);
+		}
+	}
+
+	return nullptr;
+}
+
+uint8* URCVirtualPropertyBase::GetValuePtr()
+{
+	if (uint8* ContainerPtr = GetContainerPtr())
+	{
+		if (FProperty* Property = GetProperty())
+		{
+			return Property->ContainerPtrToValuePtr<uint8>(ContainerPtr);
+		}
+	}
+
+	return nullptr;
+}
+
+const FPropertyBagPropertyDesc* URCVirtualPropertyBase::GetBagPropertyDesc() const
+{
+	if (const FInstancedPropertyBag* InstancedPropertyBag = GetPropertyBagInstance())
+	{
+		return InstancedPropertyBag->FindPropertyDescByName(PropertyName);
+	}
+	
+	return nullptr;
+}
+
+const FInstancedPropertyBag* URCVirtualPropertyBase::GetPropertyBagInstance() const
+{
+	// That is should be implemented in child classes
+	check(nullptr);
+	return nullptr;
+}
+
+EPropertyBagPropertyType URCVirtualPropertyBase::GetValueType() const
+{
+	if (const FPropertyBagPropertyDesc* PropertyBagPropertyDesc = GetBagPropertyDesc())
+	{
+		return PropertyBagPropertyDesc->ValueType;
+	}
+
+	return EPropertyBagPropertyType::None;
+}
+
+const UObject* URCVirtualPropertyBase::GetValueTypeObjectWeakPtr() const
+{
+	if (const FPropertyBagPropertyDesc* PropertyBagPropertyDesc = GetBagPropertyDesc())
+	{
+		return PropertyBagPropertyDesc->ValueTypeObject;
+	}
+
+	return nullptr;
+}
+
+void URCVirtualPropertyBase::SerializeToBackend(IStructSerializerBackend& OutBackend)
+{
+	FStructSerializerPolicies Policies; 
+	Policies.MapSerialization = EStructSerializerMapPolicies::Array;
+	FStructSerializer::SerializeElement(GetContainerPtr(), GetProperty(), INDEX_NONE, OutBackend, Policies);
+}
+
+bool URCVirtualPropertyBase::IsValueEqual(URCVirtualPropertyBase* InVirtualProperty)
+{
+	const FProperty* ThisProperty = GetProperty();
+	const uint8* ThisContainerPtr = GetContainerPtr();
+	if (ThisProperty == nullptr || ThisContainerPtr == nullptr)
+	{
+		return false;
+	}
+	
+	const FProperty* CompareProperty = InVirtualProperty->GetProperty();;
+	const uint8* CompareContainerPtr = InVirtualProperty->GetContainerPtr();
+	if (CompareProperty == nullptr || CompareContainerPtr == nullptr)
+	{
+		return false;
+	}
+
+	if (ThisProperty->GetClass() != CompareProperty->GetClass())
+	{
+		return false;
+	}
+
+	const uint8* ThisValuePtr = ThisProperty->ContainerPtrToValuePtr<uint8>(ThisContainerPtr);
+	const uint8* CompareValuePtr = CompareProperty->ContainerPtrToValuePtr<uint8>(CompareContainerPtr);
+
+	return ThisProperty->Identical(ThisValuePtr, CompareValuePtr);
+}
+
+struct FRCVirtualPropertyCastHelpers
+{
+	template<typename TPropertyType, typename TValueType>
+	static bool GetPrimitiveValue(URCVirtualPropertyBase* InVirtualProperty, TValueType& OutValue)
+	{
+		const uint8* ContainerPtr = InVirtualProperty->GetContainerPtr();
+		if (!ContainerPtr)
+		{
+			return false;
+		}
+		
+		const FProperty* Property = InVirtualProperty->GetProperty();
+		
+		if (const TPropertyType* CastProperty = CastField<TPropertyType>(Property))
+		{
+			OutValue = CastProperty->GetPropertyValue_InContainer(ContainerPtr);
+			return true;
+		}
+
+		return false;
+	}
+
+	template<typename TPropertyType, typename TValueType>
+	static bool SetPrimitiveValue(URCVirtualPropertyBase* InVirtualProperty, const TValueType& InValue)
+	{
+		uint8* ValuePtr = InVirtualProperty->GetValuePtr();
+		if (!ValuePtr)
+		{
+			return false;
+		}
+		
+		const FProperty* Property = InVirtualProperty->GetProperty();
+		
+		if (const TPropertyType* CastProperty = CastField<TPropertyType>(Property))
+		{
+			CastProperty->SetPropertyValue(ValuePtr, InValue);
+			return true;
+		}
+
+		return false;
+	}
+
+	template<typename TValueType>
+	static bool GetStructValue(URCVirtualPropertyBase* InVirtualProperty, UScriptStruct* InScriptStruct, TValueType* OutValuePtr)
+	{
+		const uint8* ValuePtr = InVirtualProperty->GetValuePtr();
+		if (!ValuePtr)
+		{
+			return false;
+		}
+
+		const FProperty* Property = InVirtualProperty->GetProperty();
+		if (const FStructProperty* StructProperty = CastField<FStructProperty>(Property))
+		{
+			if (InScriptStruct == StructProperty->Struct)
+			{
+				StructProperty->Struct->CopyScriptStruct(OutValuePtr, ValuePtr);
+				return true;
+			}
+		}
+
+		return false;
+	}
+
+	template<typename TValueType>
+	static bool SetStructValue(URCVirtualPropertyBase* InVirtualProperty, UScriptStruct* InScriptStruct, const TValueType& InValue)
+	{
+		uint8* ValuePtr = InVirtualProperty->GetValuePtr();
+		if (!ValuePtr)
+		{
+			return false;
+		}
+
+		const FProperty* Property = InVirtualProperty->GetProperty();
+		if (const FStructProperty* StructProperty = CastField<FStructProperty>(Property))
+		{
+			if (InScriptStruct == StructProperty->Struct)
+			{
+				StructProperty->Struct->CopyScriptStruct(ValuePtr, &InValue);
+				return true;
+			}
+		}
+
+		return false;
+	}
+};
+
+bool URCVirtualPropertyBase::GetValueBool(bool& OutBoolValue)
+{
+	return FRCVirtualPropertyCastHelpers::GetPrimitiveValue<FBoolProperty, bool>(this, OutBoolValue);
+}
+
+bool URCVirtualPropertyBase::GetValueInt8(int8& OutInt8)
+{
+	return FRCVirtualPropertyCastHelpers::GetPrimitiveValue<FInt8Property, int8>(this, OutInt8);
+}
+
+bool URCVirtualPropertyBase::GetValueByte(uint8& OutByte)
+{
+	return FRCVirtualPropertyCastHelpers::GetPrimitiveValue<FByteProperty, uint8&>(this, OutByte);
+}
+
+bool URCVirtualPropertyBase::GetValueInt16(int16& OutInt16)
+{
+	return FRCVirtualPropertyCastHelpers::GetPrimitiveValue<FInt16Property, int16&>(this, OutInt16);
+}
+
+bool URCVirtualPropertyBase::GetValueUint16(uint16& OutUInt16)
+{
+	return FRCVirtualPropertyCastHelpers::GetPrimitiveValue<FUInt16Property, uint16&>(this, OutUInt16);
+}
+
+bool URCVirtualPropertyBase::GetValueInt32(int32& OutInt32)
+{
+	return FRCVirtualPropertyCastHelpers::GetPrimitiveValue<FIntProperty, int32&>(this, OutInt32);
+}
+
+bool URCVirtualPropertyBase::GetValueUInt32(uint32& OutUInt32)
+{
+	return FRCVirtualPropertyCastHelpers::GetPrimitiveValue<FUInt32Property, uint32&>(this, OutUInt32);
+}
+
+bool URCVirtualPropertyBase::GetValueInt64(int64& OuyInt64)
+{
+	return FRCVirtualPropertyCastHelpers::GetPrimitiveValue<FInt64Property, int64>(this, OuyInt64);
+}
+
+bool URCVirtualPropertyBase::GetValueUint64(uint64& OuyUInt64)
+{
+return FRCVirtualPropertyCastHelpers::GetPrimitiveValue<FUInt64Property, uint64>(this, OuyUInt64);
+}
+
+bool URCVirtualPropertyBase::GetValueFloat(float& OutFloat)
+{
+	return FRCVirtualPropertyCastHelpers::GetPrimitiveValue<FFloatProperty, float>(this, OutFloat);
+}
+
+bool URCVirtualPropertyBase::GetValueDouble(double& OutDouble)
+{
+	return FRCVirtualPropertyCastHelpers::GetPrimitiveValue<FDoubleProperty, double>(this, OutDouble);
+}
+
+bool URCVirtualPropertyBase::GetValueString(FString& OutStringValue)
+{
+	return FRCVirtualPropertyCastHelpers::GetPrimitiveValue<FStrProperty, FString>(this, OutStringValue);
+}
+
+bool URCVirtualPropertyBase::GetValueName(FName& OutNameValue)
+{
+	return FRCVirtualPropertyCastHelpers::GetPrimitiveValue<FNameProperty, FName>(this, OutNameValue);
+}
+
+bool URCVirtualPropertyBase::GetValueText(FText& OutTextValue)
+{
+	return FRCVirtualPropertyCastHelpers::GetPrimitiveValue<FTextProperty, FText>(this, OutTextValue);
+}
+
+bool URCVirtualPropertyBase::GetValueNumericInteger(int64& OutInt64Value)
+{
+	const uint8* ContainerPtr = GetContainerPtr();
+	if (!ContainerPtr)
+	{
+		return false;
+	}
+		
+	const FProperty* Property = GetProperty();
+		
+	if (const FNumericProperty* NumericProperty = CastField<FNumericProperty>(Property))
+	{
+		if (NumericProperty->IsInteger())
+		{
+			const void* DataResult = NumericProperty->ContainerPtrToValuePtr<void>(ContainerPtr);
+			OutInt64Value = reinterpret_cast<int64>(DataResult);
+			return true;
+		}
+	}
+
+	return false;
+}
+
+bool URCVirtualPropertyBase::GetValueVector(FVector& OutVector)
+{
+	return FRCVirtualPropertyCastHelpers::GetStructValue<FVector>(this, TBaseStructure<FVector>::Get(), &OutVector);
+}
+
+bool URCVirtualPropertyBase::GetValueRotator(FRotator& OutRotator)
+{
+	return FRCVirtualPropertyCastHelpers::GetStructValue<FRotator>(this, TBaseStructure<FRotator>::Get(), &OutRotator);
+}
+
+bool URCVirtualPropertyBase::GetValueColor(FColor& OutColor)
+{
+	return FRCVirtualPropertyCastHelpers::GetStructValue<FColor>(this, TBaseStructure<FColor>::Get(), &OutColor);
+}
+
+bool URCVirtualPropertyBase::SetValueBool(bool InBoolValue)
+{
+	return FRCVirtualPropertyCastHelpers::SetPrimitiveValue<FBoolProperty, bool>(this, InBoolValue);
+}
+
+bool URCVirtualPropertyBase::SetValueInt8(const int8 InInt8)
+{
+	return FRCVirtualPropertyCastHelpers::SetPrimitiveValue<FInt8Property, int8>(this, InInt8);
+}
+
+bool URCVirtualPropertyBase::SetValueByte(const uint8 InByte)
+{
+	return FRCVirtualPropertyCastHelpers::SetPrimitiveValue<FByteProperty, uint8>(this, InByte);
+}
+
+bool URCVirtualPropertyBase::SetValueInt16(const int16 InInt16)
+{
+	return FRCVirtualPropertyCastHelpers::SetPrimitiveValue<FInt16Property, int16>(this, InInt16);
+}
+
+bool URCVirtualPropertyBase::SetValueUint16(const uint16 InUInt16)
+{
+	return FRCVirtualPropertyCastHelpers::SetPrimitiveValue<FUInt16Property, uint16>(this, InUInt16);
+}
+
+bool URCVirtualPropertyBase::SetValueInt32(const int32 InInt32)
+{
+	return FRCVirtualPropertyCastHelpers::SetPrimitiveValue<FIntProperty, int32>(this, InInt32);
+}
+
+bool URCVirtualPropertyBase::SetValueUInt32(const uint32 InUInt32)
+{
+	return FRCVirtualPropertyCastHelpers::SetPrimitiveValue<FUInt32Property, uint32>(this, InUInt32);
+}
+
+bool URCVirtualPropertyBase::SetValueInt64(const int64 InInt64)
+{
+	return FRCVirtualPropertyCastHelpers::SetPrimitiveValue<FInt64Property, int64>(this, InInt64);
+}
+
+bool URCVirtualPropertyBase::SetValueUint64(const uint64 InUInt64)
+{
+	return FRCVirtualPropertyCastHelpers::SetPrimitiveValue<FUInt64Property, uint64>(this, InUInt64);
+}
+
+bool URCVirtualPropertyBase::SetValueFloat(const float InFloat)
+{
+	return FRCVirtualPropertyCastHelpers::SetPrimitiveValue<FFloatProperty, float>(this, InFloat);
+}
+
+bool URCVirtualPropertyBase::SetValueDouble(const double InDouble)
+{
+	return FRCVirtualPropertyCastHelpers::SetPrimitiveValue<FDoubleProperty, double>(this, InDouble);
+}
+
+bool URCVirtualPropertyBase::SetValueString(const FString& InStringValue)
+{
+	return FRCVirtualPropertyCastHelpers::SetPrimitiveValue<FStrProperty, FString>(this, InStringValue);
+}
+
+bool URCVirtualPropertyBase::SetValueName(const FName& InNameValue)
+{
+	return FRCVirtualPropertyCastHelpers::SetPrimitiveValue<FNameProperty, FName>(this, InNameValue);
+}
+
+bool URCVirtualPropertyBase::SetValueText(const FText& InTextValue)
+{
+	return FRCVirtualPropertyCastHelpers::SetPrimitiveValue<FTextProperty, FText>(this, InTextValue);
+}
+
+bool URCVirtualPropertyBase::SetValueNumericInteger(const int64 InInt64Value)
+{
+	uint8* ValuePtr = GetValuePtr();
+	if (!ValuePtr)
+	{
+		return false;
+	}
+		
+	const FProperty* Property = GetProperty();
+		
+	if (const FNumericProperty* NumericProperty = CastField<FNumericProperty>(Property))
+	{
+		if (NumericProperty->IsInteger())
+		{
+			NumericProperty->SetIntPropertyValue(ValuePtr, InInt64Value);
+		}
+	}
+
+	return false;
+}
+
+bool URCVirtualPropertyBase::SetValueVector(const FVector& InVector)
+{
+	return FRCVirtualPropertyCastHelpers::SetStructValue<FVector>(this, TBaseStructure<FVector>::Get(), InVector);
+}
+
+bool URCVirtualPropertyBase::SetValueRotator(const FRotator& InRotator)
+{
+	return FRCVirtualPropertyCastHelpers::SetStructValue<FRotator>(this, TBaseStructure<FRotator>::Get(), InRotator);
+}
+
+bool URCVirtualPropertyBase::SetValueColor(const FColor& InColor)
+{
+	return FRCVirtualPropertyCastHelpers::SetStructValue<FColor>(this, TBaseStructure<FColor>::Get(), InColor);
+}
+
+FName URCVirtualPropertyBase::GetPropertyName() const
+{
+	if (const FProperty* Property = GetProperty())
+	{
+		return Property->GetClass()->GetFName();
+	}
+	
+	return NAME_None;
+}
+
+const FInstancedPropertyBag* URCVirtualPropertyInContainer::GetPropertyBagInstance() const
+{
+	if (ContainerWeakPtr.Get())
+	{
+		return &ContainerWeakPtr->Bag;
+	}
+	
+	return Super::GetPropertyBagInstance();
+}
+
+void URCVirtualPropertySelfContainer::AddProperty(const FName InPropertyName, const EPropertyBagPropertyType InValueType, UObject* InValueTypeObject)
+{
+	if (Bag.GetNumPropertiesInBag() > 0)
+	{
+		return;
+	}
+	
+	PropertyName = InPropertyName;
+	Id = FGuid::NewGuid();
+	
+	Bag.AddProperty(PropertyName, InValueType, InValueTypeObject);
+}
+
+bool URCVirtualPropertySelfContainer::DuplicateProperty(const FName& InPropertyName, const FProperty* InSourceProperty)
+{
+	if (Bag.GetNumPropertiesInBag() > 0)
+	{
+		return false;
+	}
+
+	PropertyName = InPropertyName;
+	Id = FGuid::NewGuid();
+	
+	Bag.AddProperty(InPropertyName, InSourceProperty);
+	const FPropertyBagPropertyDesc* BagPropertyDesc = Bag.FindPropertyDescByName(InPropertyName);
+	check(BagPropertyDesc);
+
+	return true;
+}
+
+bool URCVirtualPropertySelfContainer::DuplicatePropertyWithCopy(const FName& InPropertyName, const FProperty* InSourceProperty, const uint8* InSourceContainerPtr)
+{
+	if (!DuplicateProperty(InPropertyName, InSourceProperty))
+	{
+		return false;
+	}
+	
+	if (InSourceContainerPtr)
+	{
+		return Bag.SetValue(PropertyName, InSourceProperty, InSourceContainerPtr) == EPropertyBagResult::Success;	
+	}
+	return false;
+}
+
+bool URCVirtualPropertySelfContainer::DuplicatePropertyWithCopy(URCVirtualPropertyBase* InVirtualProperty)
+{
+	if (!DuplicateProperty(InVirtualProperty->PropertyName, InVirtualProperty->GetProperty()))
+	{
+		return false;
+	}
+	
+	if (InVirtualProperty->GetContainerPtr())
+	{
+		return Bag.SetValue(PropertyName, InVirtualProperty->GetProperty(), InVirtualProperty->GetContainerPtr()) == EPropertyBagResult::Success;	
+	}
+
+	return false;
+}
+
+void URCVirtualPropertySelfContainer::Reset()
+{
+	PropertyName = NAME_None;
+	Id = FGuid();
+	Bag.Reset();
+}
+
+const FInstancedPropertyBag* URCVirtualPropertySelfContainer::GetPropertyBagInstance() const
+{
+	return &Bag;
+}
+
+TSharedPtr<FStructOnScope> URCVirtualPropertySelfContainer::CreateStructOnScope() const
+{
+	return MakeShared<FStructOnScope>(Bag.GetPropertyBagStruct(), Bag.GetMutableValue().GetMutableMemory());
+}
diff --git a/Engine/Plugins/VirtualProduction/RemoteControl/Source/RemoteControl/Private/RCVirtualPropertyContainer.cpp b/Engine/Plugins/VirtualProduction/RemoteControl/Source/RemoteControl/Private/RCVirtualPropertyContainer.cpp
new file mode 100644
--- /dev/null
+++ b/Engine/Plugins/VirtualProduction/RemoteControl/Source/RemoteControl/Private/RCVirtualPropertyContainer.cpp
@@ -0,0 +1,215 @@
+// Copyright Epic Games, Inc. All Rights Reserved.
+
+#include "RCVirtualPropertyContainer.h"
+
+#include "UObject/StructOnScope.h"
+#include "Templates/SubclassOf.h"
+
+#include "RCVirtualProperty.h"
+
+#if WITH_EDITOR
+#include "ScopedTransaction.h"
+#endif
+
+#define LOCTEXT_NAMESPACE "URCVirtualPropertyInContainer"
+
+URCVirtualPropertyInContainer* URCVirtualPropertyContainerBase::AddProperty(const FName& InPropertyName, TSubclassOf<URCVirtualPropertyInContainer> InPropertyClass, const EPropertyBagPropertyType InValueType, UObject* InValueTypeObject)
+{
+#if WITH_EDITOR
+	const FScopedTransaction Transaction(LOCTEXT("AddProperty", "AddP roperty"));
+
+	Modify();
+
+	MarkPackageDirty();
+#endif
+
+	const FName PropertyName = GenerateUniquePropertyName(InPropertyName, InValueType, this);
+	Bag.AddProperty(PropertyName, InValueType, InValueTypeObject);
+
+	const FPropertyBagPropertyDesc* BagPropertyDesc = Bag.FindPropertyDescByName(InPropertyName);
+	if (!ensure(BagPropertyDesc))
+	{
+		return nullptr;
+	}
+	
+	// Create Property in Container
+	URCVirtualPropertyInContainer* VirtualPropertyInContainer = NewObject<URCVirtualPropertyInContainer>(this, InPropertyClass.Get());
+	VirtualPropertyInContainer->PropertyName = PropertyName;
+	VirtualPropertyInContainer->PresetWeakPtr = PresetWeakPtr;
+	VirtualPropertyInContainer->ContainerWeakPtr = this;
+	VirtualPropertyInContainer->Id = FGuid::NewGuid();
+	
+	// Add property to Set
+	VirtualProperties.Add(VirtualPropertyInContainer);
+
+	return VirtualPropertyInContainer;
+}
+
+URCVirtualPropertyInContainer* URCVirtualPropertyContainerBase::DuplicateProperty(const FName& InPropertyName, const FProperty* InSourceProperty, TSubclassOf<URCVirtualPropertyInContainer> InPropertyClass)
+{
+#if WITH_EDITOR
+	const FScopedTransaction Transaction(LOCTEXT("DuplicateProperty", "Duplicate Property"));
+
+	Modify();
+
+	MarkPackageDirty();
+#endif
+
+
+	const FPropertyBagPropertyDesc* BagPropertyDesc = Bag.FindPropertyDescByName(InPropertyName);
+	if (ensure(!BagPropertyDesc))
+	{
+		return nullptr;
+	}
+
+	Bag.AddProperty(InPropertyName, InSourceProperty);
+	
+	URCVirtualPropertyInContainer* VirtualPropertyInContainer = NewObject<URCVirtualPropertyInContainer>(this, InPropertyClass.Get());
+	VirtualPropertyInContainer->PropertyName = InPropertyName;
+	VirtualPropertyInContainer->PresetWeakPtr = PresetWeakPtr;
+	VirtualPropertyInContainer->ContainerWeakPtr = this;
+	VirtualPropertyInContainer->Id = FGuid::NewGuid();
+
+	// Add property to Set
+	VirtualProperties.Add(VirtualPropertyInContainer);
+
+	return VirtualPropertyInContainer;
+}
+
+URCVirtualPropertyInContainer* URCVirtualPropertyContainerBase::DuplicatePropertyWithCopy(const FName& InPropertyName, const FProperty* InSourceProperty, const uint8* InSourceContainerPtr, TSubclassOf<URCVirtualPropertyInContainer> InPropertyClass)
+{
+	if (InSourceContainerPtr == nullptr)
+	{
+		return nullptr;
+	}
+	
+	URCVirtualPropertyInContainer* VirtualPropertyInContainer = DuplicateProperty(InPropertyName, InSourceProperty, InPropertyClass);
+	if (VirtualPropertyInContainer == nullptr)
+	{
+		return nullptr;
+	}
+
+	const FPropertyBagPropertyDesc* BagPropertyDesc = Bag.FindPropertyDescByName(InPropertyName);
+	check(BagPropertyDesc); // Property bag should be exists after DuplicateProperty()
+
+	ensure(Bag.SetValue(InPropertyName, InSourceProperty, InSourceContainerPtr) == EPropertyBagResult::Success);
+
+	return VirtualPropertyInContainer;
+}
+
+bool URCVirtualPropertyContainerBase::RemoveProperty(const FName& InPropertyName)
+{
+#if WITH_EDITOR
+	const FScopedTransaction Transaction(LOCTEXT("RemoveProperty", "Remove Property"));
+
+	Modify();
+
+	MarkPackageDirty();
+#endif
+
+	Bag.RemovePropertyByName(InPropertyName);
+
+	for (auto PropertiesIt = VirtualProperties.CreateIterator(); PropertiesIt; ++PropertiesIt)
+	{
+		if (const URCVirtualPropertyBase* VirtualProperty = *PropertiesIt)
+		{
+			if (VirtualProperty->PropertyName == InPropertyName)
+			{
+				PropertiesIt.RemoveCurrent();
+				return true;
+			}
+		}
+	}
+
+	return false;
+}
+
+void URCVirtualPropertyContainerBase::Reset()
+{
+#if WITH_EDITOR
+	const FScopedTransaction Transaction(LOCTEXT("EmptyProperties", "Empty Properties"));
+
+	Modify();
+
+	MarkPackageDirty();
+#endif
+	
+	VirtualProperties.Empty();
+
+	Bag.Reset();
+}
+
+URCVirtualPropertyBase* URCVirtualPropertyContainerBase::GetVirtualProperty(const FName& InPropertyName) const
+{
+	for (URCVirtualPropertyBase* VirtualProperty : VirtualProperties)
+	{
+		if (VirtualProperty->PropertyName == InPropertyName)
+		{
+			return VirtualProperty;
+		}
+	}
+	
+	return nullptr;
+}
+
+int32 URCVirtualPropertyContainerBase::GetNumVirtualProperties() const
+{
+	const int32 NumPropertiesInBag = Bag.GetNumPropertiesInBag();
+	const int32 NumVirtualProperties = VirtualProperties.Num();
+
+	check(NumPropertiesInBag == NumVirtualProperties);
+
+	return NumPropertiesInBag;
+}
+
+TSharedPtr<FStructOnScope> URCVirtualPropertyContainerBase::CreateStructOnScope() const
+{
+	return MakeShared<FStructOnScope>(Bag.GetPropertyBagStruct(), Bag.GetMutableValue().GetMutableMemory());
+}
+
+FName URCVirtualPropertyContainerBase::GenerateUniquePropertyName(const FName& InPropertyName, const EPropertyBagPropertyType InValueType, const URCVirtualPropertyContainerBase* InContainer)
+{
+	auto GetFinalName = [](const FString& InPrefix, const int32 InIndex = 0)
+	{
+		FString FinalName = InPrefix;
+
+		if (InIndex > 0)
+		{
+			FinalName += TEXT("_") + FString::FromInt(InIndex);
+		}
+
+		return FinalName;
+	};
+
+	const UEnum* Enum = FindObjectChecked<UEnum>(ANY_PACKAGE, TEXT("EPropertyBagPropertyType"));
+	const FName DefaultName = Enum->GetValueAsName(InValueType);
+	
+	int32 Index = 0;
+	const FString InitialName = InPropertyName.IsNone() ? DefaultName.ToString() : InPropertyName.ToString();
+	FString FinalName = InitialName;
+
+	
+	const FPropertyBagPropertyDesc* PropertyDesc = InContainer->Bag.FindPropertyDescByName(*FinalName);
+
+	while (PropertyDesc)
+	{
+		++Index;
+		FinalName = GetFinalName(InitialName, Index);
+		PropertyDesc = InContainer->Bag.FindPropertyDescByName(*FinalName);
+	}
+
+	return *FinalName;
+}
+
+#if WITH_EDITOR
+void URCVirtualPropertyContainerBase::OnModifyPropertyValue(const FPropertyChangedEvent& PropertyChangedEvent)
+{
+	const FScopedTransaction Transaction(LOCTEXT("OnModifyPropertyValue", "On Modify Property Value"));
+
+	Modify();
+
+	MarkPackageDirty();
+}
+#endif
+
+#undef LOCTEXT_NAMESPACE
\ No newline at end of file
diff --git a/Engine/Plugins/VirtualProduction/RemoteControl/Source/RemoteControl/Private/RemoteControlPreset.cpp b/Engine/Plugins/VirtualProduction/RemoteControl/Source/RemoteControl/Private/RemoteControlPreset.cpp
--- a/Engine/Plugins/VirtualProduction/RemoteControl/Source/RemoteControl/Private/RemoteControlPreset.cpp
+++ b/Engine/Plugins/VirtualProduction/RemoteControl/Source/RemoteControl/Private/RemoteControlPreset.cpp
@@ -10,6 +10,7 @@
 #include "GameFramework/Actor.h"
 #include "HAL/IConsoleManager.h"
 #include "IRemoteControlModule.h"
+#include "RCVirtualPropertyContainer.h"
 #include "Misc/CoreDelegates.h"
 #include "Misc/ITransaction.h"
 #include "Misc/Optional.h"
@@ -34,6 +35,7 @@
 #endif
 
 URemoteControlPreset::FOnPostLoadRemoteControlPreset URemoteControlPreset::OnPostLoadRemoteControlPreset;
+URemoteControlPreset::FOnPostInitPropertiesRemoteControlPreset URemoteControlPreset::OnPostInitPropertiesRemoteControlPreset;
 
 #define LOCTEXT_NAMESPACE "RemoteControlPreset"
 
@@ -662,6 +664,8 @@ void URemoteControlPreset::PostInitProperties()
 	if (!HasAnyFlags(RF_ClassDefaultObject | RF_ArchetypeObject))
 	{
 		RegisterDelegates();
+
+		OnPostInitPropertiesRemoteControlPreset.Broadcast(this);
 	}
 }
 
diff --git a/Engine/Plugins/VirtualProduction/RemoteControl/Source/RemoteControl/Public/RCVirtualProperty.h b/Engine/Plugins/VirtualProduction/RemoteControl/Source/RemoteControl/Public/RCVirtualProperty.h
new file mode 100644
--- /dev/null
+++ b/Engine/Plugins/VirtualProduction/RemoteControl/Source/RemoteControl/Public/RCVirtualProperty.h
@@ -0,0 +1,316 @@
+﻿// Copyright Epic Games, Inc. All Rights Reserved.
+ 
+#pragma once
+
+#include "PropertyBag.h"
+#include "UObject/Object.h"
+ 
+#include "RCVirtualProperty.generated.h"
+
+class IStructSerializerBackend;
+class URemoteControlPreset;
+class URCVirtualPropertyContainerBase;
+class FStructOnScope;
+
+/**
+ * Base class for dynamic virtual properties
+ * Remote Control Virtual Properties using Property Bag and FInstancedPropertyBag to serialize FProperties values and UStruct 
+ */
+UCLASS(BlueprintType)
+class REMOTECONTROL_API URCVirtualPropertyBase : public UObject
+{
+	GENERATED_BODY()
+
+	friend struct FRCVirtualPropertyCastHelpers;
+	friend class URCVirtualPropertySelfContainer;
+
+protected:
+	/** Return pointer to memory const container */
+	virtual const uint8* GetContainerPtr() const;
+
+	/** Return pointer to memory container */
+	virtual uint8* GetContainerPtr();
+
+	/** Return pointer to const value */
+	virtual const uint8* GetValuePtr() const;
+
+	/** Return pointer to value */
+	virtual uint8* GetValuePtr();
+
+	/** Return Pointer to Bag Property Description */
+	virtual const FPropertyBagPropertyDesc* GetBagPropertyDesc() const;
+
+	/** Pointer to Instanced property which is holds  bag of properties */
+	virtual const FInstancedPropertyBag* GetPropertyBagInstance() const;
+
+public:
+	/** Returns const FProperty for this RC virtual property */
+	virtual const FProperty* GetProperty() const;
+
+	/** Returns FProperty for this RC virtual property */
+	virtual FProperty* GetProperty();
+	
+	/** Return property bag property type */
+	EPropertyBagPropertyType GetValueType() const;
+
+	/** Return pointer to object that defines the Enum, Struct, or Class. */
+	const UObject* GetValueTypeObjectWeakPtr() const;
+
+	/**
+	 * Serialize Virtual Property to given Backend
+	 *
+	 * @param OutBackend Struct Serialize Backend 
+	 */
+	void SerializeToBackend(IStructSerializerBackend& OutBackend);
+
+	/** Compare this virtual property value with given property value */
+	bool IsValueEqual(URCVirtualPropertyBase* InVirtualProperty);
+	
+	/** Get Bool value from Virtual Property */
+	UFUNCTION(BlueprintCallable, Category = "Remote Control Behaviour")
+	bool GetValueBool(bool& OutBoolValue);
+
+	/** Get Int8 value from Virtual Property */
+	UFUNCTION()
+	bool GetValueInt8(int8& OutInt8);
+
+	/** Get Byte value from Virtual Property */
+	UFUNCTION(BlueprintCallable, Category = "Remote Control Behaviour")
+	bool GetValueByte(uint8& OutByte);
+
+	/** Get Int16 value from Virtual Property */
+	UFUNCTION()
+	bool GetValueInt16(int16& OutInt16);
+
+	/** Get Uint16 value from Virtual Property */
+	UFUNCTION()
+	bool GetValueUint16(uint16& OutUInt16);
+
+	/** Get Int32 value from Virtual Property */
+	UFUNCTION(BlueprintCallable, Category = "Remote Control Behaviour")
+	bool GetValueInt32(int32& OutInt32);
+
+	/** Get Uint32 value from Virtual Property */
+	UFUNCTION()
+	bool GetValueUInt32(uint32& OutUInt32);
+
+	/** Get Int64 value from Virtual Property */
+	UFUNCTION(BlueprintCallable, Category = "Remote Control Behaviour")
+	bool GetValueInt64(int64& OuyInt64);
+
+	/** Get Uint64 value from Virtual Property */
+	UFUNCTION()
+	bool GetValueUint64(uint64& OuyUInt64);
+
+	/** Get Float value from Virtual Property */
+	UFUNCTION(BlueprintCallable, Category = "Remote Control Behaviour")
+	bool GetValueFloat(float& OutFloat);
+
+	/** Get Double value from Virtual Property */
+	UFUNCTION(BlueprintCallable, Category = "Remote Control Behaviour")
+	bool GetValueDouble(double& OutDouble);
+
+	/** Get String value from Virtual Property */
+	UFUNCTION(BlueprintCallable, Category = "Remote Control Behaviour")
+	bool GetValueString(FString& OutStringValue);
+
+	/** Get Name value from Virtual Property */
+	UFUNCTION(BlueprintCallable, Category = "Remote Control Behaviour")
+	bool GetValueName(FName& OutNameValue);
+
+	/** Get Text value from Virtual Property */
+	UFUNCTION(BlueprintCallable, Category = "Remote Control Behaviour")
+	bool GetValueText(FText& OutTextValue);
+
+	/** Get Numeric value from Virtual Property */
+	UFUNCTION(BlueprintCallable, Category = "Remote Control Behaviour")
+	bool GetValueNumericInteger(int64& OutInt64Value);
+
+	/** Get Vector value from Virtual Property */
+	UFUNCTION(BlueprintCallable, Category = "Remote Control Behaviour")
+	bool GetValueVector(FVector& OutVector);
+
+	/** Get Rotator value from Virtual Property */
+	UFUNCTION(BlueprintCallable, Category = "Remote Control Behaviour")
+	bool GetValueRotator(FRotator& OutRotator);
+
+	/** Get Color value from Virtual Property */
+	UFUNCTION(BlueprintCallable, Category = "Remote Control Behaviour")
+	bool GetValueColor(FColor& OutColor);
+
+	/** Set Bool value from Virtual Property */
+	UFUNCTION(BlueprintCallable, Category = "Remote Control Behaviour")
+	bool SetValueBool(const bool InBoolValue);
+
+	/** Set Int8 value from Virtual Property */
+	UFUNCTION()
+	bool SetValueInt8(const int8 InInt8);
+
+	/** Set Byte value from Virtual Property */
+	UFUNCTION(BlueprintCallable, Category = "Remote Control Behaviour")
+	bool SetValueByte(const uint8 InByte);
+
+	/** Set Int16 value from Virtual Property */
+	UFUNCTION()
+	bool SetValueInt16(const int16 InInt16);
+
+	/** Set Uint16 value from Virtual Property */
+	UFUNCTION()
+	bool SetValueUint16(const uint16 InUInt16);
+
+	/** Set Int32 value from Virtual Property */
+	UFUNCTION(BlueprintCallable, Category = "Remote Control Behaviour")
+	bool SetValueInt32(const int32 InInt32);
+
+	/** Set Uint32 value from Virtual Property */
+	UFUNCTION()
+	bool SetValueUInt32(const uint32 InUInt32);
+
+	/** Set Int64 value from Virtual Property */
+	UFUNCTION(BlueprintCallable, Category = "Remote Control Behaviour")
+	bool SetValueInt64(const int64 InInt64);
+
+	/** Set Uint64 value from Virtual Property */
+	UFUNCTION()
+	bool SetValueUint64(const uint64 InUInt64);
+
+	/** Set Float value from Virtual Property */
+	UFUNCTION(BlueprintCallable, Category = "Remote Control Behaviour")
+	bool SetValueFloat(const float InFloat);
+
+	/** Set Double value from Virtual Property */
+	UFUNCTION(BlueprintCallable, Category = "Remote Control Behaviour")
+	bool SetValueDouble(const double InDouble);
+
+	/** Set String value from Virtual Property */
+	UFUNCTION(BlueprintCallable, Category = "Remote Control Behaviour")
+	bool SetValueString(const FString& InStringValue);
+
+	/** Set Name value from Virtual Property */
+	UFUNCTION(BlueprintCallable, Category = "Remote Control Behaviour")
+	bool SetValueName(const FName& InNameValue);
+
+	/** Set Text value from Virtual Property */
+	UFUNCTION(BlueprintCallable, Category = "Remote Control Behaviour")
+	bool SetValueText(const FText& InTextValue);
+
+	/** Set Numeric value from Virtual Property */
+	UFUNCTION(BlueprintCallable, Category = "Remote Control Behaviour")
+	bool SetValueNumericInteger(const int64 InInt64Value);
+
+	/** Set Vector value from Virtual Property */
+	UFUNCTION(BlueprintCallable, Category = "Remote Control Behaviour")
+	bool SetValueVector(const FVector& InVector);
+
+	/** Set Rotator value from Virtual Property */
+	UFUNCTION(BlueprintCallable, Category = "Remote Control Behaviour")
+	bool SetValueRotator(const FRotator& InRotator);
+
+	/** Set Color value from Virtual Property */
+	UFUNCTION(BlueprintCallable, Category = "Remote Control Behaviour")
+	bool SetValueColor(const FColor& InColor);
+	
+	/** Set FProperty Name */
+	UFUNCTION(BlueprintPure, Category = "Remote Control Behaviour")
+	FName GetPropertyName() const;
+
+public:
+	/** Unique property name */
+	UPROPERTY()
+	FName PropertyName;
+
+	/** Property Id */
+	UPROPERTY()
+	FGuid Id;
+
+	/** Pointer to Remote Control Preset */
+	UPROPERTY()
+	TWeakObjectPtr<URemoteControlPreset> PresetWeakPtr;
+};
+
+/**
+ * Remote Control Virtual Property which is stores in container with many properties
+ * Where Property Bag has more then one Property 
+ */
+UCLASS()
+class REMOTECONTROL_API URCVirtualPropertyInContainer : public URCVirtualPropertyBase
+{
+	GENERATED_BODY()
+
+protected:
+	//~ Begin URCVirtualPropertyBase interface
+	virtual const FInstancedPropertyBag* GetPropertyBagInstance() const override;
+	//~ End URCVirtualPropertyBase interface
+
+public:
+	/** Pointer to container where stores Virtual Properties */
+	UPROPERTY()
+	TWeakObjectPtr<URCVirtualPropertyContainerBase> ContainerWeakPtr;
+};
+
+/**
+ * Remote Control Virtual Property which is stores in self defined Property Bag
+ * In this case SelfContainer holds only single property in the Property Bag
+ */
+UCLASS()
+class REMOTECONTROL_API URCVirtualPropertySelfContainer : public URCVirtualPropertyBase
+{
+	GENERATED_BODY()
+
+public:
+	/**
+	 * Adds a new property to the bag. If property with this name already in bag the function just return
+	 *
+	 * @param InPropertyName				Name of the property
+	 * @param InValueType					Property Type
+	 * @param InValueTypeObject				Property Type object if exists
+	 * 
+	 */
+	void AddProperty(const FName InPropertyName, const EPropertyBagPropertyType InValueType, UObject* InValueTypeObject = nullptr);
+
+	/**
+	 * Duplicates property from give Property. If property with this name already in bag the function just return
+	 *
+	 * @param InPropertyName				Name of the property
+	 * @param InSourceProperty				Source FProperty for duplication
+	 *
+	 * @return true if property duplicated successfully 
+	 */
+	bool DuplicateProperty(const FName& InPropertyName, const FProperty* InSourceProperty);
+
+	/**
+	 * Duplicates property from give Property and copy the value. If property with this name already in bag the function just return
+	 *
+	 * @param InPropertyName				Name of the property
+	 * @param InSourceProperty				Source FProperty for duplication
+	 * @param InSourceContainerPtr			Pointer to source container
+	 *
+	 * @return true if property duplicated and value copied successfully 
+	 */
+	bool DuplicatePropertyWithCopy(const FName& InPropertyName, const FProperty* InSourceProperty, const uint8* InSourceContainerPtr);
+
+	/**
+	 * Duplicates property from given Virtual Property. If property with this name already in bag the function just return
+	 *
+	 * @param InVirtualProperty				Virtual Property to duplicate from
+	 *
+	 * @return true if property duplicated and value copied successfully 
+	 */
+	bool DuplicatePropertyWithCopy(URCVirtualPropertyBase* InVirtualProperty);
+
+	/** Resets the property bag instance to empty and remove Virtual Property data */
+	void Reset();
+
+	/** Creates new Struct on Scope for this Property Bag UStruct and Memory */
+	TSharedPtr<FStructOnScope> CreateStructOnScope() const;
+
+protected:
+	//~ Begin URCVirtualPropertyBase interface
+	virtual const FInstancedPropertyBag* GetPropertyBagInstance() const override;
+	//~ End URCVirtualPropertyBase interface
+
+private:
+	/** Instanced property bag for store a bag of properties. */
+	UPROPERTY()
+	FInstancedPropertyBag Bag;
+};
diff --git a/Engine/Plugins/VirtualProduction/RemoteControl/Source/RemoteControl/Public/RCVirtualPropertyContainer.h b/Engine/Plugins/VirtualProduction/RemoteControl/Source/RemoteControl/Public/RCVirtualPropertyContainer.h
new file mode 100644
--- /dev/null
+++ b/Engine/Plugins/VirtualProduction/RemoteControl/Source/RemoteControl/Public/RCVirtualPropertyContainer.h
@@ -0,0 +1,113 @@
+// Copyright Epic Games, Inc. All Rights Reserved.
+
+#pragma once
+
+#include "CoreMinimal.h"
+
+#include "PropertyBag.h"
+
+#include "RCVirtualPropertyContainer.generated.h"
+
+class URCVirtualPropertyInContainer;
+class URCVirtualPropertyBase;
+class URemoteControlPreset;
+class FStructOnScope;
+
+/**
+ * Container for more then one Virtual Property
+ */
+UCLASS(BlueprintType)
+class REMOTECONTROL_API URCVirtualPropertyContainerBase : public UObject
+{
+	GENERATED_BODY()
+
+	friend class URCVirtualPropertyInContainer;
+	
+public:
+	/**
+	 * Add property to this container.
+	 * That will add property to Property Bag and Create Remote Control Virtual Property
+	 *
+	 * @param InPropertyName				Property Name to add
+	 * @param InPropertyClass				Class of the Virtual Property 
+	 * @param InValueType					Property Type
+	 * @param InValueTypeObject				Property Type object if exists
+	 *
+	 * @return Virtual Property Object
+	 */
+	virtual URCVirtualPropertyInContainer* AddProperty(const FName& InPropertyName, TSubclassOf<URCVirtualPropertyInContainer> InPropertyClass, const EPropertyBagPropertyType InValueType, UObject* InValueTypeObject = nullptr);
+
+	/**
+	 * Duplicates property from give Property.
+	 * That will add property to Property Bag and Create Remote Control Virtual Property
+	 *
+	 * @param InPropertyName				Name of the property
+	 * @param InSourceProperty				Source FProperty for duplication
+	 * @param InPropertyClass				Class of the Virtual Property 
+	 *
+	 * @return Virtual Property Object
+	 */
+	virtual URCVirtualPropertyInContainer* DuplicateProperty(const FName& InPropertyName, const FProperty* InSourceProperty, TSubclassOf<URCVirtualPropertyInContainer> InPropertyClass);
+
+	/**
+	 * Duplicates property from give Property and copy the property value
+	 * That will add property to Property Bag and Create Remote Control Virtual Property
+	 *
+	 * @param InPropertyName				Name of the property
+	 * @param InSourceProperty				Source FProperty for duplication
+	 * @param InPropertyClass				Class of the Virtual Property
+	 * @param InSourceContainerPtr			Pointer to source container
+	 *
+	 * @return Virtual Property Object
+	 */
+	virtual URCVirtualPropertyInContainer* DuplicatePropertyWithCopy(const FName& InPropertyName, const FProperty* InSourceProperty, const uint8* InSourceContainerPtr, TSubclassOf<URCVirtualPropertyInContainer> InPropertyClass);
+
+	/**
+	 * Removes a property from the container by name if it exists.
+	 */
+	virtual bool RemoveProperty(const FName& InPropertyName);
+
+	/** Resets the property bag instance to empty and remove Virtual Properties */
+	virtual void Reset();
+
+	/**
+	 * Returns virtual property by specified name.
+	 */ 
+	virtual URCVirtualPropertyBase* GetVirtualProperty(const FName& InPropertyName) const;
+
+	/**
+	 * Returns number of virtual properties.
+	 */ 
+	int32 GetNumVirtualProperties() const;
+
+	/**
+	 * Creates new Struct on Scope for this Property Bag UStruct and Memory
+	 */
+	TSharedPtr<FStructOnScope> CreateStructOnScope() const;
+
+	/**
+	 * Generates unique name for the property for specified property container
+	 */
+	static FName GenerateUniquePropertyName(const FName& InPropertyName, const EPropertyBagPropertyType InValueType, const URCVirtualPropertyContainerBase* InContainer);
+
+#if WITH_EDITOR
+	/** Delegate when object changed */
+	virtual void OnModifyPropertyValue(const FPropertyChangedEvent& PropertyChangedEvent);
+#endif
+
+protected:
+	/** Holds bag of properties. */
+	UPROPERTY()
+	FInstancedPropertyBag Bag;
+
+public:
+	/** Set of the virtual properties */
+	UPROPERTY()
+	TSet<TObjectPtr<URCVirtualPropertyBase>> VirtualProperties;
+	
+	/** Pointer to Remote Control Preset */
+	UPROPERTY()
+	TWeakObjectPtr<URemoteControlPreset> PresetWeakPtr;
+};
+
+
diff --git a/Engine/Plugins/VirtualProduction/RemoteControl/Source/RemoteControl/Public/RemoteControlPreset.h b/Engine/Plugins/VirtualProduction/RemoteControl/Source/RemoteControl/Public/RemoteControlPreset.h
--- a/Engine/Plugins/VirtualProduction/RemoteControl/Source/RemoteControl/Public/RemoteControlPreset.h
+++ b/Engine/Plugins/VirtualProduction/RemoteControl/Source/RemoteControl/Public/RemoteControlPreset.h
@@ -12,6 +12,7 @@
 
 #include "RemoteControlPreset.generated.h"
 
+class URCVirtualPropertyContainerBase;
 class AActor;
 class IStructSerializerBackend;
 class IStructDeserializerBackend;
@@ -452,6 +453,10 @@ public:
 	/** Callback for post remote control preset load, called by URemoteControlPreset::PostLoad function */
 	DECLARE_MULTICAST_DELEGATE_OneParam(FOnPostLoadRemoteControlPreset, URemoteControlPreset* /* InPreset */);
 	static FOnPostLoadRemoteControlPreset OnPostLoadRemoteControlPreset;
+
+	/** Callback for post init properties of remote control preset, called by URemoteControlPreset::PostInitProperties function */
+	DECLARE_MULTICAST_DELEGATE_OneParam(FOnPostInitPropertiesRemoteControlPreset, URemoteControlPreset* /* InPreset */);
+	static FOnPostInitPropertiesRemoteControlPreset OnPostInitPropertiesRemoteControlPreset;
 	
 	URemoteControlPreset();
 
@@ -807,6 +812,10 @@ public:
 	UPROPERTY(EditAnywhere, Category = "Remote Control Preset")
 	TArray<URemoteControlBinding*> Bindings;
 
+	/** Holds virtual controllers properties, behaviours and actions */
+	UPROPERTY(Instanced)
+	TObjectPtr<URCVirtualPropertyContainerBase> ControllerContainer;
+
 private:
 	/** Generate a unique alias for this target. */
 	FName GenerateAliasForObjects(const TArray<UObject*>& Objects);
diff --git a/Engine/Plugins/VirtualProduction/RemoteControl/Source/RemoteControl/RemoteControl.Build.cs b/Engine/Plugins/VirtualProduction/RemoteControl/Source/RemoteControl/RemoteControl.Build.cs
--- a/Engine/Plugins/VirtualProduction/RemoteControl/Source/RemoteControl/RemoteControl.Build.cs
+++ b/Engine/Plugins/VirtualProduction/RemoteControl/Source/RemoteControl/RemoteControl.Build.cs
@@ -12,6 +12,7 @@ public class RemoteControl : ModuleRules
 				"CoreUObject",
 				"Engine",
 				"RemoteControlCommon",
+				"StructUtils"
 			}
 		);
 
@@ -20,7 +21,7 @@ public class RemoteControl : ModuleRules
 				"Cbor",
 				"Engine",
 				"RemoteControlInterception",
-				"Serialization",
+				"Serialization"
 			}
 		);
 
diff --git a/Engine/Plugins/VirtualProduction/RemoteControl/Source/RemoteControlLogic/Private/Action/RCActionContainer.cpp b/Engine/Plugins/VirtualProduction/RemoteControl/Source/RemoteControlLogic/Private/Action/RCActionContainer.cpp
new file mode 100644
--- /dev/null
+++ b/Engine/Plugins/VirtualProduction/RemoteControl/Source/RemoteControlLogic/Private/Action/RCActionContainer.cpp
@@ -0,0 +1,123 @@
+﻿// Copyright Epic Games, Inc. All Rights Reserved.
+
+#include "Action/RCActionContainer.h"
+
+#include "IRemoteControlModule.h"
+#include "RCVirtualProperty.h"
+#include "RemoteControlField.h"
+#include "Action/RCAction.h"
+#include "Action/RCFunctionAction.h"
+#include "Action/RCPropertyAction.h"
+
+void URCActionContainer::ExecuteActions()
+{
+	for (const URCAction* Action : Actions)
+	{
+		Action->Execute();
+	}
+}
+
+URCPropertyAction* URCActionContainer::AddAction(const TSharedPtr<FRemoteControlProperty> InRemoteControlProperty)
+{
+	if (!InRemoteControlProperty.IsValid())
+	{
+		return nullptr;
+	}
+	
+	if (FindActionByField(InRemoteControlProperty))
+	{
+		return nullptr;
+	}
+
+	// Create new Property
+	URCPropertyAction* NewPropertyAction = NewObject<URCPropertyAction>(this);
+	NewPropertyAction->PresetWeakPtr = PresetWeakPtr;
+	NewPropertyAction->ExposedFieldId = InRemoteControlProperty->GetId();
+	NewPropertyAction->Id = FGuid::NewGuid();
+	
+	if (FRCObjectReference ObjectRef; IRemoteControlModule::Get().ResolveObjectProperty(ERCAccess::READ_ACCESS, InRemoteControlProperty->GetBoundObjects()[0], InRemoteControlProperty->FieldPathInfo.ToString(), ObjectRef))
+	{
+		const FName& PropertyName = InRemoteControlProperty->GetProperty()->GetFName();
+		NewPropertyAction->PropertySelfContainer->DuplicatePropertyWithCopy(PropertyName, InRemoteControlProperty->GetProperty(), (uint8*)ObjectRef.ContainerAdress);
+	}
+
+	// Add action to array
+	Actions.Add(NewPropertyAction);
+	
+	return NewPropertyAction;
+}
+
+URCFunctionAction* URCActionContainer::AddAction(const TSharedPtr<FRemoteControlFunction> InRemoteControlFunction)
+{
+	if (!InRemoteControlFunction.IsValid())
+	{
+		return nullptr;
+	}
+
+	if (FindActionByField(InRemoteControlFunction))
+	{
+		return nullptr;
+	}
+	
+	// Create new Function Action
+	URCFunctionAction* NewFunctionAction = NewObject<URCFunctionAction>(this);
+	NewFunctionAction->PresetWeakPtr = PresetWeakPtr;
+	NewFunctionAction->ExposedFieldId = InRemoteControlFunction->GetId();
+	NewFunctionAction->Id = FGuid::NewGuid();
+
+	Actions.Add(NewFunctionAction);
+	
+	return NewFunctionAction;
+}
+
+URCAction* URCActionContainer::FindActionByFieldId(const FGuid InId)
+{
+	for (URCAction* Action :  Actions)
+	{
+		if (Action->ExposedFieldId == InId)
+		{
+			return Action;
+		}
+	}
+
+	return nullptr;
+}
+
+URCAction* URCActionContainer::FindActionByField(const TSharedPtr<FRemoteControlField> InRemoteControlField)
+{
+	for (URCAction* Action :  Actions)
+	{
+		if (Action->ExposedFieldId == InRemoteControlField->GetId())
+		{
+			return Action;
+		}
+	}
+
+	return nullptr;
+}
+
+int32 URCActionContainer::RemoveAction(const FGuid InExposedFieldId)
+{
+	int32 RemoveCount = 0;
+	
+	for (auto ActionsIt = Actions.CreateIterator(); ActionsIt; ++ActionsIt)
+	{
+		if (const URCAction* Action = *ActionsIt; Action->ExposedFieldId == InExposedFieldId)
+		{
+			ActionsIt.RemoveCurrent();
+			RemoveCount++;
+		}
+	}
+
+	return RemoveCount;
+}
+
+int32 URCActionContainer::RemoveAction(URCAction* InAction)
+{
+	return Actions.Remove(InAction);
+}
+
+void URCActionContainer::EmptyActions()
+{
+	Actions.Empty();
+}
diff --git a/Engine/Plugins/VirtualProduction/RemoteControl/Source/RemoteControlLogic/Private/Action/RCFunctionAction.cpp b/Engine/Plugins/VirtualProduction/RemoteControl/Source/RemoteControlLogic/Private/Action/RCFunctionAction.cpp
new file mode 100644
--- /dev/null
+++ b/Engine/Plugins/VirtualProduction/RemoteControl/Source/RemoteControlLogic/Private/Action/RCFunctionAction.cpp
@@ -0,0 +1,56 @@
+﻿// Copyright Epic Games, Inc. All Rights Reserved.
+
+#include "Action/RCFunctionAction.h"
+
+#include "IRemoteControlModule.h"
+#include "RemoteControlPreset.h"
+#include "StructSerializer.h"
+#include "Backends/CborStructSerializerBackend.h"
+
+void URCFunctionAction::Execute() const
+{
+	if (!PresetWeakPtr.IsValid())
+	{
+		return;
+	}
+
+	if (!ExposedFieldId.IsValid())
+	{
+		return;
+	}
+	
+	if (TSharedPtr<FRemoteControlFunction> RemoteControlFunction = PresetWeakPtr->GetExposedEntity<FRemoteControlFunction>(ExposedFieldId).Pin())
+	{
+		// Copy the default arguments.
+		FStructOnScope FunctionArgs{ RemoteControlFunction->GetFunction() };
+		for (TFieldIterator<FProperty> It(RemoteControlFunction->GetFunction()); It; ++It)
+		{
+			if (It->HasAnyPropertyFlags(CPF_Parm) && !It->HasAnyPropertyFlags(CPF_ReturnParm | CPF_OutParm))
+			{
+				It->CopyCompleteValue_InContainer(FunctionArgs.GetStructMemory(), RemoteControlFunction->FunctionArguments->GetStructMemory());
+			}
+		}
+
+		TArray<uint8> FunctionPayload;
+		FMemoryWriter Writer(FunctionPayload);
+
+		FCborStructSerializerBackend SerializerBackend{ Writer, EStructSerializerBackendFlags::Default };
+		FStructSerializer::Serialize(FunctionArgs.GetStructMemory(), *const_cast<UStruct*>(FunctionArgs.GetStruct()), SerializerBackend);
+
+		for (UObject* Object : RemoteControlFunction->GetBoundObjects())
+		{
+			FRCCallReference CallRef;
+			CallRef.Object = Object;
+			CallRef.Function = RemoteControlFunction->GetFunction();
+
+			FRCCall Call;
+			Call.CallRef = MoveTemp(CallRef);
+			Call.bGenerateTransaction = true;
+			Call.ParamStruct = FStructOnScope(FunctionArgs.GetStruct(), FunctionArgs.GetStructMemory());
+		
+			IRemoteControlModule::Get().InvokeCall(Call, ERCPayloadType::Cbor, FunctionPayload);
+		}
+	}
+	
+	Super::Execute();
+}
diff --git a/Engine/Plugins/VirtualProduction/RemoteControl/Source/RemoteControlLogic/Private/Action/RCPropertyAction.cpp b/Engine/Plugins/VirtualProduction/RemoteControl/Source/RemoteControlLogic/Private/Action/RCPropertyAction.cpp
new file mode 100644
--- /dev/null
+++ b/Engine/Plugins/VirtualProduction/RemoteControl/Source/RemoteControlLogic/Private/Action/RCPropertyAction.cpp
@@ -0,0 +1,59 @@
+﻿// Copyright Epic Games, Inc. All Rights Reserved.
+
+#include "Action/RCPropertyAction.h"
+
+#include "IRemoteControlModule.h"
+#include "IStructSerializerBackend.h"
+#include "RCVirtualProperty.h"
+#include "RemoteControlPreset.h"
+#include "StructSerializer.h"
+#include "Action/RCAction.h"
+#include "Backends/CborStructDeserializerBackend.h"
+#include "Backends/CborStructSerializerBackend.h"
+
+URCPropertyAction::URCPropertyAction()
+{
+	PropertySelfContainer = CreateDefaultSubobject<URCVirtualPropertySelfContainer>(FName("VirtualPropertySelfContainer"));
+}
+
+void URCPropertyAction::Execute() const
+{
+	if (!PresetWeakPtr.IsValid())
+	{
+		return;
+	}
+	
+	if (TSharedPtr<FRemoteControlProperty> RemoteControlProperty = PresetWeakPtr->GetExposedEntity<FRemoteControlProperty>(ExposedFieldId).Pin())
+	{
+		FRCObjectReference ObjectRef;
+		ObjectRef.Property = RemoteControlProperty->GetProperty();
+		ObjectRef.Access = ERCAccess::WRITE_ACCESS;
+	
+		ObjectRef.PropertyPathInfo = RemoteControlProperty->FieldPathInfo.ToString();
+	
+		bool bSuccess = true;
+		for (UObject* Object : RemoteControlProperty->GetBoundObjects())
+		{
+			if (IRemoteControlModule::Get().ResolveObjectProperty(ObjectRef.Access, Object, ObjectRef.PropertyPathInfo, ObjectRef))
+			{
+				TArray<uint8> Buffer;
+				FMemoryWriter Writer(Buffer);
+	
+				FCborStructSerializerBackend SerializerBackend(Writer, EStructSerializerBackendFlags::Default);
+	
+				FStructSerializerPolicies Policies; 
+				Policies.MapSerialization = EStructSerializerMapPolicies::Array;
+				
+				PropertySelfContainer->SerializeToBackend(SerializerBackend);
+	
+				// Deserialization
+				FMemoryReader Reader(Buffer);
+				FCborStructDeserializerBackend DeserializerBackend(Reader);
+	
+				IRemoteControlModule::Get().SetObjectProperties(ObjectRef, DeserializerBackend, ERCPayloadType::Cbor, Buffer);
+			}
+		}
+	}
+	
+	Super::Execute();
+}
diff --git a/Engine/Plugins/VirtualProduction/RemoteControl/Source/RemoteControlLogic/Private/Behaviour/Builtin/RCBehaviourIsEqualNode.cpp b/Engine/Plugins/VirtualProduction/RemoteControl/Source/RemoteControlLogic/Private/Behaviour/Builtin/RCBehaviourIsEqualNode.cpp
new file mode 100644
--- /dev/null
+++ b/Engine/Plugins/VirtualProduction/RemoteControl/Source/RemoteControlLogic/Private/Behaviour/Builtin/RCBehaviourIsEqualNode.cpp
@@ -0,0 +1,67 @@
+﻿// Copyright Epic Games, Inc. All Rights Reserved.
+
+#include "Behaviour/Builtin/RCBehaviourIsEqualNode.h"
+
+#include "RCVirtualProperty.h"
+#include "Behaviour/RCIsEqualBehaviour.h"
+#include "Controller/RCController.h"
+
+URCBehaviourIsEqualNode::URCBehaviourIsEqualNode()
+{
+	DisplayName = NSLOCTEXT("Remote Control Behaviour", "Behavior Name - Is Equal", "Is Equal");
+	BehaviorDescription = NSLOCTEXT("Remote Control Behaviour", "Behavior Desc - Is Equal", "Triggers an event when the input equals the specified value.");
+}
+
+bool URCBehaviourIsEqualNode::Execute_Implementation(URCBehaviour* InBehaviour) const
+{
+	if (!ensure(InBehaviour))
+	{
+		return false;
+	}
+	
+	URCController* RCController = InBehaviour->ControllerWeakPtr.Get();
+	if (!RCController)
+	{
+		return false;
+	}
+
+	URCIsEqualBehaviour* IsEqualBehaviour = Cast<URCIsEqualBehaviour>(InBehaviour);
+	if (!IsEqualBehaviour)
+	{
+		return false;
+	}
+
+	return RCController->IsValueEqual(IsEqualBehaviour->PropertySelfContainer);
+}
+
+bool URCBehaviourIsEqualNode::IsSupported_Implementation(URCBehaviour* InBehaviour) const
+{
+	static TArray<TPair<EPropertyBagPropertyType, UObject*>> SupportedPropertyBagTypes =
+	{
+		{ EPropertyBagPropertyType::Bool, nullptr },
+		{ EPropertyBagPropertyType::Byte, nullptr },
+		{ EPropertyBagPropertyType::Int32, nullptr },
+		{ EPropertyBagPropertyType::Float, nullptr },
+		{ EPropertyBagPropertyType::Double, nullptr },
+		{ EPropertyBagPropertyType::Name, nullptr },
+		{ EPropertyBagPropertyType::String, nullptr },
+		{ EPropertyBagPropertyType::Text, nullptr },
+		{ EPropertyBagPropertyType::Struct, TBaseStructure<FVector>::Get() },
+		{ EPropertyBagPropertyType::Struct, TBaseStructure<FColor>::Get() },
+		{ EPropertyBagPropertyType::Struct, TBaseStructure<FRotator>::Get() }
+	};
+	
+	return SupportedPropertyBagTypes.ContainsByPredicate(GetIsSupportedCallback(InBehaviour));
+}
+
+UClass* URCBehaviourIsEqualNode::GetBehaviourClass() const
+{
+	return URCIsEqualBehaviour::StaticClass();
+}
+
+void URCBehaviourIsEqualNode::OnPassed_Implementation(URCBehaviour* InBehaviour) const
+{
+}
+
+
+
diff --git a/Engine/Plugins/VirtualProduction/RemoteControl/Source/RemoteControlLogic/Private/Behaviour/Builtin/RCBehaviourOnEmptyNode.cpp b/Engine/Plugins/VirtualProduction/RemoteControl/Source/RemoteControlLogic/Private/Behaviour/Builtin/RCBehaviourOnEmptyNode.cpp
new file mode 100644
--- /dev/null
+++ b/Engine/Plugins/VirtualProduction/RemoteControl/Source/RemoteControlLogic/Private/Behaviour/Builtin/RCBehaviourOnEmptyNode.cpp
@@ -0,0 +1,56 @@
+﻿// Copyright Epic Games, Inc. All Rights Reserved.
+
+#include "Behaviour/Builtin/RCBehaviourOnEmptyNode.h"
+
+#include "PropertyBag.h"
+#include "Behaviour/RCBehaviour.h"
+#include "Controller/RCController.h"
+
+URCBehaviourOnEmptyNode::URCBehaviourOnEmptyNode()
+{
+	DisplayName = NSLOCTEXT("Remote Control Behaviour", "Behavior Name - On Empty", "On Empty");
+	BehaviorDescription = NSLOCTEXT("Remote Control Behaviour", "Behavior Desc - On Empty", "Triggers an event when the input changes from no value. Matching LinkID's will copy visiiblity value.");
+}
+
+bool URCBehaviourOnEmptyNode::Execute_Implementation(URCBehaviour* InBehaviour) const
+{
+	if (!ensure(InBehaviour))
+	{
+		return false;
+	}
+
+	URCController* RCController = InBehaviour->ControllerWeakPtr.Get();
+	if (!ensure(RCController))
+	{
+		return false;
+	}
+	
+	if (FString StringValue; RCController->GetValueString(StringValue))
+	{
+		return StringValue.IsEmpty();
+	}
+
+	if (FName NameValue; RCController->GetValueName(NameValue))
+	{
+		return NameValue.IsNone();
+	}
+
+	if (FText TextValue; RCController->GetValueText(TextValue))
+	{
+		return TextValue.IsEmpty();
+	}
+	
+	return false;
+}
+
+bool URCBehaviourOnEmptyNode::IsSupported_Implementation(URCBehaviour* InBehaviour) const
+{
+	static TArray<TPair<EPropertyBagPropertyType, UObject*>> SupportedPropertyBagTypes =
+	{
+		{ EPropertyBagPropertyType::Name, nullptr },
+		{ EPropertyBagPropertyType::String, nullptr },
+		{ EPropertyBagPropertyType::Text, nullptr },
+	};
+
+	return SupportedPropertyBagTypes.ContainsByPredicate(GetIsSupportedCallback(InBehaviour));
+}
diff --git a/Engine/Plugins/VirtualProduction/RemoteControl/Source/RemoteControlLogic/Private/Behaviour/Builtin/RCBehaviourSetValueNode.cpp b/Engine/Plugins/VirtualProduction/RemoteControl/Source/RemoteControlLogic/Private/Behaviour/Builtin/RCBehaviourSetValueNode.cpp
new file mode 100644
--- /dev/null
+++ b/Engine/Plugins/VirtualProduction/RemoteControl/Source/RemoteControlLogic/Private/Behaviour/Builtin/RCBehaviourSetValueNode.cpp
@@ -0,0 +1,37 @@
+﻿// Copyright Epic Games, Inc. All Rights Reserved.
+
+#include "Behaviour/Builtin/RCBehaviourSetValueNode.h"
+
+#include "PropertyBag.h"
+
+URCBehaviourSetValueNode::URCBehaviourSetValueNode()
+{
+	DisplayName = NSLOCTEXT("Remote Control Behaviour", "Behavior Name - Set Value", "Set Value");
+	BehaviorDescription = NSLOCTEXT("Remote Control Behaviour", "Behavior Desc - Set Value", "Triggers an event when the associated property is modified");
+}
+
+
+bool URCBehaviourSetValueNode::Execute_Implementation(URCBehaviour* InBehaviour) const
+{
+	return true;
+}
+
+bool URCBehaviourSetValueNode::IsSupported_Implementation(URCBehaviour* InBehaviour) const
+{
+	static TArray<TPair<EPropertyBagPropertyType, UObject*>> SupportedPropertyBagTypes =
+	{
+		{ EPropertyBagPropertyType::Bool, nullptr },
+		{ EPropertyBagPropertyType::Byte, nullptr },
+		{ EPropertyBagPropertyType::Int32, nullptr },
+		{ EPropertyBagPropertyType::Float, nullptr },
+		{ EPropertyBagPropertyType::Double, nullptr },
+		{ EPropertyBagPropertyType::Name, nullptr },
+		{ EPropertyBagPropertyType::String, nullptr },
+		{ EPropertyBagPropertyType::Text, nullptr },
+		{ EPropertyBagPropertyType::Struct, TBaseStructure<FVector>::Get() },
+		{ EPropertyBagPropertyType::Struct, TBaseStructure<FColor>::Get() },
+		{ EPropertyBagPropertyType::Struct, TBaseStructure<FRotator>::Get() }
+	};
+	
+	return SupportedPropertyBagTypes.ContainsByPredicate(GetIsSupportedCallback(InBehaviour));
+}
diff --git a/Engine/Plugins/VirtualProduction/RemoteControl/Source/RemoteControlLogic/Private/Behaviour/RCBehaviour.cpp b/Engine/Plugins/VirtualProduction/RemoteControl/Source/RemoteControlLogic/Private/Behaviour/RCBehaviour.cpp
new file mode 100644
--- /dev/null
+++ b/Engine/Plugins/VirtualProduction/RemoteControl/Source/RemoteControlLogic/Private/Behaviour/RCBehaviour.cpp
@@ -0,0 +1,107 @@
+﻿// Copyright Epic Games, Inc. All Rights Reserved.
+
+#include "Behaviour/RCBehaviour.h"
+
+#include "Engine/Blueprint.h"
+
+#include "Action/RCActionContainer.h"
+#include "Behaviour/RCBehaviourNode.h"
+
+URCBehaviour::URCBehaviour()
+{
+	ActionContainer = CreateDefaultSubobject<URCActionContainer>(FName("ActionContainer"));
+}
+
+void URCBehaviour::Execute()
+{
+	if (!ActionContainer->Actions.Num())
+	{
+		return;
+	}
+	const URCBehaviourNode* BehaviourNode = GetBehaviourNode();
+
+	// Execute before the logic
+	BehaviourNode->PreExecute(this);
+
+	if (BehaviourNode->Execute(this))
+	{
+		ActionContainer->ExecuteActions();
+		BehaviourNode->OnPassed(this);
+	}
+}
+
+int32 URCBehaviour::GetNumActions() const
+{
+	return ActionContainer->Actions.Num();
+}
+
+UClass* URCBehaviour::GetOverrideBehaviourBlueprintClass() const
+{
+	// Blueprint extensions stores as a separate blueprints files that is why it should check can be path resolved as asset 
+	if (!OverrideBehaviourBlueprintClassPath.IsValid() || !ensure(OverrideBehaviourBlueprintClassPath.IsAsset()))
+	{
+		return nullptr;
+	}
+
+	// If class was loaded before then just resolve class, that is faster
+	if (UClass* ResolveClass = OverrideBehaviourBlueprintClassPath.ResolveClass())
+	{
+		return ResolveClass;
+	}
+
+	// Load class from the UObject path
+	return OverrideBehaviourBlueprintClassPath.TryLoadClass<URCBehaviourNode>();
+}
+
+#if WITH_EDITORONLY_DATA
+UBlueprint* URCBehaviour::GetBlueprint() const
+{
+	UClass* OverrideBehaviourBlueprintClass = GetOverrideBehaviourBlueprintClass();
+	return OverrideBehaviourBlueprintClass ? Cast<UBlueprint>(OverrideBehaviourBlueprintClass->ClassGeneratedBy) : nullptr;
+}
+#endif
+
+void URCBehaviour::SetOverrideBehaviourBlueprintClass(UBlueprint* InBlueprint)
+{
+	if (InBlueprint)
+	{
+		OverrideBehaviourBlueprintClassPath = InBlueprint->GeneratedClass.Get();
+	}
+}
+
+URCBehaviourNode* URCBehaviour::GetBehaviourNode()
+{
+	UClass* OverrideBehaviourBlueprintClass = GetOverrideBehaviourBlueprintClass();
+	UClass* FinalBehaviourNodeClass = OverrideBehaviourBlueprintClass ? OverrideBehaviourBlueprintClass : BehaviourNodeClass.Get();
+
+	if (!CachedBehaviourNode || CachedBehaviourNodeClass != FinalBehaviourNodeClass)
+	{
+		CachedBehaviourNode = NewObject<URCBehaviourNode>(this, FinalBehaviourNodeClass);
+	}
+
+	CachedBehaviourNodeClass = FinalBehaviourNodeClass;
+
+	return CachedBehaviourNode;
+}
+
+#if WITH_EDITOR
+const FText& URCBehaviour::GetDisplayName()
+{
+	if (!ensure(BehaviourNodeClass))
+	{
+		return FText::GetEmpty();
+	}
+
+	return GetDefault<URCBehaviourNode>(BehaviourNodeClass)->DisplayName;
+}
+
+const FText& URCBehaviour::GetBehaviorDescription()
+{
+	if (!ensure(BehaviourNodeClass))
+	{
+		return FText::GetEmpty();
+	}
+
+	return GetDefault<URCBehaviourNode>(BehaviourNodeClass)->BehaviorDescription;
+}
+#endif
\ No newline at end of file
diff --git a/Engine/Plugins/VirtualProduction/RemoteControl/Source/RemoteControlLogic/Private/Behaviour/RCBehaviourBlueprintNode.cpp b/Engine/Plugins/VirtualProduction/RemoteControl/Source/RemoteControlLogic/Private/Behaviour/RCBehaviourBlueprintNode.cpp
new file mode 100644
--- /dev/null
+++ b/Engine/Plugins/VirtualProduction/RemoteControl/Source/RemoteControlLogic/Private/Behaviour/RCBehaviourBlueprintNode.cpp
@@ -0,0 +1,21 @@
+﻿// Copyright Epic Games, Inc. All Rights Reserved.
+
+#include "Behaviour/RCBehaviourBlueprintNode.h"
+
+bool URCBehaviourBlueprintNode::Execute_Implementation(URCBehaviour* InBehaviour) const
+{
+	return false;
+}
+
+void URCBehaviourBlueprintNode::PreExecute_Implementation(URCBehaviour* InBehaviour) const
+{
+}
+
+bool URCBehaviourBlueprintNode::IsSupported_Implementation(URCBehaviour* InBehaviour) const
+{
+	return false;
+}
+
+void URCBehaviourBlueprintNode::OnPassed_Implementation(URCBehaviour* InBehaviour) const
+{
+}
diff --git a/Engine/Plugins/VirtualProduction/RemoteControl/Source/RemoteControlLogic/Private/Behaviour/RCBehaviourNode.cpp b/Engine/Plugins/VirtualProduction/RemoteControl/Source/RemoteControlLogic/Private/Behaviour/RCBehaviourNode.cpp
new file mode 100644
--- /dev/null
+++ b/Engine/Plugins/VirtualProduction/RemoteControl/Source/RemoteControlLogic/Private/Behaviour/RCBehaviourNode.cpp
@@ -0,0 +1,42 @@
+﻿// Copyright Epic Games, Inc. All Rights Reserved.
+
+#include "Behaviour/RCBehaviourNode.h"
+#include "RCVirtualProperty.h"
+#include "Behaviour/RCBehaviour.h"
+#include "Controller/RCController.h"
+
+void URCBehaviourNode::PreExecute(URCBehaviour* InBehaviour) const
+{
+}
+
+bool URCBehaviourNode::Execute(URCBehaviour* InBehaviour) const
+{
+	return false;
+}
+
+bool URCBehaviourNode::IsSupported(URCBehaviour* InBehaviour) const
+{
+	return false;
+}
+
+void URCBehaviourNode::OnPassed(URCBehaviour* InBehaviour) const
+{
+}
+
+UClass* URCBehaviourNode::GetBehaviourClass() const
+{
+	return URCBehaviour::StaticClass();
+}
+
+TFunction<bool(const TPair<EPropertyBagPropertyType, UObject*>)> URCBehaviourNode::GetIsSupportedCallback(URCBehaviour* InBehaviour)
+{
+	return[InBehaviour](const TPair<EPropertyBagPropertyType, UObject*> InPair)
+	{
+		if (const URCController* RCController = InBehaviour->ControllerWeakPtr.Get())
+		{
+			return RCController->GetValueType() == InPair.Key && RCController->GetValueTypeObjectWeakPtr() == InPair.Value;
+		}
+		
+		return false;
+	};
+}
diff --git a/Engine/Plugins/VirtualProduction/RemoteControl/Source/RemoteControlLogic/Private/Behaviour/RCIsEqualBehaviour.cpp b/Engine/Plugins/VirtualProduction/RemoteControl/Source/RemoteControlLogic/Private/Behaviour/RCIsEqualBehaviour.cpp
new file mode 100644
--- /dev/null
+++ b/Engine/Plugins/VirtualProduction/RemoteControl/Source/RemoteControlLogic/Private/Behaviour/RCIsEqualBehaviour.cpp
@@ -0,0 +1,24 @@
+﻿// Copyright Epic Games, Inc. All Rights Reserved.
+
+#include "Behaviour/RCIsEqualBehaviour.h"
+
+#include "RCVirtualProperty.h"
+#include "Controller/RCController.h"
+
+URCIsEqualBehaviour::URCIsEqualBehaviour()
+{
+	PropertySelfContainer = CreateDefaultSubobject<URCVirtualPropertySelfContainer>(FName("VirtualPropertySelfContainer"));
+}
+
+void URCIsEqualBehaviour::Initialize()
+{
+	URCController* RCController = ControllerWeakPtr.Get();
+	if (!RCController)
+	{
+		return;
+	}
+
+	PropertySelfContainer->DuplicatePropertyWithCopy(RCController);
+	
+	Super::Initialize();
+}
diff --git a/Engine/Plugins/VirtualProduction/RemoteControl/Source/RemoteControlLogic/Private/Controller/RCController.cpp b/Engine/Plugins/VirtualProduction/RemoteControl/Source/RemoteControlLogic/Private/Controller/RCController.cpp
new file mode 100644
--- /dev/null
+++ b/Engine/Plugins/VirtualProduction/RemoteControl/Source/RemoteControlLogic/Private/Controller/RCController.cpp
@@ -0,0 +1,76 @@
+// Copyright Epic Games, Inc. All Rights Reserved.
+
+#include "Controller/RCController.h"
+
+#include "RCVirtualPropertyContainer.h"
+#include "RemoteControlPreset.h"
+#include "Action/RCActionContainer.h"
+#include "Behaviour/RCBehaviour.h"
+#include "Behaviour/RCBehaviourNode.h"
+
+URCBehaviour* URCController::AddBehaviour(TSubclassOf<URCBehaviourNode> InBehaviourNodeClass)
+{
+	URCBehaviour* NewBehaviour = CreateBehaviour(InBehaviourNodeClass);
+	if (!ensure(NewBehaviour))
+	{
+		return nullptr;
+	}
+
+	NewBehaviour->Initialize();
+	
+	Behaviours.Add(NewBehaviour);
+
+	return NewBehaviour;
+}
+
+URCBehaviour* URCController::CreateBehaviour(TSubclassOf<URCBehaviourNode> InBehaviourNodeClass)
+{
+	const URCBehaviourNode* DefaultBehaviourNode = Cast<URCBehaviourNode>(InBehaviourNodeClass->GetDefaultObject());
+	
+	URCBehaviour* NewBehaviour = NewObject<URCBehaviour>(this, DefaultBehaviourNode->GetBehaviourClass());
+	NewBehaviour->BehaviourNodeClass = InBehaviourNodeClass;
+	NewBehaviour->Id = FGuid::NewGuid();
+	NewBehaviour->ActionContainer->PresetWeakPtr = PresetWeakPtr;
+	NewBehaviour->ControllerWeakPtr = this;
+	
+	if (!DefaultBehaviourNode->IsSupported(NewBehaviour))
+	{
+		return nullptr;
+	}
+	
+	return NewBehaviour;
+}
+
+int32 URCController::RemoveBehaviour(URCBehaviour* InBehaviour)
+{
+	return Behaviours.Remove(InBehaviour);
+}
+
+int32 URCController::RemoveBehaviour(const FGuid InBehaviourId)
+{
+	int32 RemovedCount = 0;
+	
+	for (auto BehaviourIt = Behaviours.CreateIterator(); BehaviourIt; ++BehaviourIt)
+	{
+		if (const URCBehaviour* Behaviour = *BehaviourIt; Behaviour->Id == InBehaviourId)
+		{
+			BehaviourIt.RemoveCurrent();
+			RemovedCount++;
+		}
+	}
+
+	return RemovedCount;
+}
+
+void URCController::EmptyBehaviours()
+{
+	Behaviours.Empty();
+}
+
+void URCController::ExecuteBehaviours()
+{
+	for (URCBehaviour* Behaviour : Behaviours)
+	{
+		Behaviour->Execute();
+	}
+}
diff --git a/Engine/Plugins/VirtualProduction/RemoteControl/Source/RemoteControlLogic/Private/Controller/RCControllerContainer.cpp b/Engine/Plugins/VirtualProduction/RemoteControl/Source/RemoteControlLogic/Private/Controller/RCControllerContainer.cpp
new file mode 100644
--- /dev/null
+++ b/Engine/Plugins/VirtualProduction/RemoteControl/Source/RemoteControlLogic/Private/Controller/RCControllerContainer.cpp
@@ -0,0 +1,25 @@
+﻿// Copyright Epic Games, Inc. All Rights Reserved.
+
+#include "Controller/RCControllerContainer.h"
+#include "Controller/RCController.h"
+
+#if WITH_EDITOR
+void URCControllerContainer::OnModifyPropertyValue(const FPropertyChangedEvent& Event)
+{
+	if (const FProperty* FinalProperty = (Event.Property == Event.MemberProperty) ? Event.Property : Event.MemberProperty)
+	{
+		for (URCVirtualPropertyBase* VirtualProperty : VirtualProperties)
+		{
+			if (VirtualProperty && VirtualProperty->PropertyName == FinalProperty->GetFName())
+			{
+				if (URCController* Controller = Cast<URCController>(VirtualProperty))
+				{
+					Controller->ExecuteBehaviours();
+				}
+			}
+		}
+	}
+
+	Super::OnModifyPropertyValue(Event);
+}
+#endif
diff --git a/Engine/Plugins/VirtualProduction/RemoteControl/Source/RemoteControlLogic/Private/RemoteControlLogicModule.cpp b/Engine/Plugins/VirtualProduction/RemoteControl/Source/RemoteControlLogic/Private/RemoteControlLogicModule.cpp
new file mode 100644
--- /dev/null
+++ b/Engine/Plugins/VirtualProduction/RemoteControl/Source/RemoteControlLogic/Private/RemoteControlLogicModule.cpp
@@ -0,0 +1,40 @@
+// Copyright Epic Games, Inc. All Rights Reserved.
+
+#include "Modules/ModuleManager.h"
+
+#include "RemoteControlPreset.h"
+#include "Controller/RCControllerContainer.h"
+
+/**
+ * Remote Control Logic Module Class
+ */
+class FRemoteControlLogicModule : public IModuleInterface
+{
+public:
+	//~ Begin IModuleInterface
+	virtual void StartupModule() override
+	{
+		URemoteControlPreset::OnPostInitPropertiesRemoteControlPreset.AddRaw(this, &FRemoteControlLogicModule::OnPostInitPropertiesRemoteControlPreset);
+	}
+
+	virtual void ShutdownModule() override
+	{
+		URemoteControlPreset::OnPostInitPropertiesRemoteControlPreset.RemoveAll(this);
+	}
+	//~ End IModuleInterface
+
+private:
+	void OnPostInitPropertiesRemoteControlPreset(URemoteControlPreset* InPreset) const
+	{
+		// Create a new controller for preset if that does not exists
+		if (InPreset && InPreset->ControllerContainer.IsNull())
+		{
+			URCControllerContainer* ControllerContainer = NewObject<URCControllerContainer>(InPreset);
+			ControllerContainer->PresetWeakPtr = InPreset;
+
+			InPreset->ControllerContainer = ControllerContainer;
+		}
+	}
+};
+
+IMPLEMENT_MODULE(FRemoteControlLogicModule, RemoteControlLogic);
\ No newline at end of file
diff --git a/Engine/Plugins/VirtualProduction/RemoteControl/Source/RemoteControlLogic/Private/Tests/RemoteControlLogicTestData.h b/Engine/Plugins/VirtualProduction/RemoteControl/Source/RemoteControlLogic/Private/Tests/RemoteControlLogicTestData.h
new file mode 100644
--- /dev/null
+++ b/Engine/Plugins/VirtualProduction/RemoteControl/Source/RemoteControlLogic/Private/Tests/RemoteControlLogicTestData.h
@@ -0,0 +1,26 @@
+// Copyright Epic Games, Inc. All Rights Reserved.
+
+#pragma once
+
+#include "CoreMinimal.h"
+
+#include "RemoteControlLogicTestData.generated.h"
+
+UCLASS()
+class URemoteControlLogicTestData : public UObject
+{
+public:
+	GENERATED_BODY()
+
+	UPROPERTY(EditAnywhere, BlueprintReadWrite, Category = RemoteControlTest)
+	FColor Color = FColor(1,2,3,4 );
+
+	UPROPERTY(EditAnywhere, BlueprintReadWrite, Category = RemoteControlTest)
+	int32 TestInt = 0;
+	
+	UFUNCTION(BlueprintCallable, Category = RemoteControlTest)
+	void TestIntFunction()
+	{
+		TestInt++;
+	}
+};
diff --git a/Engine/Plugins/VirtualProduction/RemoteControl/Source/RemoteControlLogic/Private/Tests/RemoteControlLogicTests.cpp b/Engine/Plugins/VirtualProduction/RemoteControl/Source/RemoteControlLogic/Private/Tests/RemoteControlLogicTests.cpp
new file mode 100644
--- /dev/null
+++ b/Engine/Plugins/VirtualProduction/RemoteControl/Source/RemoteControlLogic/Private/Tests/RemoteControlLogicTests.cpp
@@ -0,0 +1,172 @@
+﻿// Copyright Epic Games, Inc. All Rights Reserved.
+ 
+#include "CoreMinimal.h"
+#include "IRemoteControlModule.h"
+#include "RCVirtualProperty.h"
+#include "RCVirtualPropertyContainer.h"
+#include "RemoteControlLogicTestData.h"
+#include "RemoteControlPreset.h"
+#include "Action/RCAction.h"
+#include "Action/RCActionContainer.h"
+#include "Action/RCFunctionAction.h"
+#include "Action/RCPropertyAction.h"
+#include "Behaviour/RCBehaviour.h"
+#include "Behaviour/RCIsEqualBehaviour.h"
+#include "Behaviour/Builtin/RCBehaviourIsEqualNode.h"
+#include "Behaviour/Builtin/RCBehaviourOnEmptyNode.h"
+#include "Behaviour/Builtin/RCBehaviourSetValueNode.h"
+#include "Controller/RCController.h"
+#include "Misc/AutomationTest.h"
+#include "UObject/StrongObjectPtr.h"
+
+PRAGMA_DISABLE_OPTIMIZATION
+
+#define PROP_NAME(Class, Name) GET_MEMBER_NAME_CHECKED(Class, Name)
+#define GET_TEST_PROP(PropName) URemoteControlLogicTestData::StaticClass()->FindPropertyByName(PROP_NAME(URemoteControlLogicTestData, PropName))
+ 
+IMPLEMENT_SIMPLE_AUTOMATION_TEST(FRemoteControlLogicTest, "Plugin.RemoteControl.Logic.Runtime", EAutomationTestFlags::EditorContext | EAutomationTestFlags::EngineFilter)
+bool FRemoteControlLogicTest::RunTest(const FString& Parameters)
+{
+	// 1. Create Preset
+	TStrongObjectPtr<URemoteControlPreset> Preset{ NewObject<URemoteControlPreset>() };
+	TStrongObjectPtr<URemoteControlLogicTestData> TestObject{ NewObject<URemoteControlLogicTestData>() };
+ 
+	// 1.1 Copy Int test property value
+	const int32 TestIntValue = TestObject->TestInt;
+	
+	// 2. Expose Fields
+	// 2.1 Expose Properties
+	const TSharedPtr<FRemoteControlProperty> RCProp1 = Preset->ExposeProperty(TestObject.Get(), FRCFieldPathInfo{GET_TEST_PROP(Color)->GetName()}).Pin();
+ 
+	// 2.2 Expose Functions
+	// UObject* Object, UFunction* Function, FRemoteControlPresetExposeArgs Args
+	UFunction* TestIntFunction = TestObject->GetClass()->FindFunctionByName(GET_FUNCTION_NAME_CHECKED(URemoteControlLogicTestData, TestIntFunction));
+	const TSharedPtr<FRemoteControlFunction> RCFunc1 = Preset->ExposeFunction(TestObject.Get(), TestIntFunction).Pin();
+ 
+	// 3. Add Controllers
+	// 3.1.1 Create Float Property
+	URCController* FloatController =  Cast<URCController>(Preset->ControllerContainer->AddProperty(TEXT("FloatController"), URCController::StaticClass(), EPropertyBagPropertyType::Float));
+	// 3.1.2 Create Second Float Property
+	URCController* FloatController1 = Cast<URCController>(Preset->ControllerContainer->AddProperty(TEXT("FloatController1"), URCController::StaticClass(), EPropertyBagPropertyType::Float));
+	// 3.1.3 Create Bool Property
+	URCController* BoolController = Cast<URCController>(Preset->ControllerContainer->AddProperty(TEXT("BoolController"), URCController::StaticClass(), EPropertyBagPropertyType::Bool));
+	// 3.1.4 Create Int Controller
+	URCController* IntController =  Cast<URCController>(Preset->ControllerContainer->AddProperty(TEXT("IntController"), URCController::StaticClass(), EPropertyBagPropertyType::Int32));
+	// 3.1.5 Create String Controller
+	URCController* StrController =  Cast<URCController>(Preset->ControllerContainer->AddProperty(TEXT("StrController"), URCController::StaticClass(), EPropertyBagPropertyType::String));
+	// 3.1.6 Create Vector Struct Controller
+	URCController* VectorStructController =  Cast<URCController>(Preset->ControllerContainer->AddProperty(TEXT("VectorProperty"), URCController::StaticClass(), EPropertyBagPropertyType::Struct, TBaseStructure<FVector>::Get()));
+ 
+	// 3.2 Set Float Property Value to controller
+	constexpr float FloatValue = 0.65f;
+	TestTrue(TEXT("Should Set Float"), FloatController->SetValueFloat(FloatValue) == true);
+
+	// 3.3 Set Vector Property Value to container
+	const FVector VectorValue = FVector(5.f, 6.f, 7.f);
+	TestTrue(TEXT("Should Set Float"), VectorStructController->SetValueVector(VectorValue) == true);
+
+	FVector OutVectorValue = FVector(0.f);
+	TestTrue(TEXT("Should Set Float"), VectorStructController->GetValueVector(OutVectorValue) == true);
+	TestEqual(TEXT("Vectors should be the same"), VectorValue, OutVectorValue);
+	
+	// 4. Add Behaviour To Properties
+	// 4.1 Add Is Equal Behaviour
+	URCBehaviour* FloatControllerBehaviour = FloatController->AddBehaviour(URCBehaviourIsEqualNode::StaticClass());
+	const URCIsEqualBehaviour* IsEqualBehaviour = Cast<URCIsEqualBehaviour>(FloatControllerBehaviour);
+	float IsEqualBehaviourFloatValue = 0.f;
+ 
+	// 4.1 Get Is Equal behaviour Value
+	TestTrue(TEXT("Should Get Float"), IsEqualBehaviour->PropertySelfContainer->GetValueFloat(IsEqualBehaviourFloatValue));
+	TestEqual(TEXT("float controller value should be the same as a is equal behaviour"), IsEqualBehaviourFloatValue, FloatValue);
+ 
+	// 4.2 Add Other behaviours
+	const URCBehaviour* IntControllerBehaviour = FloatController->AddBehaviour(URCBehaviourSetValueNode::StaticClass());
+	const URCBehaviour* StrControllerBehaviour = StrController->AddBehaviour(URCBehaviourOnEmptyNode::StaticClass());
+ 
+	// 5 Add actions
+	// 5.1 Add Float Controller Actions
+	URCPropertyAction* FloatControllerBehaviourAction = FloatControllerBehaviour->ActionContainer->AddAction(RCProp1);
+
+	FColor OutColorValue = FColor();
+	TestTrue(TEXT("Should Get Color"), FloatControllerBehaviourAction->PropertySelfContainer->GetValueColor(OutColorValue) == true);
+	TestEqual(TEXT("Action value should be equal expose property value"), OutColorValue, TestObject->Color);
+	
+	// 5.3 Set new Action Value
+	constexpr FColor SecColorValue(5,4,3,2);
+	TestTrue(TEXT("Should Set Color"), FloatControllerBehaviourAction->PropertySelfContainer->SetValueColor(SecColorValue));
+	
+	// 5.4 Add Function Controller Actions
+	URCFunctionAction* FloatControllerBehaviourAction1 = FloatControllerBehaviour->ActionContainer->AddAction(RCFunc1);
+	URCPropertyAction* StrControllerBehaviourAction = StrControllerBehaviour->ActionContainer->AddAction(RCProp1);
+	FColor StringControllerColorValue(7,8,9,10);
+	TestTrue(TEXT("Should Set Color"), StrControllerBehaviourAction->PropertySelfContainer->SetValueColor(StringControllerColorValue));
+	
+ 
+	// 5.5 It should not add same action for second time
+	URCAction* StrControllerBehaviourAction1 = StrControllerBehaviour->ActionContainer->AddAction(RCProp1);
+	
+	// 6. Execute Behaviours
+ 
+	// 6.1 Check Exposed Property Value before execute
+	FRCObjectReference ObjectRef;
+	IRemoteControlModule::Get().ResolveObjectProperty(ERCAccess::READ_ACCESS, RCProp1->GetBoundObjects()[0], RCProp1->FieldPathInfo.ToString(), ObjectRef);
+	const uint8* RCProp1ValuePtr = RCProp1->GetProperty()->ContainerPtrToValuePtr<uint8>(ObjectRef.ContainerAdress);
+	FMemory::Memcpy(&OutColorValue, RCProp1ValuePtr, RCProp1->GetProperty()->GetSize());
+	TestNotEqual(TEXT("The exposed property should not be updated before behaviour"), OutColorValue, SecColorValue);
+ 
+	// 6.2 Execute behaviours
+	FloatController->ExecuteBehaviours();
+	
+	// 6.3 Check Exposed Property Value after execute
+	FMemory::Memcpy(&OutColorValue, RCProp1ValuePtr, RCProp1->GetProperty()->GetSize());
+	TestEqual(TEXT("The exposed property should be updated after behaviour"), OutColorValue, SecColorValue);
+	
+	// 6.4 Check test function, it should call the function with default arguments and increase int test value
+	TestEqual(TEXT("The after calling test function TestInt should be =+ 1"), TestIntValue + 1, TestObject->TestInt);
+ 
+	// 6.5 Execute and test string controller 
+	StrController->ExecuteBehaviours();
+	FMemory::Memcpy(&OutColorValue, RCProp1ValuePtr, RCProp1->GetProperty()->GetSize());
+	TestEqual(TEXT("The exposed property should be updated after behaviour"), OutColorValue, StringControllerColorValue);
+	
+	// 7. Remove Actions
+	int32 ActionNum = FloatControllerBehaviour->GetNumActions();
+	FloatControllerBehaviour->ActionContainer->RemoveAction(FloatControllerBehaviourAction1);
+	FloatControllerBehaviour->ActionContainer->RemoveAction(RCProp1->GetId());
+	TestEqual(TEXT("After remove 2 action the count should be =-2"), FloatControllerBehaviour->GetNumActions(), ActionNum - 2);
+	FloatControllerBehaviour->ActionContainer->EmptyActions();
+	TestEqual(TEXT("After exmpty actions the count should be = 0"), FloatControllerBehaviour->GetNumActions(), 0);
+ 
+	// 7. Remove Actions
+	ActionNum = StrControllerBehaviour->GetNumActions();
+	StrControllerBehaviour->ActionContainer->RemoveAction(StrControllerBehaviourAction);
+	TestEqual(TEXT("After empty actions the count should be = 0"), StrControllerBehaviour->GetNumActions(), 0);
+ 
+	// 8. Remove Behaviour by ID and by Pointer
+	const int32 BehaviourNum = FloatController->Behaviours.Num();
+	FloatController->RemoveBehaviour(FloatControllerBehaviour);
+	FloatController->RemoveBehaviour(IntControllerBehaviour->Id);
+	TestEqual(TEXT("After empty 2 behaviours the count should be =-2"), FloatController->Behaviours.Num(), BehaviourNum - 2);
+	FloatController->EmptyBehaviours();
+	TestEqual(TEXT("After empty actions the count should be = 0"), FloatController->Behaviours.Num(), 0);
+ 
+	// 9. Remove Controllers
+	const int32 NumPropertiesBeforeRemove = Preset->ControllerContainer->GetNumVirtualProperties();
+	float OutFloatValue = 0.f;
+	float OutFloatValue1 = 0.f;
+	FloatController->GetValueFloat(OutFloatValue);
+	Preset->ControllerContainer->RemoveProperty(FloatController1->PropertyName);
+	FloatController->GetValueFloat(OutFloatValue1);
+	TestEqual(TEXT("After remove properties the old value should be the same"), OutFloatValue1, OutFloatValue);
+	Preset->ControllerContainer->RemoveProperty(FloatController->PropertyName);
+	Preset->ControllerContainer->RemoveProperty(BoolController->PropertyName);
+	Preset->ControllerContainer->RemoveProperty(IntController->PropertyName);
+	TestEqual(TEXT("After remove 4 properties the count should be =-4"), Preset->ControllerContainer->GetNumVirtualProperties(), NumPropertiesBeforeRemove - 4);
+	Preset->ControllerContainer->Reset();
+	TestEqual(TEXT("After empty controllers could should be equal 0"), Preset->ControllerContainer->GetNumVirtualProperties(), 0);
+ 
+	return true;
+}
+ 
+PRAGMA_ENABLE_OPTIMIZATION
+ 
\ No newline at end of file
diff --git a/Engine/Plugins/VirtualProduction/RemoteControl/Source/RemoteControlLogic/Public/Action/RCAction.h b/Engine/Plugins/VirtualProduction/RemoteControl/Source/RemoteControlLogic/Public/Action/RCAction.h
new file mode 100644
--- /dev/null
+++ b/Engine/Plugins/VirtualProduction/RemoteControl/Source/RemoteControlLogic/Public/Action/RCAction.h
@@ -0,0 +1,34 @@
+﻿// Copyright Epic Games, Inc. All Rights Reserved.
+
+#pragma once
+
+#include "CoreMinimal.h"
+#include "RCAction.generated.h"
+
+class URemoteControlPreset;
+
+/**
+ * Base class for remote control action 
+ */
+UCLASS()
+class REMOTECONTROLLOGIC_API URCAction : public UObject
+{
+	GENERATED_BODY()
+
+public:
+	/** Execute action */
+	virtual void Execute() const {}
+
+public:
+	/** Exposed Property or Function field Id*/
+	UPROPERTY()
+	FGuid ExposedFieldId;
+
+	/** Action Id */
+	UPROPERTY()
+	FGuid Id;
+
+	/** Reference to preset */
+	UPROPERTY()
+	TWeakObjectPtr<URemoteControlPreset> PresetWeakPtr;
+};
diff --git a/Engine/Plugins/VirtualProduction/RemoteControl/Source/RemoteControlLogic/Public/Action/RCActionContainer.h b/Engine/Plugins/VirtualProduction/RemoteControl/Source/RemoteControlLogic/Public/Action/RCActionContainer.h
new file mode 100644
--- /dev/null
+++ b/Engine/Plugins/VirtualProduction/RemoteControl/Source/RemoteControlLogic/Public/Action/RCActionContainer.h
@@ -0,0 +1,64 @@
+﻿// Copyright Epic Games, Inc. All Rights Reserved.
+
+#pragma once
+
+#include "CoreMinimal.h"
+#include "UObject/Object.h"
+
+#include "RCActionContainer.generated.h"
+
+class URemoteControlPreset;
+class URCAction;
+class URCFunctionAction;
+class URCPropertyAction;
+
+struct FRemoteControlField;
+struct FRemoteControlFunction;
+struct FRemoteControlProperty;
+
+/**
+ * Container for created actions
+ */
+UCLASS()
+class REMOTECONTROLLOGIC_API URCActionContainer : public UObject
+{
+	GENERATED_BODY()
+
+public:
+	/** Execute container actions */
+	void ExecuteActions();
+
+	/** Add remote control property action  */
+	URCPropertyAction* AddAction(const TSharedPtr<FRemoteControlProperty> InRemoteControlProperty);
+
+	/** Add remote control property function */
+	URCFunctionAction* AddAction(const TSharedPtr<FRemoteControlFunction> InRemoteControlFunction);
+
+	/** Find Action by given exposed filed id */
+	URCAction* FindActionByFieldId(const FGuid InFieldId);
+
+	/** Find Action by given remote control field */
+	URCAction* FindActionByField(const TSharedPtr<FRemoteControlField> InRemoteControlField);
+	
+	/** Remove action by exposed field Id */
+	virtual int32 RemoveAction(const FGuid InExposedFieldId);
+
+	/** Remove Action by given action UObject */
+	virtual int32 RemoveAction(URCAction* InAction);
+
+	/** Empty action set */
+	void EmptyActions();
+
+public:
+	/** Set of actions */
+	UPROPERTY()
+	TSet<TObjectPtr<URCAction>> Actions;
+	
+	/** Set of child action container */
+	UPROPERTY()
+	TSet<TObjectPtr<URCActionContainer>> ActionContainers;
+
+	/** Reference to Preset */
+	UPROPERTY()
+	TWeakObjectPtr<URemoteControlPreset> PresetWeakPtr;
+};
diff --git a/Engine/Plugins/VirtualProduction/RemoteControl/Source/RemoteControlLogic/Public/Action/RCFunctionAction.h b/Engine/Plugins/VirtualProduction/RemoteControl/Source/RemoteControlLogic/Public/Action/RCFunctionAction.h
new file mode 100644
--- /dev/null
+++ b/Engine/Plugins/VirtualProduction/RemoteControl/Source/RemoteControlLogic/Public/Action/RCFunctionAction.h
@@ -0,0 +1,20 @@
+﻿// Copyright Epic Games, Inc. All Rights Reserved.
+
+#pragma once
+
+#include "RCAction.h"
+#include "RCFunctionAction.generated.h"
+
+/**
+ * Function Action which is executes bound exposed function 
+ */
+UCLASS()
+class REMOTECONTROLLOGIC_API URCFunctionAction : public URCAction
+{
+	GENERATED_BODY()
+
+public:
+	//~ Begin URCAction Interface
+	virtual void Execute() const override;
+	//~ End URCAction Interface
+};
diff --git a/Engine/Plugins/VirtualProduction/RemoteControl/Source/RemoteControlLogic/Public/Action/RCPropertyAction.h b/Engine/Plugins/VirtualProduction/RemoteControl/Source/RemoteControlLogic/Public/Action/RCPropertyAction.h
new file mode 100644
--- /dev/null
+++ b/Engine/Plugins/VirtualProduction/RemoteControl/Source/RemoteControlLogic/Public/Action/RCPropertyAction.h
@@ -0,0 +1,29 @@
+﻿// Copyright Epic Games, Inc. All Rights Reserved.
+
+#pragma once
+
+#include "RCAction.h"
+#include "RCPropertyAction.generated.h"
+
+class URCVirtualPropertySelfContainer;
+
+/**
+ * Property Action which is set property value from the action to exposed property
+ */
+UCLASS()
+class REMOTECONTROLLOGIC_API URCPropertyAction : public URCAction
+{
+	GENERATED_BODY()
+
+public:
+	URCPropertyAction();
+	
+	//~ Begin URCAction interface
+	virtual void Execute() const override;
+	//~ End URCAction interface
+
+public:
+	/** Virtual Property Container */
+	UPROPERTY()
+	TObjectPtr<URCVirtualPropertySelfContainer> PropertySelfContainer = nullptr;
+};
diff --git a/Engine/Plugins/VirtualProduction/RemoteControl/Source/RemoteControlLogic/Public/Behaviour/Builtin/RCBehaviourIsEqualNode.h b/Engine/Plugins/VirtualProduction/RemoteControl/Source/RemoteControlLogic/Public/Behaviour/Builtin/RCBehaviourIsEqualNode.h
new file mode 100644
--- /dev/null
+++ b/Engine/Plugins/VirtualProduction/RemoteControl/Source/RemoteControlLogic/Public/Behaviour/Builtin/RCBehaviourIsEqualNode.h
@@ -0,0 +1,32 @@
+﻿// Copyright Epic Games, Inc. All Rights Reserved.
+
+#pragma once
+
+#include "Behaviour/RCBehaviourNode.h"
+#include "RCBehaviourIsEqualNode.generated.h"
+
+/**
+ * Compares given property value with controller property value
+ */
+UCLASS()
+class REMOTECONTROLLOGIC_API URCBehaviourIsEqualNode : public URCBehaviourNode
+{
+	GENERATED_BODY()
+
+public:
+	
+	URCBehaviourIsEqualNode();
+
+	//~ Begin URCBehaviourNode interface
+	UFUNCTION(BlueprintNativeEvent, BlueprintCallable, Category = "Remote Control Behaviour")
+	bool Execute(URCBehaviour* InBehaviour) const override;
+
+	UFUNCTION(BlueprintNativeEvent, BlueprintCallable, Category = "Remote Control Behaviour")
+	bool IsSupported(URCBehaviour* InBehaviour) const override;
+
+	UFUNCTION(BlueprintNativeEvent, BlueprintCallable, Category = "Remote Control Behaviour")
+	void OnPassed(URCBehaviour* InBehaviour) const;
+
+	UClass* GetBehaviourClass() const override;
+	//~ End URCBehaviourNode interface
+};
diff --git a/Engine/Plugins/VirtualProduction/RemoteControl/Source/RemoteControlLogic/Public/Behaviour/Builtin/RCBehaviourOnEmptyNode.h b/Engine/Plugins/VirtualProduction/RemoteControl/Source/RemoteControlLogic/Public/Behaviour/Builtin/RCBehaviourOnEmptyNode.h
new file mode 100644
--- /dev/null
+++ b/Engine/Plugins/VirtualProduction/RemoteControl/Source/RemoteControlLogic/Public/Behaviour/Builtin/RCBehaviourOnEmptyNode.h
@@ -0,0 +1,27 @@
+﻿// Copyright Epic Games, Inc. All Rights Reserved.
+
+#pragma once
+
+#include "Behaviour/RCBehaviourNode.h"
+#include "RCBehaviourOnEmptyNode.generated.h"
+
+/**
+ * Check is the controlled property is empty
+ */
+UCLASS()
+class REMOTECONTROLLOGIC_API URCBehaviourOnEmptyNode : public URCBehaviourNode
+{
+	GENERATED_BODY()
+
+public:
+
+	URCBehaviourOnEmptyNode();
+
+	//~ Begin URCBehaviourNode interface
+	UFUNCTION(BlueprintNativeEvent, BlueprintCallable, Category = "Remote Control Behaviour")
+	bool Execute(URCBehaviour* InBehaviour) const override;
+
+	UFUNCTION(BlueprintNativeEvent, BlueprintCallable, Category = "Remote Control Behaviour")
+	bool IsSupported(URCBehaviour* InBehaviour) const override;
+	//~ End URCBehaviourNode interface
+};
diff --git a/Engine/Plugins/VirtualProduction/RemoteControl/Source/RemoteControlLogic/Public/Behaviour/Builtin/RCBehaviourSetValueNode.h b/Engine/Plugins/VirtualProduction/RemoteControl/Source/RemoteControlLogic/Public/Behaviour/Builtin/RCBehaviourSetValueNode.h
new file mode 100644
--- /dev/null
+++ b/Engine/Plugins/VirtualProduction/RemoteControl/Source/RemoteControlLogic/Public/Behaviour/Builtin/RCBehaviourSetValueNode.h
@@ -0,0 +1,29 @@
+﻿// Copyright Epic Games, Inc. All Rights Reserved.
+
+#pragma once
+
+#include "Behaviour/RCBehaviourNode.h"
+#include "RCBehaviourSetValueNode.generated.h"
+
+/**
+ * Simple pass behaviour which is returns execute true all the time.
+ *
+ * That is needed for simple the execute the actions without any logic
+ */
+UCLASS()
+class REMOTECONTROLLOGIC_API URCBehaviourSetValueNode : public URCBehaviourNode
+{
+	GENERATED_BODY()
+
+public:
+
+	URCBehaviourSetValueNode();
+
+	//~ Begin URCBehaviourNode interface
+	UFUNCTION(BlueprintNativeEvent, BlueprintCallable, Category = "Remote Control Behaviour")
+	bool Execute(URCBehaviour* InBehaviour) const override;
+
+	UFUNCTION(BlueprintNativeEvent, BlueprintCallable, Category = "Remote Control Behaviour")
+	bool IsSupported(URCBehaviour* InBehaviour) const override;
+	//~ End URCBehaviourNode interface
+};
diff --git a/Engine/Plugins/VirtualProduction/RemoteControl/Source/RemoteControlLogic/Public/Behaviour/RCBehaviour.h b/Engine/Plugins/VirtualProduction/RemoteControl/Source/RemoteControlLogic/Public/Behaviour/RCBehaviour.h
new file mode 100644
--- /dev/null
+++ b/Engine/Plugins/VirtualProduction/RemoteControl/Source/RemoteControlLogic/Public/Behaviour/RCBehaviour.h
@@ -0,0 +1,103 @@
+﻿// Copyright Epic Games, Inc. All Rights Reserved.
+
+#pragma once
+
+#include "CoreMinimal.h"
+#include "UObject/Object.h"
+#include "Templates/SubclassOf.h"
+#include "RCBehaviour.generated.h"
+
+class URCBehaviourNode;
+class URCController;
+class URemoteControlPreset;
+class URCActionContainer;
+class URCVirtualProperty;
+
+/**
+ * Base class for remote control behaviour.
+ *
+ * Behaviour is container for:
+ * - Set of behaviour conditions
+ * - And associated actions which should be executed if that is passed the behaviour
+ *
+ * Behaviour can be extended in Blueprints or CPP classes
+ */
+UCLASS(BlueprintType)
+class REMOTECONTROLLOGIC_API URCBehaviour : public UObject
+{
+	GENERATED_BODY()
+
+public:
+	URCBehaviour();
+
+	/** Initialize behaviour functionality */
+	virtual void Initialize() {}
+
+	/** Execute the behaviour */
+	virtual void Execute();
+
+	/** Get number of action associated with behaviour */
+	int32 GetNumActions() const;
+
+	/**
+	 * Return blueprint class associated with behaviour if exists
+	 */
+	UClass* GetOverrideBehaviourBlueprintClass() const;
+
+#if WITH_EDITORONLY_DATA
+	/**
+	 * Return blueprint instance associated with behaviour if exists
+	 */
+	UBlueprint* GetBlueprint() const;
+#endif
+
+	/**
+	 * Set blueprint class for this behaviour
+	 */
+	void SetOverrideBehaviourBlueprintClass(UBlueprint* InBlueprint);
+
+#if WITH_EDITOR
+	/** Get Display Name for this Behaviour */
+	const FText& GetDisplayName();
+
+	/** Get Description for this Behaviour */
+	const FText& GetBehaviorDescription();
+#endif
+
+private:
+	/**
+	 * It created the node if it called first time
+	 * If BehaviourNodeClass changes it creates new instance
+	 * Or just return cached one
+	 */
+	URCBehaviourNode* GetBehaviourNode();
+
+public:
+	/** Associated cpp behaviour */
+	UPROPERTY()
+	TSubclassOf<URCBehaviourNode> BehaviourNodeClass;
+
+	/** Class path to associated blueprint class behaviour */
+	UPROPERTY()
+	FSoftClassPath OverrideBehaviourBlueprintClassPath;
+
+	/** Behaviour Id */
+	UPROPERTY()
+	FGuid Id;
+
+	/** Action container which is associated with current behaviour */
+	UPROPERTY()
+	TObjectPtr<URCActionContainer> ActionContainer;
+
+	/** Reference to controller virtual property with this behaviour */
+	UPROPERTY()
+	TWeakObjectPtr<URCController> ControllerWeakPtr;
+
+private:
+	/** Cached behaviour node class */
+	TSubclassOf<UObject> CachedBehaviourNodeClass;
+
+	/** Cached behaviour node */
+	UPROPERTY(Instanced)
+	TObjectPtr<URCBehaviourNode> CachedBehaviourNode;
+};
diff --git a/Engine/Plugins/VirtualProduction/RemoteControl/Source/RemoteControlLogic/Public/Behaviour/RCBehaviourBlueprintNode.h b/Engine/Plugins/VirtualProduction/RemoteControl/Source/RemoteControlLogic/Public/Behaviour/RCBehaviourBlueprintNode.h
new file mode 100644
--- /dev/null
+++ b/Engine/Plugins/VirtualProduction/RemoteControl/Source/RemoteControlLogic/Public/Behaviour/RCBehaviourBlueprintNode.h
@@ -0,0 +1,28 @@
+﻿// Copyright Epic Games, Inc. All Rights Reserved.
+
+#pragma once
+
+#include "Behaviour/RCBehaviourNode.h"
+#include "RCBehaviourBlueprintNode.generated.h"
+
+/**
+ * Base blueprint class for behaviour node which holds the logic to execute behaviour 
+ */
+UCLASS(Abstract, Blueprintable)
+class REMOTECONTROLLOGIC_API URCBehaviourBlueprintNode : public URCBehaviourNode
+{
+	GENERATED_BODY()
+
+public:
+	UFUNCTION(BlueprintNativeEvent, BlueprintCallable, Category = "Remote Control Behaviour")
+	bool IsSupported(URCBehaviour* InBehaviour) const override;
+	
+	UFUNCTION(BlueprintNativeEvent, BlueprintCallable, Category = "Remote Control Behaviour")
+	bool Execute(URCBehaviour* InBehaviour) const override;
+
+	UFUNCTION(BlueprintNativeEvent, BlueprintCallable, Category = "Remote Control Behaviour")
+	void PreExecute(URCBehaviour* InBehaviour) const override;
+
+	UFUNCTION(BlueprintNativeEvent, BlueprintCallable, Category = "Remote Control Behaviour")
+	void OnPassed(URCBehaviour* InBehaviour) const override;
+};
diff --git a/Engine/Plugins/VirtualProduction/RemoteControl/Source/RemoteControlLogic/Public/Behaviour/RCBehaviourNode.h b/Engine/Plugins/VirtualProduction/RemoteControl/Source/RemoteControlLogic/Public/Behaviour/RCBehaviourNode.h
new file mode 100644
--- /dev/null
+++ b/Engine/Plugins/VirtualProduction/RemoteControl/Source/RemoteControlLogic/Public/Behaviour/RCBehaviourNode.h
@@ -0,0 +1,69 @@
+﻿// Copyright Epic Games, Inc. All Rights Reserved.
+
+#pragma once
+
+#include "CoreMinimal.h"
+#include "UObject/Object.h"
+#include "RCBehaviourNode.generated.h"
+
+class URCBehaviour;
+enum class EPropertyBagPropertyType : uint8;
+
+/**
+ * Base class for behaviour node which holds the logic to execute behaviour 
+ */
+UCLASS(Abstract, BlueprintType, Blueprintable, EditInlineNew)
+class REMOTECONTROLLOGIC_API URCBehaviourNode : public UObject
+{
+	GENERATED_BODY()
+
+public:
+	/**
+	 * Whether node supported based on input handle
+	 * 
+	 * @param InBehaviour Remote Control Behaviour
+	 */
+	virtual bool IsSupported(URCBehaviour* InBehaviour) const;
+	
+	/**
+	 * Called before behaviour event
+	 * 
+	 * @param InBehaviour Remote Control Behaviour
+	 */
+	virtual void PreExecute(URCBehaviour* InBehaviour) const;
+
+	/**
+	 * Execute behaviour event
+	 * 
+	 * @param InBehaviour Remote Control Behaviour
+	 */
+	virtual bool Execute(URCBehaviour* InBehaviour) const;
+
+	/**
+	 * Called if executed with true
+	 * 
+	 * @param InBehaviour Remote Control Behaviour 
+	 */
+	virtual void OnPassed(URCBehaviour* InBehaviour) const;
+	
+	/** Get the Behaviour class for this Behaviour node */
+	virtual UClass* GetBehaviourClass() const;
+
+	/**
+	 * User-friendly display name for this Behavior, displayed in Action Panel's Header.
+	 * Custom behavior blueprints can set this value in defaults to update the display name
+	 */
+	UPROPERTY(EditDefaultsOnly, Category = "Remote Control Behavior")
+	FText DisplayName;
+
+	/**
+	 * Detailed description of what this behavior does, displayed in Action Panel's Header.
+	 * Custom behavior blueprints can set this value in defaults to update the display name
+	 */
+	UPROPERTY(EditDefaultsOnly, Category = "Remote Control Behavior")
+	FText BehaviorDescription;
+
+protected:
+	/** Return supported check callback */
+	static TFunction<bool(const TPair<EPropertyBagPropertyType, UObject*>)> GetIsSupportedCallback(URCBehaviour* InBehaviour);
+};
diff --git a/Engine/Plugins/VirtualProduction/RemoteControl/Source/RemoteControlLogic/Public/Behaviour/RCIsEqualBehaviour.h b/Engine/Plugins/VirtualProduction/RemoteControl/Source/RemoteControlLogic/Public/Behaviour/RCIsEqualBehaviour.h
new file mode 100644
--- /dev/null
+++ b/Engine/Plugins/VirtualProduction/RemoteControl/Source/RemoteControlLogic/Public/Behaviour/RCIsEqualBehaviour.h
@@ -0,0 +1,30 @@
+﻿// Copyright Epic Games, Inc. All Rights Reserved.
+
+#pragma once
+
+#include "RCBehaviour.h"
+#include "RCIsEqualBehaviour.generated.h"
+
+class URCVirtualPropertySelfContainer;
+class URCUserDefinedStruct;
+
+/**
+ * Custom behaviour for Is Equal logic
+ */
+UCLASS()
+class REMOTECONTROLLOGIC_API URCIsEqualBehaviour : public URCBehaviour
+{
+	GENERATED_BODY()
+
+public:
+	URCIsEqualBehaviour();
+	
+	//~ Begin URCBehaviour interface
+	virtual void Initialize() override;
+	//~ End URCBehaviour interface
+
+public:
+	/** Pointer to property container */
+	UPROPERTY()
+	TObjectPtr<URCVirtualPropertySelfContainer> PropertySelfContainer;
+};
diff --git a/Engine/Plugins/VirtualProduction/RemoteControl/Source/RemoteControlLogic/Public/Controller/RCController.h b/Engine/Plugins/VirtualProduction/RemoteControl/Source/RemoteControlLogic/Public/Controller/RCController.h
new file mode 100644
--- /dev/null
+++ b/Engine/Plugins/VirtualProduction/RemoteControl/Source/RemoteControlLogic/Public/Controller/RCController.h
@@ -0,0 +1,42 @@
+﻿// Copyright Epic Games, Inc. All Rights Reserved.
+
+#pragma once
+
+#include "RCVirtualProperty.h"
+#include "RCController.generated.h"
+
+class URCBehaviour;
+class URCBehaviourNode;
+
+/**
+ * Remote Control Controller. Container for Behaviours and Actions
+ */
+UCLASS(BlueprintType)
+class REMOTECONTROLLOGIC_API URCController : public URCVirtualPropertyInContainer
+{
+	GENERATED_BODY()
+
+public:
+	/** Create and add behaviour to behaviour set */
+	virtual URCBehaviour* AddBehaviour(TSubclassOf<URCBehaviourNode> InBehaviourNodeClass);
+
+	/** Create new behaviour */
+	virtual URCBehaviour* CreateBehaviour(TSubclassOf<URCBehaviourNode> InBehaviourNodeClass);
+
+	/** Remove the behaviour by behaviour UObject pointer */
+	virtual int32 RemoveBehaviour(URCBehaviour* InBehaviour);
+
+	/** Remove the behaviour by behaviour id */
+	virtual int32 RemoveBehaviour(const FGuid InBehaviourId);
+
+	/** Removes all behaviours. */
+	virtual void EmptyBehaviours();
+
+	/** Execute all behaviours for this controller. */
+	virtual void ExecuteBehaviours();
+
+public:
+	/** Set of the behaviours */
+	UPROPERTY()
+	TSet<TObjectPtr<URCBehaviour>> Behaviours;
+};
diff --git a/Engine/Plugins/VirtualProduction/RemoteControl/Source/RemoteControlLogic/Public/Controller/RCControllerContainer.h b/Engine/Plugins/VirtualProduction/RemoteControl/Source/RemoteControlLogic/Public/Controller/RCControllerContainer.h
new file mode 100644
--- /dev/null
+++ b/Engine/Plugins/VirtualProduction/RemoteControl/Source/RemoteControlLogic/Public/Controller/RCControllerContainer.h
@@ -0,0 +1,28 @@
+﻿// Copyright Epic Games, Inc. All Rights Reserved.
+
+#pragma once
+
+#include "RCVirtualPropertyContainer.h"
+#include "RCControllerContainer.generated.h"
+
+class URCActionContainer;
+
+/**
+ * Controller Container which holds all virtual controller properties 
+ */
+UCLASS()
+class REMOTECONTROLLOGIC_API URCControllerContainer : public URCVirtualPropertyContainerBase
+{
+	GENERATED_BODY()
+	
+private:
+#if WITH_EDITOR
+	/** Delegate when object changed */
+	virtual void OnModifyPropertyValue(const FPropertyChangedEvent& PropertyChangedEvent) override;
+#endif
+
+protected:
+	/** Controller Container could holds the set of shared Action containers which is holds some action but independent from the behaviour */
+	UPROPERTY()
+	TSet<TObjectPtr<URCActionContainer>> SharedActionContainers;
+};
diff --git a/Engine/Plugins/VirtualProduction/RemoteControl/Source/RemoteControlLogic/RemoteControlLogic.Build.cs b/Engine/Plugins/VirtualProduction/RemoteControl/Source/RemoteControlLogic/RemoteControlLogic.Build.cs
new file mode 100644
--- /dev/null
+++ b/Engine/Plugins/VirtualProduction/RemoteControl/Source/RemoteControlLogic/RemoteControlLogic.Build.cs
@@ -0,0 +1,35 @@
+// Copyright Epic Games, Inc. All Rights Reserved.
+
+using UnrealBuildTool;
+
+public class RemoteControlLogic : ModuleRules
+{
+	public RemoteControlLogic(ReadOnlyTargetRules Target) : base(Target)
+	{
+		PublicDependencyModuleNames.AddRange(
+			new string[] {
+				"Core",
+				"CoreUObject",
+			}
+		);
+
+		PrivateDependencyModuleNames.AddRange(
+			new string[] {
+				"Cbor",
+				"Engine",
+				"RemoteControl",
+				"Serialization",
+				"StructUtils"
+			}
+		);
+		
+		if (Target.bBuildEditor)
+		{
+			PrivateDependencyModuleNames.AddRange(
+				new string[] {
+					"UnrealEd",
+				}
+			);
+		}
+	}
+}
