[RemoteControl] Fix bound object pointers always returning a stale pointer after the map is reloaded

diff --git a/Engine/Plugins/VirtualProduction/RemoteControl/Source/RemoteControl/Private/RemoteControlBinding.cpp b/Engine/Plugins/VirtualProduction/RemoteControl/Source/RemoteControl/Private/RemoteControlBinding.cpp
--- a/Engine/Plugins/VirtualProduction/RemoteControl/Source/RemoteControl/Private/RemoteControlBinding.cpp
+++ b/Engine/Plugins/VirtualProduction/RemoteControl/Source/RemoteControl/Private/RemoteControlBinding.cpp
@@ -130,8 +130,8 @@ void URemoteControlLevelDependantBinding::SetBoundObject(const TSoftObjectPtr<UO
 	{
 		UObject* EditorObject = FindObjectInCounterpartWorld(InObject.Get(), ECounterpartWorldTarget::Editor);
 		ULevel* OuterLevel = EditorObject->GetTypedOuter<ULevel>();
-		BoundObjectMap.FindOrAdd(OuterLevel) = EditorObject;
-		SubLevelSelectionMap.FindOrAdd(OuterLevel->GetWorld()) = OuterLevel;
+		BoundObjectMapByPath.FindOrAdd(OuterLevel) = EditorObject;
+		SubLevelSelectionMapByPath.FindOrAdd(OuterLevel->GetWorld()) = OuterLevel;
 		
 		Name = EditorObject->GetName();
 		
@@ -160,19 +160,19 @@ void URemoteControlLevelDependantBinding::InitializeForNewLevel()
 		if (UWorld* CurrentWorld = GetCurrentWorld())
 		{
 			ULevel* CurrentLevel = CurrentWorld->PersistentLevel;
-			if (BoundObjectMap.Contains(CurrentLevel))
+			if (BoundObjectMapByPath.Contains(CurrentLevel))
 			{
 				// If there is already a binding for this level, don't overwrite it.
 				return;
 			}
 
-			if (TSoftObjectPtr<UObject>* BoundObjectPtr = BoundObjectMap.Find(LevelWithLastSuccessfulResolve))
+			if (TSoftObjectPtr<UObject>* BoundObjectPtr = BoundObjectMapByPath.Find(LevelWithLastSuccessfulResolve.ToSoftObjectPath()))
 			{
 				// Try to find the bound object in the current world by reparenting its path to the current level.
 				FSoftObjectPath NewPath = CurrentLevel->GetPathName() + BoundObjectPtr->ToSoftObjectPath().GetSubPathString().RightChop(PersistentLevelStrLength);
 				if (NewPath.ResolveObject())
 				{
-					BoundObjectMap.Add(CurrentLevel, TSoftObjectPtr<UObject>{NewPath});
+					BoundObjectMapByPath.Add(CurrentLevel, TSoftObjectPtr<UObject>{NewPath});
 				}
 			}
 		}
@@ -181,13 +181,13 @@ void URemoteControlLevelDependantBinding::InitializeForNewLevel()
 
 void URemoteControlLevelDependantBinding::UnbindObject(const TSoftObjectPtr<UObject>& InBoundObject)
 {
-	for (auto It = BoundObjectMap.CreateIterator(); It; ++It)
+	for (auto It = BoundObjectMapByPath.CreateIterator(); It; ++It)
 	{
 		if (It.Value() == InBoundObject)
 		{
 			if (InBoundObject)
 			{
-				SubLevelSelectionMap.Remove(InBoundObject->GetWorld());
+				SubLevelSelectionMapByPath.Remove(InBoundObject->GetWorld());
 			}
 
 			It.RemoveCurrent();
@@ -224,12 +224,12 @@ UObject* URemoteControlLevelDependantBinding::Resolve() const
 
 bool URemoteControlLevelDependantBinding::IsValid() const
 {
-	return BoundObjectMap.Num() > 0;
+	return BoundObjectMapByPath.Num() > 0;
 }
 
 bool URemoteControlLevelDependantBinding::IsBound(const TSoftObjectPtr<UObject>& Object) const
 {
-	for (const TPair<TSoftObjectPtr<ULevel>, TSoftObjectPtr<UObject>>& Pair : BoundObjectMap)
+	for (const TPair<FSoftObjectPath, TSoftObjectPtr<UObject>>& Pair : BoundObjectMapByPath)
 	{
 		if (Pair.Value == Object)
 		{
@@ -246,13 +246,16 @@ bool URemoteControlLevelDependantBinding::PruneDeletedObjects()
 	{
 		if (UWorld* World = GetCurrentWorld())
 		{
-			if (TSoftObjectPtr<ULevel> LastLevelForBinding = SubLevelSelectionMap.FindRef(World))
+			if (const TSoftObjectPtr<ULevel>* LastLevelForBinding = SubLevelSelectionMapByPath.Find(World))
 			{
-				if (!BoundObjectMap.FindRef(LastLevelForBinding))
+				const FSoftObjectPath& LevelPath = LastLevelForBinding->ToSoftObjectPath();
+				const TSoftObjectPtr<UObject>* BoundObject = BoundObjectMapByPath.Find(LevelPath);
+
+				if (!BoundObject || !*BoundObject)
 				{
 					Modify();
-					BoundObjectMap.Remove(LastLevelForBinding);
-					SubLevelSelectionMap.Remove(World);
+					BoundObjectMapByPath.Remove(LevelPath);
+					SubLevelSelectionMapByPath.Remove(World);
 					return true;
 				}
 			}
@@ -262,28 +265,71 @@ bool URemoteControlLevelDependantBinding::PruneDeletedObjects()
 	return false;
 }
 
+void URemoteControlLevelDependantBinding::PostLoad()
+{
+	Super::PostLoad();
+
+#if WITH_EDITOR
+	PRAGMA_DISABLE_DEPRECATION_WARNINGS
+
+	if (BoundObjectMap_DEPRECATED.Num() > 0)
+	{
+		for (const TPair<TSoftObjectPtr<ULevel>, TSoftObjectPtr<UObject>>& Pair : BoundObjectMap_DEPRECATED)
+		{
+			BoundObjectMapByPath.Add(Pair.Key.ToSoftObjectPath(), Pair.Value);
+		}
+		BoundObjectMap_DEPRECATED.Empty();
+	}
+
+	if (SubLevelSelectionMap_DEPRECATED.Num() > 0)
+	{
+		for (const TPair<TSoftObjectPtr<UWorld>, TSoftObjectPtr<ULevel>>& Pair : SubLevelSelectionMap_DEPRECATED)
+		{
+			SubLevelSelectionMapByPath.Add(Pair.Key.ToSoftObjectPath(), Pair.Value);
+		}
+		SubLevelSelectionMap_DEPRECATED.Empty();
+	}
+
+	PRAGMA_ENABLE_DEPRECATION_WARNINGS
+#endif
+}
+
 TSoftObjectPtr<UObject> URemoteControlLevelDependantBinding::ResolveForCurrentWorld() const
 {
+	// Note that we use Find rather than FindRef throughout this function. This lets us dereference
+	// the BoundObjectMapByPath/SubLevelSelectionMapByPath's values directly so that their internal weak pointers
+	// are updated when stale (rather than dereferencing copies and leaving stale pointers in the originals).
+
 	if (UWorld* World = GetCurrentWorld())
 	{
 		// Try finding the object using the sub level selection map first.
-		if (TSoftObjectPtr<ULevel> LastBindingLevel = SubLevelSelectionMap.FindRef(World))
+		if (TSoftObjectPtr<ULevel>* LastBindingLevel = SubLevelSelectionMapByPath.Find(World))
 		{
-			return BoundObjectMap.FindRef(LastBindingLevel);
+			if (const TSoftObjectPtr<UObject>* ObjectPtr = BoundObjectMapByPath.Find(LastBindingLevel->ToSoftObjectPath()))
+			{
+				if (*ObjectPtr)
+				{
+					return *ObjectPtr;
+				}
+			}
 		}
+		
 		// Resort to old method where we use the first level we find in the bound object map,
 		// and add an entry in the sub level selection map.
 		for (auto LevelIt = World->GetLevelIterator(); LevelIt; ++LevelIt)
 		{
-			TSoftObjectPtr<ULevel> WeakLevel = *LevelIt;
-			if (TSoftObjectPtr<UObject> ObjectPtr = BoundObjectMap.FindRef(WeakLevel))
+			const TSoftObjectPtr<ULevel>& WeakLevel = *LevelIt;
+			if (const TSoftObjectPtr<UObject>* ObjectPtr = BoundObjectMapByPath.Find(WeakLevel.ToSoftObjectPath()))
 			{
-				SubLevelSelectionMap.FindOrAdd(World) = WeakLevel;
-
-				return ObjectPtr;
+				if (*ObjectPtr)
+				{
+					SubLevelSelectionMapByPath.FindOrAdd(World) = WeakLevel;
+					return *ObjectPtr;
+				}
 			}
 		}
 	}
+
 	return nullptr;
 }
 
@@ -356,7 +402,7 @@ UWorld* URemoteControlLevelDependantBinding::GetCurrentWorld()
 
 void URemoteControlLevelDependantBinding::SetBoundObject(const TSoftObjectPtr<ULevel>& Level, const TSoftObjectPtr<UObject>& BoundObject)
 {
-	BoundObjectMap.FindOrAdd(Level) = BoundObject;
+	BoundObjectMapByPath.FindOrAdd(Level.ToSoftObjectPath()) = BoundObject;
 	const FSoftObjectPath& Path = BoundObject.ToSoftObjectPath();
 	const FString& SubPath = Path.GetSubPathString();
 	FString LeftPart;
diff --git a/Engine/Plugins/VirtualProduction/RemoteControl/Source/RemoteControl/Private/RemoteControlPreset.cpp b/Engine/Plugins/VirtualProduction/RemoteControl/Source/RemoteControl/Private/RemoteControlPreset.cpp
--- a/Engine/Plugins/VirtualProduction/RemoteControl/Source/RemoteControl/Private/RemoteControlPreset.cpp
+++ b/Engine/Plugins/VirtualProduction/RemoteControl/Source/RemoteControl/Private/RemoteControlPreset.cpp
@@ -956,15 +956,15 @@ URemoteControlBinding* URemoteControlPreset::FindMatchingBinding(const URemoteCo
 				continue;
 			}
 
-			TSoftObjectPtr<ULevel> CurrentWorldLevel = LevelDependantBindingIt->SubLevelSelectionMap.FindRef(InObject->GetWorld());
+			TSoftObjectPtr<ULevel> CurrentWorldLevel = LevelDependantBindingIt->SubLevelSelectionMapByPath.FindRef(InObject->GetWorld());
 
 			bool bSameBoundObjectMap = true;
 			// Check if the binding it has the same bound object map except for the current level.
-			for (const TPair<TSoftObjectPtr<ULevel>, TSoftObjectPtr<UObject>>& Pair : InLevelDependingBinding->BoundObjectMap)
+			for (const TPair<FSoftObjectPath, TSoftObjectPtr<UObject>>& Pair : InLevelDependingBinding->BoundObjectMapByPath)
 			{
-				if (Pair.Key != CurrentWorldLevel)
+				if (Pair.Key != CurrentWorldLevel.ToSoftObjectPath())
 				{
-					TSoftObjectPtr<UObject>* BoundObject = LevelDependantBindingIt->BoundObjectMap.Find(Pair.Key);
+					TSoftObjectPtr<UObject>* BoundObject = LevelDependantBindingIt->BoundObjectMapByPath.Find(Pair.Key);
 					if (BoundObject)
 					{
 						if (*BoundObject != Pair.Value)
diff --git a/Engine/Plugins/VirtualProduction/RemoteControl/Source/RemoteControl/Private/RemoteControlPresetRebindingManager.cpp b/Engine/Plugins/VirtualProduction/RemoteControl/Source/RemoteControl/Private/RemoteControlPresetRebindingManager.cpp
--- a/Engine/Plugins/VirtualProduction/RemoteControl/Source/RemoteControl/Private/RemoteControlPresetRebindingManager.cpp
+++ b/Engine/Plugins/VirtualProduction/RemoteControl/Source/RemoteControl/Private/RemoteControlPresetRebindingManager.cpp
@@ -596,7 +596,7 @@ TArray<FGuid> FRemoteControlPresetRebindingManager::RebindAllEntitiesUnderSameAc
 		if (Binding && Binding->IsA<URemoteControlLevelDependantBinding>())
 		{
 			URemoteControlLevelDependantBinding* LDBinding = CastChecked<URemoteControlLevelDependantBinding>(Binding);
-			FString Path = LDBinding->BoundObjectMap.FindRef(LDBinding->LevelWithLastSuccessfulResolve).ToString();
+			FString Path = LDBinding->BoundObjectMapByPath.FindRef(LDBinding->LevelWithLastSuccessfulResolve.ToSoftObjectPath()).ToString();
 			static const TCHAR* PersistentLevel = TEXT("PersistentLevel.");
 			const int32 PersistentLevelStringLength = 16;
 
diff --git a/Engine/Plugins/VirtualProduction/RemoteControl/Source/RemoteControl/Public/RemoteControlBinding.h b/Engine/Plugins/VirtualProduction/RemoteControl/Source/RemoteControl/Public/RemoteControlBinding.h
--- a/Engine/Plugins/VirtualProduction/RemoteControl/Source/RemoteControl/Public/RemoteControlBinding.h
+++ b/Engine/Plugins/VirtualProduction/RemoteControl/Source/RemoteControl/Public/RemoteControlBinding.h
@@ -148,6 +148,8 @@ public:
 	virtual bool PruneDeletedObjects() override;
 	//~ Begin URemoteControlBinding Interface
 
+	virtual void PostLoad() override;
+
 	/**
 	 *	Set the bound object by specifying the level it resides in.
 	 *	@Note Useful is you want to set the bound object without loading the level/object.
@@ -178,19 +180,27 @@ private:
 	void InitializeBindingContext(UObject* InObject) const;
 
 private:
+#if WITH_EDITORONLY_DATA
+	UPROPERTY(meta = (DeprecatedProperty, DeprecationMessage = "BoundObjectMap is deprecated, please use BoundObjectMapByPath instead."))
+	TMap<TSoftObjectPtr<ULevel>, TSoftObjectPtr<UObject>> BoundObjectMap_DEPRECATED;
+
+	UPROPERTY(meta = (DeprecatedProperty, DeprecationMessage = "SubLevelSelectionMap is deprecated, please use SubLevelSelectionMapByPath instead."))
+	mutable TMap<TSoftObjectPtr<UWorld>, TSoftObjectPtr<ULevel>> SubLevelSelectionMap_DEPRECATED;
+#endif
+
 	/**
 	 *	The map bound objects with their level as key.
 	 */
 	UPROPERTY()
-	TMap<TSoftObjectPtr<ULevel>, TSoftObjectPtr<UObject>> BoundObjectMap;
+	TMap<FSoftObjectPath, TSoftObjectPtr<UObject>> BoundObjectMapByPath;
 
 	/**
 	 * Keeps track of which sublevel was last used when binding in a particular world.
-	 * Used in the case where a binding points to objects that end up in the same world but in different sublevels, 
-	 * this ensures that we know which object was last 
+	 * Used in the case where a binding points to objects that end up in the same world but in different sublevels,
+	 * this ensures that we know which object was last
 	 */
 	UPROPERTY()
-	mutable TMap<TSoftObjectPtr<UWorld>, TSoftObjectPtr<ULevel>> SubLevelSelectionMap;
+	mutable TMap<FSoftObjectPath, TSoftObjectPtr<ULevel>> SubLevelSelectionMapByPath;
 
 	/**
 	 * Caches the last level that had a successful resolve.
