[WebRemoteControl] Send property update messages for the same type as a single batch
Remove tracked actors when world is destroyed and ignore tracking for actors outside the current editor world

diff --git a/Engine/Plugins/VirtualProduction/RemoteControl/Source/WebRemoteControl/Private/WebSocketMessageHandler.cpp b/Engine/Plugins/VirtualProduction/RemoteControl/Source/WebRemoteControl/Private/WebSocketMessageHandler.cpp
--- a/Engine/Plugins/VirtualProduction/RemoteControl/Source/WebRemoteControl/Private/WebSocketMessageHandler.cpp
+++ b/Engine/Plugins/VirtualProduction/RemoteControl/Source/WebRemoteControl/Private/WebSocketMessageHandler.cpp
@@ -225,6 +225,23 @@ namespace WebSocketMessageHandlerMiscUtils
 		return nullptr;
 #endif
 	}
+
+	uint64 GetPresetPropertyClassId(URemoteControlPreset* Preset, const FGuid& PropertyId)
+	{
+		FRCObjectReference ObjectRef;
+		if (TSharedPtr<FRemoteControlProperty> RCProperty = Preset->GetExposedEntity<FRemoteControlProperty>(PropertyId).Pin())
+		{
+			if (RCProperty->IsBound())
+			{
+				if (const FProperty* Property = RCProperty->GetProperty())
+				{
+					return Property->GetClass()->GetId();
+				}
+			}
+		}
+
+		return 0;
+	}
 }
 
 FWebSocketMessageHandler::FWebSocketMessageHandler(FRCWebSocketServer* InServer, const FGuid& InActingClientId)
@@ -312,6 +329,7 @@ void FWebSocketMessageHandler::UnregisterRoutes(FWebRemoteControlModule* WebRemo
 		GEngine->OnLevelActorAdded().Remove(OnActorAddedHandle);
 		GEngine->OnLevelActorDeleted().Remove(OnActorDeletedHandle);
 		GEngine->OnLevelActorListChanged().Remove(OnActorListChangedHandle);
+		GEngine->OnLevelActorListChanged().Remove(OnWorldDestroyedHandle);
 	}
 #endif
 
@@ -337,6 +355,7 @@ void FWebSocketMessageHandler::RegisterActorHandlers()
 		OnActorAddedHandle = GEngine->OnLevelActorAdded().AddRaw(this, &FWebSocketMessageHandler::OnActorAdded);
 		OnActorDeletedHandle = GEngine->OnLevelActorDeleted().AddRaw(this, &FWebSocketMessageHandler::OnActorDeleted);
 		OnActorListChangedHandle = GEngine->OnLevelActorListChanged().AddRaw(this, &FWebSocketMessageHandler::OnActorListChanged);
+		OnWorldDestroyedHandle = GEngine->OnWorldDestroyed().AddRaw(this, &FWebSocketMessageHandler::OnWorldDestroyed);
 	}
 #endif
 }
@@ -651,21 +670,27 @@ void FWebSocketMessageHandler::ProcessChangedProperties()
 		// Each client will have a custom payload that doesnt contain the events it triggered.
 		for (const TPair<FGuid, TSet<FGuid>>& ClientToEventsPair : Entry.Value)
 		{
-			// This should be improved in the future, we create one message per modified property to avoid
-			// sending a list of non-uniform properties (ie. Color, Transform), ideally these should be grouped by underlying
-			// property class. See UE-139683
+			// Categorize modified properties by type so we don't try to put multiple types of data in a single request
+			TMap<uint64, TSet<FGuid>> PropertyIdsByType;
 			for (const FGuid& Id : ClientToEventsPair.Value)
+			{
+				const uint64 ClassId = WebSocketMessageHandlerMiscUtils::GetPresetPropertyClassId(Preset, Id);
+				PropertyIdsByType.FindOrAdd(ClassId).Emplace(Id);
+			}
+
+			// Send a property change event for each property type
+			for (const TPair<uint64, TSet<FGuid>>& ClassToEventsPair : PropertyIdsByType)
 			{
 				const int64* ClientSequenceNumber = ClientSequenceNumbers.Find(ClientToEventsPair.Key);
 				const int64 SequenceNumber = ClientSequenceNumber ? *ClientSequenceNumber : DefaultSequenceNumber;
 
 				TArray<uint8> WorkingBuffer;
-				if (ClientToEventsPair.Value.Num() && WritePropertyChangeEventPayload(Preset, { Id }, SequenceNumber, WorkingBuffer))
+				if (ClientToEventsPair.Value.Num() && WritePropertyChangeEventPayload(Preset, { ClassToEventsPair.Value }, SequenceNumber, WorkingBuffer))
 				{
 					TArray<uint8> Payload;
 					WebRemoteControlUtils::ConvertToUTF8(WorkingBuffer, Payload);
 					Server->Send(ClientToEventsPair.Key, Payload);
-				}	
+				}
 			}
 		}
 	}
@@ -1334,6 +1359,14 @@ bool FWebSocketMessageHandler::WriteActorPropertyChangePayload(URemoteControlPre
 
 void FWebSocketMessageHandler::OnActorAdded(AActor* Actor)
 {
+#if WITH_EDITOR
+	const UWorld* World = WebSocketMessageHandlerMiscUtils::GetEditorWorld();
+	if (Actor->GetWorld() != World)
+	{
+		return;
+	}
+#endif
+
 	// Array of classes this actor is a child of and which are being watched by a client
 	TArray<TWeakObjectPtr<UClass>, TInlineAllocator<8>> WatchedClasses;
 
@@ -1471,6 +1504,15 @@ void FWebSocketMessageHandler::OnActorListChanged()
 	}
 }
 
+void FWebSocketMessageHandler::OnWorldDestroyed(UWorld* World)
+{
+	// Remove any actors in this world
+	for (AActor* Actor : TActorRange<AActor>(World))
+	{
+		OnActorDeleted(Actor);
+	}
+}
+
 void FWebSocketMessageHandler::OnObjectPropertyChanged(UObject* Object, struct FPropertyChangedEvent& Event)
 {
 	// We only care about name changes
diff --git a/Engine/Plugins/VirtualProduction/RemoteControl/Source/WebRemoteControl/Private/WebSocketMessageHandler.h b/Engine/Plugins/VirtualProduction/RemoteControl/Source/WebRemoteControl/Private/WebSocketMessageHandler.h
--- a/Engine/Plugins/VirtualProduction/RemoteControl/Source/WebRemoteControl/Private/WebSocketMessageHandler.h
+++ b/Engine/Plugins/VirtualProduction/RemoteControl/Source/WebRemoteControl/Private/WebSocketMessageHandler.h
@@ -194,6 +194,11 @@ private:
 	 */
 	void OnActorDeleted(AActor* Actor);
 
+	/**
+	 * Called when a world is destroyed.
+	 */
+	void OnWorldDestroyed(UWorld* World);
+
 	/**
 	 * Called when an untracked change to the actor list happens in the editor.
 	 */
@@ -327,6 +332,9 @@ private:
 	/** Handle for when the list of actors changes. */
 	FDelegateHandle OnActorListChangedHandle;
 
+	/** Handle for when a world is destroyed. */
+	FDelegateHandle OnWorldDestroyedHandle;
+
 	/**
 	 * Actors that we are actively watching to send events to subscribers.
 	 * The key is not a weak pointer, so it shouldn't be accessed in case it's stale. Use the value instead.
