ProtocolBindingVM pointer cleanup

diff --git a/Engine/Plugins/VirtualProduction/RemoteControl/Source/RemoteControlProtocolWidgets/Private/ViewModels/ProtocolBindingViewModel.cpp b/Engine/Plugins/VirtualProduction/RemoteControl/Source/RemoteControlProtocolWidgets/Private/ViewModels/ProtocolBindingViewModel.cpp
--- a/Engine/Plugins/VirtualProduction/RemoteControl/Source/RemoteControlProtocolWidgets/Private/ViewModels/ProtocolBindingViewModel.cpp
+++ b/Engine/Plugins/VirtualProduction/RemoteControl/Source/RemoteControlProtocolWidgets/Private/ViewModels/ProtocolBindingViewModel.cpp
@@ -584,7 +584,7 @@ bool FProtocolBindingViewModel::IsValid(FText& OutMessage)
 	check(IsValid());
 
 	EValidity Result = EValidity::Ok;
-	const FString* AdditionalMessage = nullptr;
+	TUniquePtr<FString> AdditionalMessage = MakeUnique<FString>();
 
 	if (Ranges.Num() <= 1)
 	{
@@ -646,7 +646,7 @@ bool FProtocolBindingViewModel::IsValid(FText& OutMessage)
 
 		if (DuplicateBindings.Len() > 0)
 		{
-			AdditionalMessage = new FString(MoveTemp(DuplicateBindings));
+			AdditionalMessage.Reset(new FString(MoveTemp(DuplicateBindings)));
 		}
 	}
 
