RCP support for property setters

diff --git a/Engine/Plugins/VirtualProduction/RemoteControl/Source/RemoteControl/Private/RemoteControlModule.cpp b/Engine/Plugins/VirtualProduction/RemoteControl/Source/RemoteControl/Private/RemoteControlModule.cpp
--- a/Engine/Plugins/VirtualProduction/RemoteControl/Source/RemoteControl/Private/RemoteControlModule.cpp
+++ b/Engine/Plugins/VirtualProduction/RemoteControl/Source/RemoteControl/Private/RemoteControlModule.cpp
@@ -290,6 +290,42 @@ namespace RemoteControlSetterUtils
 		InOutArgs.Call.ParamStruct = MoveTemp(InFunctionArguments);
 	}
 
+	void CreateRCCall(FConvertToFunctionCallArgs& InOutArgs, FProperty* InPropertyWithSetter, FStructOnScope&& InFunctionArguments, FRCInterceptionPayload& OutPayload)
+	{
+		ensure(InFunctionArguments.GetStruct());
+		
+		// Create the output payload for interception purposes.
+		FMemoryWriter Writer{OutPayload.Payload};
+		FCborStructSerializerBackend WriterBackend{Writer, EStructSerializerBackendFlags::Default};
+		FStructSerializer::Serialize(InFunctionArguments.GetStructMemory(), *const_cast<UStruct*>(InFunctionArguments.GetStruct()), WriterBackend, FStructSerializerPolicies());
+		OutPayload.Type = ERCPayloadType::Cbor;
+
+		InOutArgs.Call.bGenerateTransaction = InOutArgs.ObjectReference.Access == ERCAccess::WRITE_TRANSACTION_ACCESS ? true : false;
+		InOutArgs.Call.CallRef.PropertyWithSetter = InPropertyWithSetter;
+		InOutArgs.Call.CallRef.Object = InOutArgs.ObjectReference.Object;
+		InOutArgs.Call.ParamStruct = MoveTemp(InFunctionArguments);
+		
+		InOutArgs.Call.ParamData.SetNumUninitialized(InPropertyWithSetter->GetSize());
+		InPropertyWithSetter->InitializeValue(InOutArgs.Call.ParamData.GetData());
+		const void* ValuePtr = InPropertyWithSetter->ContainerPtrToValuePtr<uint8>(InOutArgs.Call.ParamStruct.GetStructMemory());
+		InPropertyWithSetter->CopyCompleteValue(InOutArgs.Call.ParamData.GetData(), ValuePtr);
+	}
+
+	void CreateRCCall(FConvertToFunctionCallArgs& InOutArgs, FProperty* InPropertyWithSetter, FStructOnScope&& InFunctionArguments)
+	{
+		ensure(InFunctionArguments.GetStruct());
+		
+		InOutArgs.Call.bGenerateTransaction = InOutArgs.ObjectReference.Access == ERCAccess::WRITE_TRANSACTION_ACCESS ? true : false;
+		InOutArgs.Call.CallRef.PropertyWithSetter = InPropertyWithSetter;
+		InOutArgs.Call.CallRef.Object = InOutArgs.ObjectReference.Object;
+		InOutArgs.Call.ParamStruct = MoveTemp(InFunctionArguments);
+
+		InOutArgs.Call.ParamData.SetNumUninitialized(InPropertyWithSetter->GetSize());
+		InPropertyWithSetter->InitializeValue(InOutArgs.Call.ParamData.GetData());
+		const void* ValuePtr = InPropertyWithSetter->ContainerPtrToValuePtr<uint8>(InFunctionArguments.GetStructMemory());
+		InPropertyWithSetter->CopyCompleteValue(InOutArgs.Call.ParamData.GetData(), ValuePtr);
+	}
+
 	/** Create the payload to pass to be passed to a property's setter function. */
 	TOptional<FStructOnScope> CreateSetterFunctionPayload(UFunction* InSetterFunction, FConvertToFunctionCallArgs& InOutArgs)
 	{
@@ -345,6 +381,47 @@ namespace RemoteControlSetterUtils
 		return OptionalArgsOnScope;
 	}
 
+	/** Create the payload to pass to be passed to a property's native setter. */
+	TOptional<FStructOnScope> CreateSetterFunctionPayload(FProperty* InPropertyWithSetter, FConvertToFunctionCallArgs& InOutArgs)
+	{
+		TOptional<FStructOnScope> OptionalArgsOnScope;
+
+		FStructOnScope ArgsOnScope{ InOutArgs.ObjectReference.ContainerType.Get() };
+
+		bool bSuccess = false;
+		if (InOutArgs.ValuePtrOverride)
+		{
+			// We already have a pointer directly to the data we want, so copy the data into the function arguments
+			const UStruct* ArgsStruct = ArgsOnScope.GetStruct();
+
+			check(ArgsStruct);
+
+			uint8* SetterArgData = InOutArgs.ObjectReference.Property->ContainerPtrToValuePtr<uint8>(ArgsOnScope.GetStructMemory());
+			InOutArgs.ObjectReference.Property->CopyCompleteValue(SetterArgData, InOutArgs.ValuePtrOverride);
+
+			bSuccess = true;
+		}
+		else
+		{
+			// Data needs to be deserialized from the passed reader backend
+
+			// First put the complete property value from the object in the struct on scope
+			// in case the user only a part of the incoming structure (ie. Providing only { "x": 2 } in the case of a vector.
+			const uint8* ContainerAddress = InOutArgs.ObjectReference.Property->ContainerPtrToValuePtr<uint8>(InOutArgs.ObjectReference.ContainerAdress);
+			InOutArgs.ObjectReference.Property->CopyCompleteValue(ArgsOnScope.GetStructMemory(), ContainerAddress);
+
+			// Deserialize on top of the setter argument
+			bSuccess = FStructDeserializer::Deserialize((void*)ArgsOnScope.GetStructMemory(), *const_cast<UStruct*>(ArgsOnScope.GetStruct()), InOutArgs.ReaderBackend, FStructDeserializerPolicies());
+		}
+
+		if (bSuccess)
+		{
+			OptionalArgsOnScope = MoveTemp(ArgsOnScope);
+		}
+
+		return OptionalArgsOnScope;
+	}
+
 	bool ConvertModificationToFunctionCall(FConvertToFunctionCallArgs& InOutArgs, UFunction* InSetterFunction, FRCInterceptionPayload& OutPayload)
 	{
 		if (TOptional<FStructOnScope> ArgsOnScope = CreateSetterFunctionPayload(InSetterFunction, InOutArgs))
@@ -367,6 +444,28 @@ namespace RemoteControlSetterUtils
 		return false;
 	}
 
+	bool ConvertModificationToFunctionCall(FConvertToFunctionCallArgs& InOutArgs, FProperty* InPropertyWithSetter, FRCInterceptionPayload& OutPayload)
+	{
+		if (TOptional<FStructOnScope> ArgsOnScope = CreateSetterFunctionPayload(InPropertyWithSetter, InOutArgs))
+		{
+			CreateRCCall(InOutArgs, InPropertyWithSetter, MoveTemp(*ArgsOnScope), OutPayload);
+			return true;
+		}
+
+		return false;
+	}
+
+	bool ConvertModificationToFunctionCall(FConvertToFunctionCallArgs& InOutArgs, FProperty* InPropertyWithSetter)
+	{
+		if (TOptional<FStructOnScope> ArgsOnScope = CreateSetterFunctionPayload(InPropertyWithSetter, InOutArgs))
+		{
+			CreateRCCall(InOutArgs, InPropertyWithSetter, MoveTemp(*ArgsOnScope));
+			return true;
+		}
+
+		return false;
+	}
+
 	bool ConvertModificationToFunctionCall(FConvertToFunctionCallArgs& InOutArgs)
 	{
 		TRACE_CPUPROFILER_EVENT_SCOPE(RemoteControlSetterUtils::ConvertModificationToFunctionCall);
@@ -375,7 +474,11 @@ namespace RemoteControlSetterUtils
 			return false;
 		}
 
-		if (UFunction* SetterFunction = RemoteControlPropertyUtilities::FindSetterFunction(InOutArgs.ObjectReference.Property.Get(), InOutArgs.ObjectReference.Object->GetClass()))
+		if (InOutArgs.ObjectReference.Property->HasSetter())
+		{
+			return ConvertModificationToFunctionCall(InOutArgs, InOutArgs.ObjectReference.Property.Get());
+		}
+		else if (UFunction* SetterFunction = RemoteControlPropertyUtilities::FindSetterFunction(InOutArgs.ObjectReference.Property.Get(), InOutArgs.ObjectReference.Object->GetClass()))
 		{
 			return ConvertModificationToFunctionCall(InOutArgs, SetterFunction);
 		}
@@ -391,6 +494,10 @@ namespace RemoteControlSetterUtils
 			return false;
 		}
 
+		if (InArgs.ObjectReference.Property->HasSetter())
+		{
+			return ConvertModificationToFunctionCall(InArgs, InArgs.ObjectReference.Property.Get(), OutPayload);
+		}
 		if (UFunction* SetterFunction = RemoteControlPropertyUtilities::FindSetterFunction(InArgs.ObjectReference.Property.Get(), InArgs.ObjectReference.Object->GetClass()))
 		{
 			return ConvertModificationToFunctionCall(InArgs, SetterFunction, OutPayload);
@@ -599,7 +706,14 @@ bool FRemoteControlModule::InvokeCall(FRCCall& InCall, ERCPayloadType InPayloadT
 		FEditorScriptExecutionGuard ScriptGuard;
 		if (ensureAlways(InCall.CallRef.Object.IsValid()))
 		{
-		InCall.CallRef.Object->ProcessEvent(InCall.CallRef.Function.Get(), InCall.ParamStruct.GetStructMemory());
+			if(InCall.CallRef.PropertyWithSetter.IsValid())
+			{
+				InCall.CallRef.PropertyWithSetter->CallSetter(InCall.CallRef.Object.Get(), InCall.ParamData.GetData());
+			}
+			else
+			{
+				InCall.CallRef.Object->ProcessEvent(InCall.CallRef.Function.Get(), InCall.ParamStruct.GetStructMemory());
+			}
 		}
 
 #if WITH_EDITOR
@@ -1500,7 +1614,7 @@ bool FRemoteControlModule::PropertyModificationShouldUseSetter(UObject* Object,
 		return false;
 	}
 
-	return !!RemoteControlPropertyUtilities::FindSetterFunction(Property, Object->GetClass());
+	return Property->HasSetter() || !!RemoteControlPropertyUtilities::FindSetterFunction(Property, Object->GetClass());
 }
 
 bool FRemoteControlModule::DeserializeDeltaModificationData(const FRCObjectReference& ObjectAccess, IStructDeserializerBackend& Backend, ERCModifyOperation Operation, TArray<uint8>& OutData)
diff --git a/Engine/Plugins/VirtualProduction/RemoteControl/Source/RemoteControl/Private/Tests/RemoteControlDeltaAPITestData.h b/Engine/Plugins/VirtualProduction/RemoteControl/Source/RemoteControl/Private/Tests/RemoteControlDeltaAPITestData.h
--- a/Engine/Plugins/VirtualProduction/RemoteControl/Source/RemoteControl/Private/Tests/RemoteControlDeltaAPITestData.h
+++ b/Engine/Plugins/VirtualProduction/RemoteControl/Source/RemoteControl/Private/Tests/RemoteControlDeltaAPITestData.h
@@ -33,12 +33,17 @@ public:
 		return GET_MEMBER_NAME_CHECKED(URemoteControlDeltaAPITestObject, Int32WithSetterValue);
 	}
 
+	static const FName GetFloatWithSetterValuePropertyName()
+	{
+		return GET_MEMBER_NAME_CHECKED(URemoteControlDeltaAPITestObject, FloatWithSetterValue);
+	}
+
 	UFUNCTION(BlueprintGetter)
 	int32 GetInt32WithSetterValue() { return Int32WithSetterValue; }
 
 	UFUNCTION(BlueprintSetter)
 	void SetInt32WithSetterValue(const int32 NewValue) { Int32WithSetterValue = NewValue; }
-
+	
 	UPROPERTY(EditAnywhere, Category = "RC")
 	int32 Int32Value = Int32ValueDefault;
 
@@ -57,4 +62,17 @@ public:
 private:
 	UPROPERTY(EditAnywhere, Category = "RC", BlueprintGetter = GetInt32WithSetterValue, BlueprintSetter = SetInt32WithSetterValue)
 	int32 Int32WithSetterValue;
-};
\ No newline at end of file
+
+	/** Uses Setter - not BlueprintSetter */
+	UPROPERTY(EditAnywhere, Setter, Category = "RC")
+	float FloatWithSetterValue = FloatValueDefault;
+
+public:
+	float GetFloatWithSetterValue() { return FloatWithSetterValue; }
+	void SetFloatWithSetterValue(const float& InValue);
+};
+
+inline void URemoteControlDeltaAPITestObject::SetFloatWithSetterValue(const float& InValue)
+{
+	FloatWithSetterValue = InValue;
+}
diff --git a/Engine/Plugins/VirtualProduction/RemoteControl/Source/RemoteControl/Private/Tests/RemoteControlDeltaAPITests.spec.cpp b/Engine/Plugins/VirtualProduction/RemoteControl/Source/RemoteControl/Private/Tests/RemoteControlDeltaAPITests.spec.cpp
--- a/Engine/Plugins/VirtualProduction/RemoteControl/Source/RemoteControl/Private/Tests/RemoteControlDeltaAPITests.spec.cpp
+++ b/Engine/Plugins/VirtualProduction/RemoteControl/Source/RemoteControl/Private/Tests/RemoteControlDeltaAPITests.spec.cpp
@@ -64,6 +64,7 @@ void FRemoteControlDeltaAPISpec::ResetTestObject()
 	TestObject->VectorValue = URemoteControlDeltaAPITestObject::VectorValueDefault;
 	TestObject->ColorValue = URemoteControlDeltaAPITestObject::ColorValueDefault;
 	TestObject->SetInt32WithSetterValue(URemoteControlDeltaAPITestObject::Int32ValueDefault);
+	TestObject->SetFloatWithSetterValue(URemoteControlDeltaAPITestObject::FloatValueDefault);
 
 	TestObject->StructValue.ColorValue = FRemoteControlDeltaAPITestStruct::ColorValueDefault;
 }
@@ -160,7 +161,7 @@ void FRemoteControlDeltaAPISpec::Define()
 			TestEqual("Z value", TestObject->VectorValue.Z, URemoteControlDeltaAPITestObject::VectorValueDefault.Z);
 		});
 
-		It("should add property with a getter/setter", [this]
+		It("should add property with a blueprint getter/setter", [this]
 		{
 			FProperty* ValueProperty = FindFProperty<FProperty>(URemoteControlDeltaAPITestObject::StaticClass(), URemoteControlDeltaAPITestObject::GetInt32WithSetterValuePropertyName());
 
@@ -173,7 +174,20 @@ void FRemoteControlDeltaAPISpec::Define()
 			TestEqual("Combined value", TestObject->GetInt32WithSetterValue(), URemoteControlDeltaAPITestObject::Int32ValueDefault + AddedValue);
 		});
 
-		It("should add property with a getter/setter and no intercept payload", [this]
+		It("should add property with a getter/setter", [this]
+		{
+			FProperty* ValueProperty = FindFProperty<FProperty>(URemoteControlDeltaAPITestObject::StaticClass(), URemoteControlDeltaAPITestObject::GetFloatWithSetterValuePropertyName());
+
+			const float AddedValue = 4.0f;
+			TestObject->SetFloatWithSetterValue(AddedValue);
+
+			SerializeElementAndReset(ValueProperty);
+			ApplyOperation(ValueProperty, ERCModifyOperation::ADD);
+
+			TestEqual("Combined value", TestObject->GetFloatWithSetterValue(), URemoteControlDeltaAPITestObject::FloatValueDefault + AddedValue);
+		});
+
+		It("should add property with a blueprint getter/setter and no intercept payload", [this]
 		{
 			FProperty* ValueProperty = FindFProperty<FProperty>(URemoteControlDeltaAPITestObject::StaticClass(), URemoteControlDeltaAPITestObject::GetInt32WithSetterValuePropertyName());
 
@@ -279,7 +293,7 @@ void FRemoteControlDeltaAPISpec::Define()
 			TestEqual("Combined Z value", TestObject->VectorValue.Z, URemoteControlDeltaAPITestObject::VectorValueDefault.Z - 6.7f);
 		});
 
-		It("should subtract property with a getter/setter", [this]
+		It("should subtract property with a blueprint getter/setter", [this]
 		{
 			FProperty* ValueProperty = FindFProperty<FProperty>(URemoteControlDeltaAPITestObject::StaticClass(), URemoteControlDeltaAPITestObject::GetInt32WithSetterValuePropertyName());
 
@@ -292,7 +306,7 @@ void FRemoteControlDeltaAPISpec::Define()
 			TestEqual("Combined value", TestObject->GetInt32WithSetterValue(), URemoteControlDeltaAPITestObject::Int32ValueDefault - SubtractedValue);
 		});
 
-		It("should subtract property with a getter/setter and no intercept payload", [this]
+		It("should subtract property with a blueprint getter/setter and no intercept payload", [this]
 		{
 			FProperty* ValueProperty = FindFProperty<FProperty>(URemoteControlDeltaAPITestObject::StaticClass(), URemoteControlDeltaAPITestObject::GetInt32WithSetterValuePropertyName());
 
@@ -398,7 +412,7 @@ void FRemoteControlDeltaAPISpec::Define()
 			TestEqual("Z value", TestObject->VectorValue.Z, URemoteControlDeltaAPITestObject::VectorValueDefault.Z);
 		});
 
-		It("should multiply property with a getter/setter", [this]
+		It("should multiply property with a blueprint getter/setter", [this]
 		{
 			FProperty* ValueProperty = FindFProperty<FProperty>(URemoteControlDeltaAPITestObject::StaticClass(), URemoteControlDeltaAPITestObject::GetInt32WithSetterValuePropertyName());
 
@@ -411,7 +425,7 @@ void FRemoteControlDeltaAPISpec::Define()
 			TestEqual("Combined value", TestObject->GetInt32WithSetterValue(), URemoteControlDeltaAPITestObject::Int32ValueDefault * Multiplier);
 		});
 
-		It("should multiply property with a getter/setter and no intercept payload", [this]
+		It("should multiply property with a blueprint getter/setter and no intercept payload", [this]
 		{
 			FProperty* ValueProperty = FindFProperty<FProperty>(URemoteControlDeltaAPITestObject::StaticClass(), URemoteControlDeltaAPITestObject::GetInt32WithSetterValuePropertyName());
 
@@ -530,7 +544,7 @@ void FRemoteControlDeltaAPISpec::Define()
 			TestEqual("Z value", TestObject->VectorValue.Z, URemoteControlDeltaAPITestObject::VectorValueDefault.Z);
 		});
 
-		It("should multiply property with a getter/setter", [this]
+		It("should multiply property with a blueprint getter/setter", [this]
 		{
 			FProperty* ValueProperty = FindFProperty<FProperty>(URemoteControlDeltaAPITestObject::StaticClass(), URemoteControlDeltaAPITestObject::GetInt32WithSetterValuePropertyName());
 
@@ -543,7 +557,7 @@ void FRemoteControlDeltaAPISpec::Define()
 			TestEqual("Combined value", TestObject->GetInt32WithSetterValue(), URemoteControlDeltaAPITestObject::Int32ValueDefault / Divisor);
 		});
 
-		It("should multiply property with a getter/setter and no intercept payload", [this]
+		It("should multiply property with a blueprint getter/setter and no intercept payload", [this]
 		{
 			FProperty* ValueProperty = FindFProperty<FProperty>(URemoteControlDeltaAPITestObject::StaticClass(), URemoteControlDeltaAPITestObject::GetInt32WithSetterValuePropertyName());
 
diff --git a/Engine/Plugins/VirtualProduction/RemoteControl/Source/RemoteControl/Public/IRemoteControlModule.h b/Engine/Plugins/VirtualProduction/RemoteControl/Source/RemoteControl/Public/IRemoteControlModule.h
--- a/Engine/Plugins/VirtualProduction/RemoteControl/Source/RemoteControl/Public/IRemoteControlModule.h
+++ b/Engine/Plugins/VirtualProduction/RemoteControl/Source/RemoteControl/Public/IRemoteControlModule.h
@@ -49,19 +49,30 @@ struct FRCCallReference
 	FRCCallReference()
 		: Object(nullptr)
 		, Function(nullptr)
+		, PropertyWithSetter(nullptr)
 	{}
 
 	bool IsValid() const
 	{
-		return Object.IsValid() && Function.IsValid();
+		/** The object and either the function or property must be valid */
+		return Object.IsValid() && (Function.IsValid() || PropertyWithSetter.IsValid());
 	}
 
 	TWeakObjectPtr<UObject> Object; 
 	TWeakObjectPtr<UFunction> Function;
+	TWeakFieldPtr<FProperty> PropertyWithSetter;
 	
 	friend uint32 GetTypeHash(const FRCCallReference& CallRef)
 	{
-		return CallRef.IsValid() ? HashCombine(GetTypeHash(CallRef.Object), GetTypeHash(CallRef.Function)) : 0;
+		if(!CallRef.IsValid())
+		{
+			return 0;
+		}
+		
+		return HashCombine(GetTypeHash(CallRef.Object),
+			CallRef.PropertyWithSetter.IsValid()
+			? GetTypeHash(CallRef.PropertyWithSetter)
+			: GetTypeHash(CallRef.Function));
 	}
 };
 
@@ -72,11 +83,14 @@ struct FRCCall
 {
 	bool IsValid() const
 	{
-		return CallRef.IsValid() && ParamStruct.IsValid();
+		return CallRef.IsValid() && (ParamStruct.IsValid() || !ParamData.IsEmpty());
 	}
 
 	FRCCallReference CallRef;
+	// Payload for UFunctiom
 	FStructOnScope ParamStruct;
+	// Payload for native function
+	TArray<uint8> ParamData;
 	bool bGenerateTransaction = false;
 };
 
