RCPropertyContainer setter fix

diff --git a/Engine/Plugins/VirtualProduction/RemoteControl/Source/RemoteControlCommon/Private/RCPropertyContainer.cpp b/Engine/Plugins/VirtualProduction/RemoteControl/Source/RemoteControlCommon/Private/RCPropertyContainer.cpp
--- a/Engine/Plugins/VirtualProduction/RemoteControl/Source/RemoteControlCommon/Private/RCPropertyContainer.cpp
+++ b/Engine/Plugins/VirtualProduction/RemoteControl/Source/RemoteControlCommon/Private/RCPropertyContainer.cpp
@@ -424,11 +424,6 @@ namespace PropertyContainers
 			}
 		}
 
-		if(NewProperty)
-		{
-			
-		}
-
 		return NewProperty;
 	}
 #endif
diff --git a/Engine/Plugins/VirtualProduction/RemoteControl/Source/RemoteControlCommon/Public/RCPropertyContainer.h b/Engine/Plugins/VirtualProduction/RemoteControl/Source/RemoteControlCommon/Public/RCPropertyContainer.h
--- a/Engine/Plugins/VirtualProduction/RemoteControl/Source/RemoteControlCommon/Public/RCPropertyContainer.h
+++ b/Engine/Plugins/VirtualProduction/RemoteControl/Source/RemoteControlCommon/Public/RCPropertyContainer.h
@@ -22,19 +22,23 @@ public:
 
 	/** Sets the value as ValueType, Not necessarily valid if using incorrect ValueType. */
 	template <typename ValueType>
-	void SetValue(const ValueType& InValue)
+	typename TEnableIf<TIsPointer<ValueType>::Value, void>::Type
+	SetValue(const ValueType& InValue)
 	{
+		using DecayedValueType = typename TDecay<TRemovePointer<ValueType>>::Type;
+
 #if WITH_EDITOR
 		Modify();
 		FProperty* Property = GetValueProperty();
 		FEditPropertyChain EditChain;
 		EditChain.AddHead(Property);
 		PreEditChange(EditChain);
-#endif	
+#endif
 
+		DecayedValueType* ValuePtr = GetValueProperty()->ContainerPtrToValuePtr<DecayedValueType>(this);
 		GetValueProperty()->CopyCompleteValue(
-			GetValueProperty()->ContainerPtrToValuePtr<ValueType>(this),
-			&InValue);
+			(void*)ValuePtr,
+			(const DecayedValueType*)InValue);
 
 #if WITH_EDITOR
 		FPropertyChangedEvent EditPropertyChangeEvent(Property, EPropertyChangeType::ValueSet);
@@ -44,6 +48,13 @@ public:
 #endif
 	}
 
+	template <typename ValueType>
+	typename TEnableIf<!TIsPointer<ValueType>::Value, void>::Type
+	SetValue(const ValueType& InValue)
+	{
+		SetValue(&InValue);		
+	}
+	
 	/** Writes to the provided raw data pointer. Returns size for array, string, etc. */
 	SIZE_T GetValue(uint8* OutData);
 
