[WebRemoteControl] Add WebSocket routes to control when transactions are created

diff --git a/Engine/Plugins/VirtualProduction/RemoteControl/Source/RemoteControl/Private/RemoteControlFunctionLibrary.cpp b/Engine/Plugins/VirtualProduction/RemoteControl/Source/RemoteControl/Private/RemoteControlFunctionLibrary.cpp
--- a/Engine/Plugins/VirtualProduction/RemoteControl/Source/RemoteControl/Private/RemoteControlFunctionLibrary.cpp
+++ b/Engine/Plugins/VirtualProduction/RemoteControl/Source/RemoteControl/Private/RemoteControlFunctionLibrary.cpp
@@ -56,7 +56,7 @@ bool URemoteControlFunctionLibrary::ExposeActor(URemoteControlPreset* Preset, AA
 	return false;
 }
 
-bool URemoteControlFunctionLibrary::ApplyColorWheelDelta(UObject* TargetObject, const FString& PropertyName, const FColorWheelColor& DeltaValue, const FColorWheelColor& ReferenceColor)
+bool URemoteControlFunctionLibrary::ApplyColorWheelDelta(UObject* TargetObject, const FString& PropertyName, const FColorWheelColor& DeltaValue, const FColorWheelColor& ReferenceColor, bool bIsInteractive)
 {
 	if (!TargetObject)
 	{
@@ -102,12 +102,13 @@ bool URemoteControlFunctionLibrary::ApplyColorWheelDelta(UObject* TargetObject,
 
 #if WITH_EDITOR
 		TargetObject->PreEditChange(Property);
+		TargetObject->Modify();
 #endif
 
 		ColorProperty->SetValue_InContainer(TargetObject, &Color);
 
 #if WITH_EDITOR
-		FPropertyChangedEvent ChangeEvent(Property, EPropertyChangeType::ValueSet);
+		FPropertyChangedEvent ChangeEvent(Property, bIsInteractive ? EPropertyChangeType::Interactive : EPropertyChangeType::ValueSet);
 		TargetObject->PostEditChangeProperty(ChangeEvent);
 #endif
 
diff --git a/Engine/Plugins/VirtualProduction/RemoteControl/Source/RemoteControl/Private/RemoteControlInterceptionHelpers.h b/Engine/Plugins/VirtualProduction/RemoteControl/Source/RemoteControl/Private/RemoteControlInterceptionHelpers.h
--- a/Engine/Plugins/VirtualProduction/RemoteControl/Source/RemoteControl/Private/RemoteControlInterceptionHelpers.h
+++ b/Engine/Plugins/VirtualProduction/RemoteControl/Source/RemoteControl/Private/RemoteControlInterceptionHelpers.h
@@ -22,6 +22,9 @@ constexpr ERCAccess ToInternal(ERCIAccess Value)
 
 	case ERCIAccess::WRITE_TRANSACTION_ACCESS:
 		return ERCAccess::WRITE_TRANSACTION_ACCESS;
+
+	case ERCIAccess::WRITE_MANUAL_TRANSACTION_ACCESS:
+		return ERCAccess::WRITE_MANUAL_TRANSACTION_ACCESS;
 	}
 
 	return ERCAccess::NO_ACCESS;
@@ -82,6 +85,9 @@ constexpr ERCIAccess ToExternal(ERCAccess Value)
 
 	case ERCAccess::WRITE_TRANSACTION_ACCESS:
 		return ERCIAccess::WRITE_TRANSACTION_ACCESS;
+
+	case ERCAccess::WRITE_MANUAL_TRANSACTION_ACCESS:
+		return ERCIAccess::WRITE_MANUAL_TRANSACTION_ACCESS;
 	}
 
 	return ERCIAccess::NO_ACCESS;
diff --git a/Engine/Plugins/VirtualProduction/RemoteControl/Source/RemoteControl/Private/RemoteControlInterceptionProcessor.cpp b/Engine/Plugins/VirtualProduction/RemoteControl/Source/RemoteControl/Private/RemoteControlInterceptionProcessor.cpp
--- a/Engine/Plugins/VirtualProduction/RemoteControl/Source/RemoteControl/Private/RemoteControlInterceptionProcessor.cpp
+++ b/Engine/Plugins/VirtualProduction/RemoteControl/Source/RemoteControl/Private/RemoteControlInterceptionProcessor.cpp
@@ -104,7 +104,7 @@ void FRemoteControlInterceptionProcessor::InvokeCall(FRCIFunctionMetadata& InFun
 		
 		FRCCall Call;
 		Call.CallRef = MoveTemp(CallRef);
-		Call.bGenerateTransaction = InFunction.bGenerateTransaction;
+		Call.TransactionMode = InFunction.bGenerateTransaction ? ERCTransactionMode::AUTOMATIC : ERCTransactionMode::NONE;
 		Call.ParamStruct = FStructOnScope(FunctionArgs.GetStruct(), FunctionArgs.GetStructMemory());
 
 		IRemoteControlModule::Get().InvokeCall(Call);
diff --git a/Engine/Plugins/VirtualProduction/RemoteControl/Source/RemoteControl/Private/RemoteControlModule.cpp b/Engine/Plugins/VirtualProduction/RemoteControl/Source/RemoteControl/Private/RemoteControlModule.cpp
--- a/Engine/Plugins/VirtualProduction/RemoteControl/Source/RemoteControl/Private/RemoteControlModule.cpp
+++ b/Engine/Plugins/VirtualProduction/RemoteControl/Source/RemoteControl/Private/RemoteControlModule.cpp
@@ -185,7 +185,7 @@ namespace RemoteControlUtil
 	/** Returns whether the access is a write access regardless of if it generates a transaction. */
 	bool IsWriteAccess(ERCAccess Access)
 	{
-		return Access == ERCAccess::WRITE_ACCESS || Access == ERCAccess::WRITE_TRANSACTION_ACCESS;
+		return Access == ERCAccess::WRITE_ACCESS || Access == ERCAccess::WRITE_TRANSACTION_ACCESS || Access == ERCAccess::WRITE_MANUAL_TRANSACTION_ACCESS;
 	}
 
 	/** Helper function to check if a value is 0 without always calling FMath::IsNearlyZero, since that results in an ambiguous call for non-float values. */
@@ -274,7 +274,21 @@ namespace RemoteControlSetterUtils
 		FStructSerializer::Serialize(InFunctionArguments.GetStructMemory(), *const_cast<UStruct*>(InFunctionArguments.GetStruct()), WriterBackend, FStructSerializerPolicies());
 		OutPayload.Type = ERCPayloadType::Cbor;
 
-		InOutArgs.Call.bGenerateTransaction = InOutArgs.ObjectReference.Access == ERCAccess::WRITE_TRANSACTION_ACCESS ? true : false;
+		switch (InOutArgs.ObjectReference.Access)
+		{
+		case ERCAccess::WRITE_TRANSACTION_ACCESS:
+			InOutArgs.Call.TransactionMode = ERCTransactionMode::AUTOMATIC;
+			break;
+
+		case ERCAccess::WRITE_MANUAL_TRANSACTION_ACCESS:
+			InOutArgs.Call.TransactionMode = ERCTransactionMode::MANUAL;
+			break;
+
+		default:
+			InOutArgs.Call.TransactionMode = ERCTransactionMode::NONE;
+			break;
+		}
+
 		InOutArgs.Call.CallRef.Function = InFunction;
 		InOutArgs.Call.CallRef.Object = InOutArgs.ObjectReference.Object;
 		InOutArgs.Call.ParamStruct = MoveTemp(InFunctionArguments);
@@ -284,7 +298,21 @@ namespace RemoteControlSetterUtils
 	{
 		ensure(InFunctionArguments.GetStruct() && InFunctionArguments.GetStruct()->IsA<UFunction>());
 
-		InOutArgs.Call.bGenerateTransaction = InOutArgs.ObjectReference.Access == ERCAccess::WRITE_TRANSACTION_ACCESS ? true : false;
+		switch (InOutArgs.ObjectReference.Access)
+		{
+		case ERCAccess::WRITE_TRANSACTION_ACCESS:
+			InOutArgs.Call.TransactionMode = ERCTransactionMode::AUTOMATIC;
+			break;
+
+		case ERCAccess::WRITE_MANUAL_TRANSACTION_ACCESS:
+			InOutArgs.Call.TransactionMode = ERCTransactionMode::MANUAL;
+			break;
+
+		default:
+			InOutArgs.Call.TransactionMode = ERCTransactionMode::NONE;
+			break;
+		}
+
 		InOutArgs.Call.CallRef.Function = InFunction;
 		InOutArgs.Call.CallRef.Object = InOutArgs.ObjectReference.Object;
 		InOutArgs.Call.ParamStruct = MoveTemp(InFunctionArguments);
@@ -300,7 +328,22 @@ namespace RemoteControlSetterUtils
 		FStructSerializer::Serialize(InFunctionArguments.GetStructMemory(), *const_cast<UStruct*>(InFunctionArguments.GetStruct()), WriterBackend, FStructSerializerPolicies());
 		OutPayload.Type = ERCPayloadType::Cbor;
 
-		InOutArgs.Call.bGenerateTransaction = InOutArgs.ObjectReference.Access == ERCAccess::WRITE_TRANSACTION_ACCESS ? true : false;
+
+		switch (InOutArgs.ObjectReference.Access)
+		{
+		case ERCAccess::WRITE_TRANSACTION_ACCESS:
+			InOutArgs.Call.TransactionMode = ERCTransactionMode::AUTOMATIC;
+			break;
+
+		case ERCAccess::WRITE_MANUAL_TRANSACTION_ACCESS:
+			InOutArgs.Call.TransactionMode = ERCTransactionMode::MANUAL;
+			break;
+
+		default:
+			InOutArgs.Call.TransactionMode = ERCTransactionMode::NONE;
+			break;
+		}
+
 		InOutArgs.Call.CallRef.PropertyWithSetter = InPropertyWithSetter;
 		InOutArgs.Call.CallRef.Object = InOutArgs.ObjectReference.Object;
 		InOutArgs.Call.ParamStruct = MoveTemp(InFunctionArguments);
@@ -314,8 +357,22 @@ namespace RemoteControlSetterUtils
 	void CreateRCCall(FConvertToFunctionCallArgs& InOutArgs, FProperty* InPropertyWithSetter, FStructOnScope&& InFunctionArguments)
 	{
 		ensure(InFunctionArguments.GetStruct());
-		
-		InOutArgs.Call.bGenerateTransaction = InOutArgs.ObjectReference.Access == ERCAccess::WRITE_TRANSACTION_ACCESS ? true : false;
+
+		switch (InOutArgs.ObjectReference.Access)
+		{
+		case ERCAccess::WRITE_TRANSACTION_ACCESS:
+			InOutArgs.Call.TransactionMode = ERCTransactionMode::AUTOMATIC;
+			break;
+
+		case ERCAccess::WRITE_MANUAL_TRANSACTION_ACCESS:
+			InOutArgs.Call.TransactionMode = ERCTransactionMode::MANUAL;
+			break;
+
+		default:
+			InOutArgs.Call.TransactionMode = ERCTransactionMode::NONE;
+			break;
+		}
+
 		InOutArgs.Call.CallRef.PropertyWithSetter = InPropertyWithSetter;
 		InOutArgs.Call.CallRef.Object = InOutArgs.ObjectReference.Object;
 		InOutArgs.Call.ParamStruct = MoveTemp(InFunctionArguments);
@@ -639,12 +696,15 @@ bool FRemoteControlModule::ResolveCall(const FString& ObjectPath, const FString&
 bool FRemoteControlModule::InvokeCall(FRCCall& InCall, ERCPayloadType InPayloadType, const TArray<uint8>& InInterceptPayload)
 {
 	UE_LOG(LogRemoteControl, VeryVerbose, TEXT("Invoke function"));
+
 	if (InCall.IsValid())
 	{
+		const bool bGenerateTransaction = InCall.TransactionMode == ERCTransactionMode::AUTOMATIC;
+
 		// Check the replication path before apply property values
 		if (InInterceptPayload.Num() != 0)
 		{
-			FRCIFunctionMetadata FunctionMetadata(InCall.CallRef.Object->GetPathName(), InCall.CallRef.Function->GetPathName(), InCall.bGenerateTransaction, ToExternal(InPayloadType), InInterceptPayload);
+			FRCIFunctionMetadata FunctionMetadata(InCall.CallRef.Object->GetPathName(), InCall.CallRef.Function->GetPathName(), bGenerateTransaction, ToExternal(InPayloadType), InInterceptPayload);
 
 			// Initialization
 			IModularFeatures& ModularFeatures = IModularFeatures::Get();
@@ -672,35 +732,23 @@ bool FRemoteControlModule::InvokeCall(FRCCall& InCall, ERCPayloadType InPayloadT
 		}
 
 #if WITH_EDITOR
-		if (CVarRemoteControlEnableOngoingChangeOptimization.GetValueOnAnyThread() == 1)
+		const bool bUseOngoingChangeOptimization = CVarRemoteControlEnableOngoingChangeOptimization.GetValueOnAnyThread() == 1;
+		if (bUseOngoingChangeOptimization)
 		{
-			// If we have a different ongoing change that hasn't been finalized, do that before handling the next function call.
-			if (OngoingModification && GetTypeHash(*OngoingModification) != GetTypeHash(InCall.CallRef))
-			{
-				constexpr bool bForceFinalizeChange = true;
-				TestOrFinalizeOngoingChange(true);
-			}
-				
-			if (!OngoingModification)
-			{
-				if (InCall.bGenerateTransaction)
-				{
-					if (GEditor)
-					{
-						GEditor->BeginTransaction(LOCTEXT("RemoteCallTransaction", "Remote Call Transaction Wrap"));
-					}
-
-					if (ensureAlways(InCall.CallRef.Object.IsValid()))
-					{
-					InCall.CallRef.Object->Modify();
-				}
-			}
-		}
+			EndOngoingModificationIfMismatched(GetTypeHash(InCall.CallRef));
 		}
-		else if (GEditor && InCall.bGenerateTransaction)
+
+		const bool bIsNewTransaction = bGenerateTransaction && (!bUseOngoingChangeOptimization || !OngoingModification);
+		if (bIsNewTransaction && GEditor)
 		{
 			GEditor->BeginTransaction(LOCTEXT("RemoteCallTransaction", "Remote Call Transaction Wrap"));
 		}
+
+		const bool bIsManualTransaction = InCall.TransactionMode == ERCTransactionMode::MANUAL;
+		if ((bIsNewTransaction || bIsManualTransaction) && ensureAlways(InCall.CallRef.Object.IsValid()))
+		{
+			InCall.CallRef.Object->Modify();
+		}
 			
 #endif
 		FEditorScriptExecutionGuard ScriptGuard;
@@ -717,7 +765,7 @@ bool FRemoteControlModule::InvokeCall(FRCCall& InCall, ERCPayloadType InPayloadT
 		}
 
 #if WITH_EDITOR
-		if (CVarRemoteControlEnableOngoingChangeOptimization.GetValueOnAnyThread() == 1)
+		if (bUseOngoingChangeOptimization)
 		{
 			// If we've called the same function recently, refresh the triggered flag and snapshot the object to the transaction buffer
 			if (OngoingModification && GetTypeHash(*OngoingModification) == GetTypeHash(InCall.CallRef))
@@ -735,13 +783,13 @@ bool FRemoteControlModule::InvokeCall(FRCCall& InCall, ERCPayloadType InPayloadT
 					}
 				}
 			}
-			else
+			else if (!bIsManualTransaction)
 			{
 				OngoingModification = InCall.CallRef;
-				OngoingModification->bHasStartedTransaction = InCall.bGenerateTransaction;
+				OngoingModification->bHasStartedTransaction = bGenerateTransaction;
 			}
 		}
-		else if (GEditor && InCall.bGenerateTransaction)
+		else if (GEditor && bGenerateTransaction)
 		{
 			GEditor->EndTransaction();
 		}
@@ -1003,17 +1051,13 @@ bool FRemoteControlModule::SetObjectProperties(const FRCObjectReference& ObjectA
 		FRCObjectReference MutableObjectReference = ObjectAccess;
 
 #if WITH_EDITOR
-		const bool bGenerateTransaction = MutableObjectReference.Access == ERCAccess::WRITE_TRANSACTION_ACCESS;
-		if (CVarRemoteControlEnableOngoingChangeOptimization.GetValueOnAnyThread() == 1)
+		const bool bUseOngoingChangeOptimization = CVarRemoteControlEnableOngoingChangeOptimization.GetValueOnAnyThread() == 1;
+		if (bUseOngoingChangeOptimization)
 		{
-			FString ObjectPath = MutableObjectReference.Object->GetPathName();
+			const FString ObjectPath = MutableObjectReference.Object->GetPathName();
 
 			// If we have a change that hasn't yet generated a post edit change property, do that before handling the next change.
-			if (OngoingModification && GetTypeHash(*OngoingModification) != GetTypeHash(MutableObjectReference))
-			{
-				constexpr bool bForcePostEditChange = true;
-				TestOrFinalizeOngoingChange(true);
-			}
+			EndOngoingModificationIfMismatched(GetTypeHash(MutableObjectReference));
 
 			// This step is necessary because the object might get recreated by a PostEditChange called in TestOrFinalizeOngoingChange.
 			if (!MutableObjectReference.Object.IsValid())
@@ -1032,30 +1076,27 @@ bool FRemoteControlModule::SetObjectProperties(const FRCObjectReference& ObjectA
 					return false;
 				}
 			}
+		}
 
-			// Only create the transaction if we have no ongoing change.
-			if (!OngoingModification)
-			{
-				if (GEditor && bGenerateTransaction)
-				{
-					GEditor->BeginTransaction(LOCTEXT("RemoteSetPropertyTransaction", "Remote Set Object Property"));
+		const bool bGenerateTransaction = MutableObjectReference.Access == ERCAccess::WRITE_TRANSACTION_ACCESS;
+		const bool bIsNewTransaction = bGenerateTransaction && (!bUseOngoingChangeOptimization || !OngoingModification);
 
-					// Call modify since it's not called by PreEditChange until the end of the ongoing change.
-					MutableObjectReference.Object->Modify();
-				}
-			}
-		}
-		else 
+		if (bIsNewTransaction && GEditor)
 		{
-			if (GEditor && bGenerateTransaction)
+			GEditor->BeginTransaction(LOCTEXT("RemoteSetPropertyTransaction", "Remote Set Object Property"));
+
+			if (bUseOngoingChangeOptimization)
 			{
-				GEditor->BeginTransaction(LOCTEXT("RemoteSetPropertyTransaction", "Remote Set Object Property"));
+				// Call modify since it's not called by PreEditChange until the end of the ongoing change.
+				MutableObjectReference.Object->Modify();
 			}
+		}
 
+		if (!bUseOngoingChangeOptimization)
+		{
 			FEditPropertyChain PreEditChain;
 			MutableObjectReference.PropertyPathInfo.ToEditPropertyChain(PreEditChain);
 			MutableObjectReference.Object->PreEditChange(PreEditChain);
-				
 		}
 #endif
 
@@ -1141,7 +1182,7 @@ bool FRemoteControlModule::SetObjectProperties(const FRCObjectReference& ObjectA
 					}
 				}
 			}
-			else
+			else if (MutableObjectReference.Access != ERCAccess::WRITE_MANUAL_TRANSACTION_ACCESS)
 			{
 				OngoingModification = MutableObjectReference;
 				OngoingModification->bHasStartedTransaction = bGenerateTransaction;
@@ -1478,6 +1519,54 @@ void FRemoteControlModule::UnregisterEntityFactory(const FName InFactoryName)
 	EntityFactories.Remove(InFactoryName);
 }
 
+FGuid FRemoteControlModule::BeginManualEditorTransaction(const FText& InDescription, uint32 TypeHash)
+{
+#if WITH_EDITOR
+	if (CVarRemoteControlEnableOngoingChangeOptimization.GetValueOnAnyThread() == 1)
+	{
+		EndOngoingModificationIfMismatched(TypeHash);
+
+		if (OngoingModification)
+		{
+			// We weren't able to finish the ongoing modification, so we won't start a new transaction
+			return FGuid();
+		}
+	}
+
+	if (GEditor && GEditor->Trans)
+	{
+		if (GEditor->BeginTransaction(InDescription) != INDEX_NONE)
+		{
+			const FTransaction* Transaction = GEditor->Trans->GetTransaction(GEditor->Trans->GetQueueLength() - 1);
+			if (ensure(Transaction))
+			{
+				return Transaction->GetId();
+			}
+		}
+	}
+#endif
+
+	return FGuid();
+}
+
+int32 FRemoteControlModule::EndManualEditorTransaction(const FGuid& TransactionId)
+{
+#if WITH_EDITOR
+	if (GEditor && GEditor->Trans && GEditor->Trans->IsActive())
+	{
+		if (const FTransaction* CurrentTransaction = GEditor->Trans->GetTransaction(GEditor->Trans->GetQueueLength() - 1))
+		{
+			if (CurrentTransaction->GetId() == TransactionId)
+			{
+				return GEditor->EndTransaction();
+			}
+		}
+	}
+#endif
+
+	return INDEX_NONE;
+}
+
 void FRemoteControlModule::CachePresets() const
 {
 	TArray<FAssetData> Assets;
@@ -1694,6 +1783,14 @@ bool FRemoteControlModule::DeserializeDeltaModificationData(const FRCObjectRefer
 }
 
 #if WITH_EDITOR
+void FRemoteControlModule::EndOngoingModificationIfMismatched(uint32 TypeHash)
+{
+	if (OngoingModification && GetTypeHash(*OngoingModification) != TypeHash)
+	{
+		TestOrFinalizeOngoingChange(true);
+	}
+}
+
 void FRemoteControlModule::TestOrFinalizeOngoingChange(bool bForceEndChange)
 {
 	if (OngoingModification)
diff --git a/Engine/Plugins/VirtualProduction/RemoteControl/Source/RemoteControl/Private/RemoteControlModule.h b/Engine/Plugins/VirtualProduction/RemoteControl/Source/RemoteControl/Private/RemoteControlModule.h
--- a/Engine/Plugins/VirtualProduction/RemoteControl/Source/RemoteControl/Private/RemoteControlModule.h
+++ b/Engine/Plugins/VirtualProduction/RemoteControl/Source/RemoteControl/Private/RemoteControlModule.h
@@ -50,6 +50,8 @@ public:
 	virtual bool PropertySupportsRawModificationWithoutEditor(FProperty* Property, UClass* OwnerClass = nullptr) const override;
 	virtual void RegisterEntityFactory( const FName InFactoryName, const TSharedRef<IRemoteControlPropertyFactory>& InFactory) override;
 	virtual void UnregisterEntityFactory( const FName InFactoryName ) override;
+	virtual FGuid BeginManualEditorTransaction(const FText& InDescription, uint32 TypeHash) override;
+	virtual int32 EndManualEditorTransaction(const FGuid& TransactionId) override;
 	virtual const TMap<FName, TSharedPtr<IRemoteControlPropertyFactory>>& GetEntityFactories() const override { return EntityFactories; };
 	//~ End IRemoteControlModule
 
@@ -80,6 +82,12 @@ private:
 	static bool DeserializeDeltaModificationData(const FRCObjectReference& ObjectAccess, IStructDeserializerBackend& Backend, ERCModifyOperation Operation, TArray<uint8>& OutData);
 
 #if WITH_EDITOR
+	/**
+	 * End the ongoing modification if one exists and is a mismatch for the new object to edit.
+	 * @param TypeHash The type hash of the object we want to modify after this check.
+	 */
+	void EndOngoingModificationIfMismatched(uint32 TypeHash);
+
 	/** Finalize an ongoing change, triggering post edit change on the tracked object. */
 	void TestOrFinalizeOngoingChange(bool bForceEndChange = false);
 
diff --git a/Engine/Plugins/VirtualProduction/RemoteControl/Source/RemoteControl/Private/Tests/RemoteControlInterceptionTest.cpp b/Engine/Plugins/VirtualProduction/RemoteControl/Source/RemoteControl/Private/Tests/RemoteControlInterceptionTest.cpp
--- a/Engine/Plugins/VirtualProduction/RemoteControl/Source/RemoteControl/Private/Tests/RemoteControlInterceptionTest.cpp
+++ b/Engine/Plugins/VirtualProduction/RemoteControl/Source/RemoteControl/Private/Tests/RemoteControlInterceptionTest.cpp
@@ -348,7 +348,7 @@ struct FRemoteControlInterceptionTest
 		
 		FRCCall Call;
 		Call.CallRef = MoveTemp(CallRef);
-		Call.bGenerateTransaction = true;
+		Call.TransactionMode = ERCTransactionMode::AUTOMATIC;
 		Call.ParamStruct = FStructOnScope(FunctionArgs.GetStruct(), FunctionArgs.GetStructMemory());
 
 		const bool bSuccess = IRemoteControlModule::Get().InvokeCall(Call, ERCPayloadType::Json, JsonBuffer);
diff --git a/Engine/Plugins/VirtualProduction/RemoteControl/Source/RemoteControl/Public/IRemoteControlModule.h b/Engine/Plugins/VirtualProduction/RemoteControl/Source/RemoteControl/Public/IRemoteControlModule.h
--- a/Engine/Plugins/VirtualProduction/RemoteControl/Source/RemoteControl/Public/IRemoteControlModule.h
+++ b/Engine/Plugins/VirtualProduction/RemoteControl/Source/RemoteControl/Public/IRemoteControlModule.h
@@ -76,6 +76,14 @@ struct FRCCallReference
 	}
 };
 
+UENUM()
+enum class ERCTransactionMode : uint8
+{
+	NONE = 0,
+	AUTOMATIC,
+	MANUAL,
+};
+
 /**
  * Object to hold a UObject remote call
  */
@@ -91,7 +99,7 @@ struct FRCCall
 	FStructOnScope ParamStruct;
 	// Payload for native function
 	TArray<uint8> ParamData;
-	bool bGenerateTransaction = false;
+	ERCTransactionMode TransactionMode;
 };
 
 /**
@@ -113,6 +121,7 @@ enum class ERCAccess : uint8
 	READ_ACCESS,
 	WRITE_ACCESS,
 	WRITE_TRANSACTION_ACCESS,
+	WRITE_MANUAL_TRANSACTION_ACCESS,
 };
 
 /**
@@ -422,6 +431,21 @@ public:
 	 */
 	virtual void UnregisterEntityFactory( const FName InFactoryName ) = 0;
 
+	/**
+	 * Start a manual editor transaction originating from a remote request.
+	 * @param InDescription A description of the transaction.
+	 * @param TypeHash The type hash of the object or call to which the transaction applies, or 0 if not applicable (e.g. multiple objects being modified).
+	 * @return The ID of the new transaction if one was created, or an invalid ID otherwise.
+	 */
+	virtual FGuid BeginManualEditorTransaction(const FText& InDescription, uint32 TypeHash) = 0;
+
+	/**
+	 * End a manual editor transaction originating from a remote request.
+	 * @param TransactionId The ID of the transaction. If this doesn't match the current active transaction, the transaction won't be ended.
+	 * @return INDEX_NONE if there was no transaction to end; otherwise, the number of remaining actions in the transaction (i.e. 1 means this was the last action and the transaction will end).
+	 */
+	virtual int32 EndManualEditorTransaction(const FGuid& TransactionId) = 0;
+
 	/** Get map of the factories which is responsible for the Remote Control property creation */
 	virtual const TMap<FName, TSharedPtr<IRemoteControlPropertyFactory>>& GetEntityFactories() const = 0;
 };
diff --git a/Engine/Plugins/VirtualProduction/RemoteControl/Source/RemoteControl/Public/RemoteControlFunctionLibrary.h b/Engine/Plugins/VirtualProduction/RemoteControl/Source/RemoteControl/Public/RemoteControlFunctionLibrary.h
--- a/Engine/Plugins/VirtualProduction/RemoteControl/Source/RemoteControl/Public/RemoteControlFunctionLibrary.h
+++ b/Engine/Plugins/VirtualProduction/RemoteControl/Source/RemoteControl/Public/RemoteControlFunctionLibrary.h
@@ -98,8 +98,9 @@ public:
 	 * @param PropertyName the name of the property to modify.
 	 * @param DeltaValue the amount to change the color by.
 	 * @param ReferenceColor if the color's current position on the wheel is ambiguous as calculated from RGB values (e.g. black), use this reference color's position instead.
+	 * @param bIsInteractive if true, this is treated as an interactive change. If false, it will be treated as the final value set change.
 	 * @return Whether the operation was successful.
 	 */
 	UFUNCTION(BlueprintCallable, Category = "Remote Control Color")
-	static bool ApplyColorWheelDelta(UObject* TargetObject, const FString& PropertyName, const FColorWheelColor& DeltaValue, const FColorWheelColor& ReferenceColor);
+	static bool ApplyColorWheelDelta(UObject* TargetObject, const FString& PropertyName, const FColorWheelColor& DeltaValue, const FColorWheelColor& ReferenceColor, bool bIsInteractive);
 };
\ No newline at end of file
diff --git a/Engine/Plugins/VirtualProduction/RemoteControl/Source/RemoteControlLogic/Private/Action/RCFunctionAction.cpp b/Engine/Plugins/VirtualProduction/RemoteControl/Source/RemoteControlLogic/Private/Action/RCFunctionAction.cpp
--- a/Engine/Plugins/VirtualProduction/RemoteControl/Source/RemoteControlLogic/Private/Action/RCFunctionAction.cpp
+++ b/Engine/Plugins/VirtualProduction/RemoteControl/Source/RemoteControlLogic/Private/Action/RCFunctionAction.cpp
@@ -45,7 +45,7 @@ void URCFunctionAction::Execute() const
 
 			FRCCall Call;
 			Call.CallRef = MoveTemp(CallRef);
-			Call.bGenerateTransaction = true;
+			Call.TransactionMode = ERCTransactionMode::AUTOMATIC;
 			Call.ParamStruct = FStructOnScope(FunctionArgs.GetStruct(), FunctionArgs.GetStructMemory());
 		
 			IRemoteControlModule::Get().InvokeCall(Call, ERCPayloadType::Cbor, FunctionPayload);
diff --git a/Engine/Plugins/VirtualProduction/RemoteControl/Source/WebRemoteControl/Private/RemoteControlRequest.h b/Engine/Plugins/VirtualProduction/RemoteControl/Source/WebRemoteControl/Private/RemoteControlRequest.h
--- a/Engine/Plugins/VirtualProduction/RemoteControl/Source/WebRemoteControl/Private/RemoteControlRequest.h
+++ b/Engine/Plugins/VirtualProduction/RemoteControl/Source/WebRemoteControl/Private/RemoteControlRequest.h
@@ -696,10 +696,19 @@ struct FRCWebSocketPresetSetPropertyBody : public FRCRequest
 	ERCModifyOperation Operation = ERCModifyOperation::EQUAL;
 
 	/**
-	 * Whether a transaction should be created for the call.
+	 * How to handle generating transactions for this property change.
+	 * If NONE, don't generate a transaction immediately.
+	 * If AUTOMATIC, let the Remote Control system automatically start and end the transaction after enough time passes.
+	 * If MANUAL, TransactionId must also be set and the changes will only be applied if that transaction is still active.
 	 */
 	UPROPERTY()
-	bool GenerateTransaction = false;
+	ERCTransactionMode TransactionMode = ERCTransactionMode::NONE;
+
+	/**
+	 * The ID of the transaction with which to associate these changes. Must be provided if TransactionMode is Manual.
+	 */
+	UPROPERTY()
+	int32 TransactionId = -1;
 
 	/**
 	 * If true, ignore the other parameters and just reset the property to its default value.
@@ -715,7 +724,6 @@ struct FRCWebSocketPresetSetPropertyBody : public FRCRequest
 	int64 SequenceNumber = -1;
 };
 
-
 /**
  * Holds a request made via websocket to call an exposed function on an object.
  */
@@ -725,10 +733,26 @@ struct FRCWebSocketCallBody : public FRCCallRequest
 	GENERATED_BODY()
 
 	FRCWebSocketCallBody()
-		: FRCCallRequest()
+	: FRCCallRequest()
 	{
 	}
 
+	/**
+	 * How to handle generating transactions for this property change.
+	 * If NONE, don't generate a transaction immediately.
+	 * If AUTOMATIC, let the Remote Control system automatically start and end the transaction after enough time passes.
+	 * If MANUAL, TransactionId must also be set and the changes will only be applied if that transaction is still active.
+	 * If bGenerateTransaction is true, this value will be treated as if it was AUTOMATIC.
+	 */
+	UPROPERTY()
+	ERCTransactionMode TransactionMode = ERCTransactionMode::NONE;
+
+	/**
+	 * The ID of the transaction with which to associate these changes. Must be provided if TransactionMode is Manual.
+	 */
+	UPROPERTY()
+	int32 TransactionId = -1;
+
 	/**
 	 * The sequence number of this change. The highest sequence number received from this client will be
 	 * sent back to the client in future PresetFieldsChanged events.
@@ -736,3 +760,60 @@ struct FRCWebSocketCallBody : public FRCCallRequest
 	UPROPERTY()
 	int64 SequenceNumber = -1;
 };
+
+/**
+ * Holds a request made via websocket to start a transaction.
+ */
+USTRUCT()
+struct FRCWebSocketTransactionStartBody : public FRCRequest
+{
+	GENERATED_BODY()
+
+	FRCWebSocketTransactionStartBody()
+	{
+		AddStructParameter(ParametersFieldLabel());
+	}
+
+	/**
+	 * Get the label for the property value struct.
+	 */
+	static FString ParametersFieldLabel() { return TEXT("Parameters"); }
+
+	/**
+	 * The description of the transaction.
+	 */
+	UPROPERTY()
+	FString Description;
+
+	/**
+	 * The ID that will be used to refer to the transaction in future messages.
+	 */
+	UPROPERTY()
+	int32 TransactionId;
+};
+
+
+/**
+ * Holds a request made via websocket to end a transaction.
+ */
+USTRUCT()
+struct FRCWebSocketTransactionEndBody : public FRCRequest
+{
+	GENERATED_BODY()
+
+	FRCWebSocketTransactionEndBody()
+	{
+		AddStructParameter(ParametersFieldLabel());
+	}
+
+	/**
+	 * Get the label for the property value struct.
+	 */
+	static FString ParametersFieldLabel() { return TEXT("Parameters"); }
+
+	/**
+	 * The ID of the transaction. If this doesn't match the current editor transaction, it won't be ended.
+	 */
+	UPROPERTY()
+	int32 TransactionId;
+};
diff --git a/Engine/Plugins/VirtualProduction/RemoteControl/Source/WebRemoteControl/Private/RemoteControlResponse.h b/Engine/Plugins/VirtualProduction/RemoteControl/Source/WebRemoteControl/Private/RemoteControlResponse.h
--- a/Engine/Plugins/VirtualProduction/RemoteControl/Source/WebRemoteControl/Private/RemoteControlResponse.h
+++ b/Engine/Plugins/VirtualProduction/RemoteControl/Source/WebRemoteControl/Private/RemoteControlResponse.h
@@ -444,3 +444,36 @@ struct FRCActorsChangedEvent
 	UPROPERTY()
 	TMap<FString, FRCActorsChangedData> Changes;
 };
+
+/**
+ * Event sent to a client that contributed to a transaction, indicating that the transaction was either cancelled or finalized.
+ */
+USTRUCT()
+struct FRCTransactionEndedEvent
+{
+	GENERATED_BODY()
+
+	FRCTransactionEndedEvent()
+	: Type(TEXT("TransactionEnded"))
+	{
+	}
+
+	/**
+	 * Type of the event.
+	 */
+	UPROPERTY()
+	FString Type;
+
+	/**
+	 * The client-specific ID of the transaction that was ended.
+	 */
+	UPROPERTY()
+	int32 TransactionId;
+
+	/**
+	 * The highest sequence number received from the receiving client at the time that the transaction ended.
+	 */
+	UPROPERTY()
+	int64 SequenceNumber = -1;
+};
+
diff --git a/Engine/Plugins/VirtualProduction/RemoteControl/Source/WebRemoteControl/Private/WebRemoteControl.cpp b/Engine/Plugins/VirtualProduction/RemoteControl/Source/WebRemoteControl/Private/WebRemoteControl.cpp
--- a/Engine/Plugins/VirtualProduction/RemoteControl/Source/WebRemoteControl/Private/WebRemoteControl.cpp
+++ b/Engine/Plugins/VirtualProduction/RemoteControl/Source/WebRemoteControl/Private/WebRemoteControl.cpp
@@ -969,6 +969,7 @@ bool FWebRemoteControlModule::HandleObjectPropertyRoute(const FHttpServerRequest
 	break;
 	case ERCAccess::WRITE_ACCESS:
 	case ERCAccess::WRITE_TRANSACTION_ACCESS:
+	case ERCAccess::WRITE_MANUAL_TRANSACTION_ACCESS:
 	{
 		const FBlockDelimiters& PropertyValueDelimiters = DeserializedRequest.GetStructParameters().FindChecked(FRCObjectRequest::PropertyValueLabel());
 		if (bResetToDefault)
@@ -1090,7 +1091,7 @@ bool FWebRemoteControlModule::HandlePresetCallFunctionRoute(const FHttpServerReq
 
 				FRCCall Call;
 				Call.CallRef = MoveTemp(CallRef);
-				Call.bGenerateTransaction = CallRequest.GenerateTransaction;
+				Call.TransactionMode = CallRequest.GenerateTransaction ? ERCTransactionMode::AUTOMATIC : ERCTransactionMode::NONE;
 				Call.ParamStruct = FStructOnScope(FunctionArgs.GetStruct(), FunctionArgs.GetStructMemory());
 
 				// Invoke call with replication payload
@@ -1179,7 +1180,8 @@ bool FWebRemoteControlModule::HandlePresetSetPropertyRoute(const FHttpServerRequ
 	}
 
 	const bool bSuccess = WebRemoteControlInternalUtils::ModifyPropertyUsingPayload(
-		*RemoteControlProperty.Get(), SetPropertyRequest, SetPropertyRequest.TCHARBody, ActingClientId, *WebSocketHandler);
+		*RemoteControlProperty.Get(), SetPropertyRequest, SetPropertyRequest.TCHARBody, ActingClientId, *WebSocketHandler,
+		SetPropertyRequest.GenerateTransaction ? ERCAccess::WRITE_TRANSACTION_ACCESS : ERCAccess::WRITE_ACCESS);
 
 	if (bSuccess)
 	{
diff --git a/Engine/Plugins/VirtualProduction/RemoteControl/Source/WebRemoteControl/Private/WebRemoteControlInternalUtils.cpp b/Engine/Plugins/VirtualProduction/RemoteControl/Source/WebRemoteControl/Private/WebRemoteControlInternalUtils.cpp
--- a/Engine/Plugins/VirtualProduction/RemoteControl/Source/WebRemoteControl/Private/WebRemoteControlInternalUtils.cpp
+++ b/Engine/Plugins/VirtualProduction/RemoteControl/Source/WebRemoteControl/Private/WebRemoteControlInternalUtils.cpp
@@ -138,7 +138,7 @@ bool DeserializeCall(const FHttpServerRequest& InRequest, FRCCall& OutCall, cons
 	if (IRemoteControlModule::Get().ResolveCall(CallRequest.ObjectPath, CallRequest.FunctionName, OutCall.CallRef, &ErrorText))
 	{
 		// Initialize the param struct with default parameters
-		OutCall.bGenerateTransaction = CallRequest.GenerateTransaction;
+		OutCall.TransactionMode = CallRequest.GenerateTransaction ? ERCTransactionMode::AUTOMATIC : ERCTransactionMode::NONE;
 		OutCall.ParamStruct = FStructOnScope(OutCall.CallRef.Function.Get());
 
 		// If some parameters were provided, deserialize them
diff --git a/Engine/Plugins/VirtualProduction/RemoteControl/Source/WebRemoteControl/Private/WebRemoteControlInternalUtils.h b/Engine/Plugins/VirtualProduction/RemoteControl/Source/WebRemoteControl/Private/WebRemoteControlInternalUtils.h
--- a/Engine/Plugins/VirtualProduction/RemoteControl/Source/WebRemoteControl/Private/WebRemoteControlInternalUtils.h
+++ b/Engine/Plugins/VirtualProduction/RemoteControl/Source/WebRemoteControl/Private/WebRemoteControlInternalUtils.h
@@ -361,9 +361,10 @@ namespace WebRemoteControlInternalUtils
 	 * @param Payload The payload from which to deserialize property data.
 	 * @param ClientId The ID of the client that sent this request.
 	 * @param WebSocketHandler The WebSocket handler that will be notified of this remote change.
+	 * @param Access The access mode to use for this operation.
 	 */
 	template <typename RequestType>
-	bool ModifyPropertyUsingPayload(FRemoteControlProperty& Property, const RequestType& Request, const TArrayView<uint8>& Payload, const FGuid& ClientId, FWebSocketMessageHandler& WebSocketHandler)
+	bool ModifyPropertyUsingPayload(FRemoteControlProperty& Property, const RequestType& Request, const TArrayView<uint8>& Payload, const FGuid& ClientId, FWebSocketMessageHandler& WebSocketHandler, ERCAccess Access)
 	{
 		FRCObjectReference ObjectRef;
 
@@ -376,7 +377,7 @@ namespace WebRemoteControlInternalUtils
 		FRCJsonStructDeserializerBackend Backend(NewPayloadReader);
 
 		ObjectRef.Property = Property.GetProperty();
-		ObjectRef.Access = Request.GenerateTransaction ? ERCAccess::WRITE_TRANSACTION_ACCESS : ERCAccess::WRITE_ACCESS;
+		ObjectRef.Access = Access;
 
 		bool bSuccess = true;
 
@@ -397,6 +398,16 @@ namespace WebRemoteControlInternalUtils
 				}
 			}
 
+#if WITH_EDITOR
+			if (Access == ERCAccess::WRITE_MANUAL_TRANSACTION_ACCESS && ObjectRef.Object.IsValid())
+			{
+				// This transaction is being manually controlled, so RemoteControlModule's automatic transaction handling won't call this for us
+				FEditPropertyChain PreEditChain;
+				ObjectRef.PropertyPathInfo.ToEditPropertyChain(PreEditChain);
+				ObjectRef.Object->PreEditChange(PreEditChain);
+			}
+#endif
+
 			if (Request.ResetToDefault)
 			{
 				// set interception flag as an extra argument {}
@@ -409,6 +420,16 @@ namespace WebRemoteControlInternalUtils
 				// Set a ERCPayloadType and TCHARBody in order to follow the replication path
 				bSuccess &= IRemoteControlModule::Get().SetObjectProperties(ObjectRef, Backend, ERCPayloadType::Json, NewPayload, Request.Operation);
 			}
+
+#if WITH_EDITOR
+			if (Access == ERCAccess::WRITE_MANUAL_TRANSACTION_ACCESS && ObjectRef.Object.IsValid())
+			{
+				// This transaction is being manually controlled, so RemoteControlModule's automatic transaction handling won't call this for us
+				FPropertyChangedEvent PropertyEvent = ObjectRef.PropertyPathInfo.ToPropertyChangedEvent();
+				PropertyEvent.ChangeType = EPropertyChangeType::Interactive;
+				ObjectRef.Object->PostEditChangeProperty(PropertyEvent);
+			}
+#endif
 		}
 
 		return bSuccess;
diff --git a/Engine/Plugins/VirtualProduction/RemoteControl/Source/WebRemoteControl/Private/WebSocketMessageHandler.cpp b/Engine/Plugins/VirtualProduction/RemoteControl/Source/WebRemoteControl/Private/WebSocketMessageHandler.cpp
--- a/Engine/Plugins/VirtualProduction/RemoteControl/Source/WebRemoteControl/Private/WebSocketMessageHandler.cpp
+++ b/Engine/Plugins/VirtualProduction/RemoteControl/Source/WebRemoteControl/Private/WebSocketMessageHandler.cpp
@@ -18,10 +18,15 @@
 
 #if WITH_EDITOR
 #include "Editor.h"
+#include "Editor/TransBuffer.h"
 #endif
 
+#define LOCTEXT_NAMESPACE "WebRemoteControl"
+
 static TAutoConsoleVariable<int32> CVarWebRemoteControlFramesBetweenPropertyNotifications(TEXT("WebControl.FramesBetweenPropertyNotifications"), 5, TEXT("The number of frames between sending batches of property notifications."));
 const int64 FWebSocketMessageHandler::DefaultSequenceNumber = -1;
+const int32 FWebSocketMessageHandler::InvalidTransactionId = -1;
+const FTimespan FWebSocketMessageHandler::TransactionTimeout = FTimespan::FromSeconds(3.f);
 
 namespace WebSocketMessageHandlerStructUtils
 {
@@ -269,6 +274,14 @@ void FWebSocketMessageHandler::RegisterRoutes(FWebRemoteControlModule* WebRemote
 	{
 		FCoreDelegates::OnPostEngineInit.AddRaw(this, &FWebSocketMessageHandler::RegisterActorHandlers);
 	}
+
+	if (GEditor)
+	{
+		if (UTransBuffer* TransBuffer = Cast<UTransBuffer>(GEditor->Trans))
+		{
+			TransBuffer->OnTransactionStateChanged().AddRaw(this, &FWebSocketMessageHandler::HandleTransactionStateChanged);
+		}
+	}
 #endif
 	
 	// WebSocket routes
@@ -313,6 +326,18 @@ void FWebSocketMessageHandler::RegisterRoutes(FWebRemoteControlModule* WebRemote
 		TEXT("object.call"),
 		FWebSocketMessageDelegate::CreateRaw(this, &FWebSocketMessageHandler::HandleWebSocketFunctionCall)
 	));
+
+	RegisterRoute(WebRemoteControl, MakeUnique<FRemoteControlWebsocketRoute>(
+		TEXT("Begin a manual editor transaction. Be sure to send transaction.end when finished."),
+		TEXT("transaction.begin"),
+		FWebSocketMessageDelegate::CreateRaw(this, &FWebSocketMessageHandler::HandleWebSocketBeginEditorTransaction)
+	));
+
+	RegisterRoute(WebRemoteControl, MakeUnique<FRemoteControlWebsocketRoute>(
+		TEXT("End a manual editor transaction"),
+		TEXT("transaction.end"),
+		FWebSocketMessageDelegate::CreateRaw(this, &FWebSocketMessageHandler::HandleWebSocketEndEditorTransaction)
+	));
 }
 
 void FWebSocketMessageHandler::UnregisterRoutes(FWebRemoteControlModule* WebRemoteControl)
@@ -331,6 +356,14 @@ void FWebSocketMessageHandler::UnregisterRoutes(FWebRemoteControlModule* WebRemo
 		GEngine->OnLevelActorListChanged().Remove(OnActorListChangedHandle);
 		GEngine->OnLevelActorListChanged().Remove(OnWorldDestroyedHandle);
 	}
+
+	if (GEditor)
+	{
+		if (UTransBuffer* TransBuffer = Cast<UTransBuffer>(GEditor->Trans))
+		{
+			TransBuffer->OnTransactionStateChanged().RemoveAll(this);
+		}
+	}
 #endif
 
 	for (const TUniquePtr<FRemoteControlWebsocketRoute>& Route : Routes)
@@ -580,6 +613,33 @@ void FWebSocketMessageHandler::HandleWebSocketPresetModifyProperty(const FRemote
 		return;
 	}
 
+	ERCAccess Access;
+	switch (Body.TransactionMode)
+	{
+	case ERCTransactionMode::NONE:
+		Access = ERCAccess::WRITE_ACCESS;
+		break;
+
+	case ERCTransactionMode::AUTOMATIC:
+		Access = ERCAccess::WRITE_TRANSACTION_ACCESS;
+		break;
+
+	case ERCTransactionMode::MANUAL:
+		Access = ERCAccess::WRITE_MANUAL_TRANSACTION_ACCESS;
+
+		// Indicate that we want to contribute to this transaction if it's active
+		if (!ContributeToTransaction(WebSocketMessage.ClientId, Body.TransactionId))
+		{
+			return;
+		}
+
+		break;
+
+	default:
+		UE_LOG(LogRemoteControl, Warning, TEXT("Unknown transaction mode %d"), Body.TransactionMode);
+		return;
+	}
+
 	URemoteControlPreset* Preset = IRemoteControlModule::Get().ResolvePreset(Body.PresetName);
 	if (Preset == nullptr)
 	{
@@ -596,7 +656,7 @@ void FWebSocketMessageHandler::HandleWebSocketPresetModifyProperty(const FRemote
 
 	UpdateSequenceNumber(WebSocketMessage.ClientId, Body.SequenceNumber);
 
-	WebRemoteControlInternalUtils::ModifyPropertyUsingPayload(*RemoteControlProperty.Get(), Body, WebSocketMessage.RequestPayload, WebSocketMessage.ClientId, *this);
+	WebRemoteControlInternalUtils::ModifyPropertyUsingPayload(*RemoteControlProperty.Get(), Body, WebSocketMessage.RequestPayload, WebSocketMessage.ClientId, *this, Access);
 }
 
 void FWebSocketMessageHandler::HandleWebSocketFunctionCall(const FRemoteControlWebSocketMessage& WebSocketMessage)
@@ -607,11 +667,16 @@ void FWebSocketMessageHandler::HandleWebSocketFunctionCall(const FRemoteControlW
 		return;
 	}
 
+	if (Body.GenerateTransaction)
+	{
+		Body.TransactionMode = ERCTransactionMode::AUTOMATIC;
+	}
+
 	FRCCall Call;
 	if (IRemoteControlModule::Get().ResolveCall(Body.ObjectPath, Body.FunctionName, Call.CallRef, nullptr))
 	{
 		// Initialize the param struct with default parameters
-		Call.bGenerateTransaction = Body.GenerateTransaction;
+		Call.TransactionMode = Body.GenerateTransaction ? ERCTransactionMode::AUTOMATIC : Body.TransactionMode;
 		Call.ParamStruct = FStructOnScope(Call.CallRef.Function.Get());
 
 		// If some parameters were provided, deserialize them
@@ -644,11 +709,66 @@ void FWebSocketMessageHandler::HandleWebSocketFunctionCall(const FRemoteControlW
 		return;
 	}
 
+	if (Body.TransactionMode == ERCTransactionMode::MANUAL)
+	{
+		// Indicate that we want to contribute to this transaction
+		if (!ContributeToTransaction(WebSocketMessage.ClientId, Body.TransactionId))
+		{
+			return;
+		}
+	}
+
 	UpdateSequenceNumber(WebSocketMessage.ClientId, Body.SequenceNumber);
 
 	IRemoteControlModule::Get().InvokeCall(Call);
 }
 
+void FWebSocketMessageHandler::HandleWebSocketBeginEditorTransaction(const FRemoteControlWebSocketMessage& WebSocketMessage)
+{
+#if WITH_EDITOR
+	FRCWebSocketTransactionStartBody Body;
+	if (!WebRemoteControlInternalUtils::DeserializeRequestPayload(WebSocketMessage.RequestPayload, nullptr, Body))
+	{
+		return;
+	}
+
+	const FText Description = FText::Format(LOCTEXT("RemoteControlTransaction", "Remote Control - {0}"), FText::FromString(Body.Description));
+	const FGuid TransactionGuid = IRemoteControlModule::Get().BeginManualEditorTransaction(Description, 0);
+
+	if (TransactionGuid.IsValid())
+	{
+		ClientsByTransactionGuid.FindOrAdd(TransactionGuid);
+		ContributeToTransaction(WebSocketMessage.ClientId, Body.TransactionId);
+
+		TransactionIdsByClientId.FindOrAdd(WebSocketMessage.ClientId).Add({ TransactionGuid, Body.TransactionId });
+	}
+	else
+	{
+		// Send a message indicating that the transaction ended immediately so the client knows it wasn't created
+		FRCTransactionEndedEvent Event;
+		Event.TransactionId = Body.TransactionId;
+		Event.SequenceNumber = GetSequenceNumber(WebSocketMessage.ClientId);
+
+		TArray<uint8> Payload;
+		WebRemoteControlUtils::SerializeMessage(Event, Payload);
+		Server->Send(WebSocketMessage.ClientId, Payload);
+	}
+#endif
+}
+
+void FWebSocketMessageHandler::HandleWebSocketEndEditorTransaction(const FRemoteControlWebSocketMessage& WebSocketMessage)
+{
+#if WITH_EDITOR
+	FRCWebSocketTransactionEndBody Body;
+	if (!WebRemoteControlInternalUtils::DeserializeRequestPayload(WebSocketMessage.RequestPayload, nullptr, Body))
+	{
+		return;
+	}
+
+	EndClientTransaction(WebSocketMessage.ClientId, Body.TransactionId);
+#endif
+}
+
 void FWebSocketMessageHandler::ProcessChangedProperties()
 {
 	//Go over each property that were changed for each preset
@@ -681,8 +801,7 @@ void FWebSocketMessageHandler::ProcessChangedProperties()
 			// Send a property change event for each property type
 			for (const TPair<uint64, TSet<FGuid>>& ClassToEventsPair : PropertyIdsByType)
 			{
-				const int64* ClientSequenceNumber = ClientSequenceNumbers.Find(ClientToEventsPair.Key);
-				const int64 SequenceNumber = ClientSequenceNumber ? *ClientSequenceNumber : DefaultSequenceNumber;
+				const int64 SequenceNumber = GetSequenceNumber(ClientToEventsPair.Key);
 
 				TArray<uint8> WorkingBuffer;
 				if (ClientToEventsPair.Value.Num() && WritePropertyChangeEventPayload(Preset, { ClassToEventsPair.Value }, SequenceNumber, WorkingBuffer))
@@ -943,7 +1062,7 @@ void FWebSocketMessageHandler::OnConnectionClosedCallback(FGuid ClientId)
 		TArray<FGuid>& ClientList = Iter.Value();
 
 		ClientList.Remove(ClientId);
-		if (ClientList.Num() == 0)
+		if (ClientList.IsEmpty())
 		{
 			RemoteControl.DestroyTransientPreset(PresetId);
 			PresetNotificationMap.Remove(PresetId);
@@ -951,6 +1070,21 @@ void FWebSocketMessageHandler::OnConnectionClosedCallback(FGuid ClientId)
 		}
 	}
 
+	// Remove the client as a listener for any active transactions
+	for (auto Iter = ClientsByTransactionGuid.CreateIterator(); Iter; ++Iter)
+	{
+		IRemoteControlModule::Get().EndManualEditorTransaction(Iter.Key());
+
+		TMap<FGuid, FDateTime>& Clients = Iter.Value();
+		Clients.Remove(ClientId);
+		if (Clients.IsEmpty())
+		{
+			Iter.RemoveCurrent();
+		}
+	}
+
+	TransactionIdsByClientId.Remove(ClientId);
+
 	/** Remove this client's config. */
 	ClientConfigMap.Remove(ClientId);
 	ClientSequenceNumbers.Remove(ClientId);
@@ -958,16 +1092,18 @@ void FWebSocketMessageHandler::OnConnectionClosedCallback(FGuid ClientId)
 
 void FWebSocketMessageHandler::OnEndFrame()
 {
-	//Early exit if no clients are requesting notifications
-	if (PresetNotificationMap.Num() <= 0 && ActorNotificationMap.Num() <= 0)
-	{
-		return;
-	}
-
 	PropertyNotificationFrameCounter++;
 
 	if (PropertyNotificationFrameCounter >= CVarWebRemoteControlFramesBetweenPropertyNotifications.GetValueOnGameThread())
 	{
+		TimeOutTransactions();
+
+		//Early exit if no clients are requesting notifications
+		if (PresetNotificationMap.Num() <= 0 && ActorNotificationMap.Num() <= 0)
+		{
+			return;
+		}
+
 		PropertyNotificationFrameCounter = 0;
 		ProcessChangedProperties();
 		ProcessChangedActorProperties();
@@ -980,6 +1116,30 @@ void FWebSocketMessageHandler::OnEndFrame()
 	}
 }
 
+void FWebSocketMessageHandler::TimeOutTransactions()
+{
+	TArray<TPair<FGuid, int32>> TransactionsToEnd;
+	const FDateTime Now = FDateTime::Now();
+
+	for (const TPair<FGuid, TMap<FGuid, FDateTime>>& TransactionClientsPair : ClientsByTransactionGuid)
+	{
+		for (const TPair<FGuid, FDateTime>& ClientTimePair : TransactionClientsPair.Value)
+		{
+			if (Now >= ClientTimePair.Value + TransactionTimeout)
+			{
+				// This transaction has timed out; we should force it to end
+				TransactionsToEnd.Add({ ClientTimePair.Key, GetClientTransactionId(ClientTimePair.Key, TransactionClientsPair.Key) });
+			}
+		}
+	}
+
+	// Do this as a separate step since it may remove entries from the map during iteration
+	for (auto TransactionIterator : TransactionsToEnd)
+	{
+		EndClientTransaction(TransactionIterator.Key, TransactionIterator.Value);
+	}
+}
+
 void FWebSocketMessageHandler::ProcessAddedProperties()
 {
 	for (const TPair<FGuid, TArray<FGuid>>& Entry : PerFrameAddedProperties)
@@ -1591,6 +1751,45 @@ void FWebSocketMessageHandler::OnObjectTransacted(UObject* Object, const class F
 	}
 }
 
+#if WITH_EDITOR
+void FWebSocketMessageHandler::HandleTransactionStateChanged(const FTransactionContext& InTransactionContext, const ETransactionStateEventType InTransactionState)
+{
+	if (InTransactionState == ETransactionStateEventType::TransactionCanceled || InTransactionState == ETransactionStateEventType::TransactionFinalized)
+	{
+		HandleTransactionEnded(InTransactionContext.TransactionId);
+	}
+}
+#endif
+
+void FWebSocketMessageHandler::HandleTransactionEnded(const FGuid& TransactionGuid)
+{
+	if (const TMap<FGuid, FDateTime>* TransactionClients = ClientsByTransactionGuid.Find(TransactionGuid))
+	{
+		FRCTransactionEndedEvent Event;
+
+		// Notify the clients
+		for (const TPair<FGuid, FDateTime>& ClientTimePair : *TransactionClients)
+		{
+			const FGuid& ClientId = ClientTimePair.Key;
+
+			Event.TransactionId = GetClientTransactionId(ClientId, TransactionGuid);
+			Event.SequenceNumber = GetSequenceNumber(ClientId);
+
+			TArray<uint8> Payload;
+			WebRemoteControlUtils::SerializeMessage(Event, Payload);
+			Server->Send(ClientId, Payload);
+		}
+
+		// Forget about the transaction
+		ClientsByTransactionGuid.Remove(TransactionGuid);
+
+		for (TPair<FGuid, TMap<FGuid, int32>>& ClientIdsPair : TransactionIdsByClientId)
+		{
+			ClientIdsPair.Value.Remove(TransactionGuid);
+		}
+	}
+}
+
 void FWebSocketMessageHandler::StartWatchingActor(AActor* Actor, UClass* WatchedClass)
 {
 	FWatchedActorData* ActorData = WatchedActors.Find(Actor);
@@ -1676,3 +1875,87 @@ void FWebSocketMessageHandler::UpdateSequenceNumber(const FGuid& ClientId, int64
 		StoredSequenceNumber = NewSequenceNumber;
 	}
 }
+
+bool FWebSocketMessageHandler::ContributeToTransaction(const FGuid& ClientId, int32 TransactionId)
+{
+	const FGuid InternalId = GetTransactionGuid(ClientId, TransactionId);
+	if (ClientId.IsValid())
+	{
+		if (TMap<FGuid, FDateTime>* TransactionClients = ClientsByTransactionGuid.Find(InternalId))
+		{
+			TransactionClients->Add(ClientId, FDateTime::Now());
+			return true;
+		}
+	}
+
+	return false;
+}
+
+void FWebSocketMessageHandler::EndClientTransaction(const FGuid& ClientId, int32 TransactionId)
+{
+	if (TMap<FGuid, int32>* TransactionIds = TransactionIdsByClientId.Find(ClientId))
+	{
+		const FGuid TransactionGuid = GetTransactionGuid(ClientId, TransactionId);
+
+		if (!TransactionGuid.IsValid())
+		{
+			return;
+		}
+
+		// Stop tracking the transaction for this client
+		if (TMap<FGuid, FDateTime>* TransactionClients = ClientsByTransactionGuid.Find(TransactionGuid))
+		{
+			const int32 Result = IRemoteControlModule::Get().EndManualEditorTransaction(TransactionGuid);
+
+			if (Result != INDEX_NONE && Result <= 1)
+			{
+				// This was the last action, so the transaction will end. This will clean up the client and ID maps for us.
+				HandleTransactionEnded(TransactionGuid);
+				return;
+			}
+
+			// Transaction still exists, but we can remove the client from its list
+			TransactionClients->Remove(ClientId);
+		}
+
+		// Do this last so we can still look up the client ID/transaction GUID mapping in HandleTransactionEnded
+		TransactionIds->Remove(TransactionGuid);
+	}
+}
+
+int64 FWebSocketMessageHandler::GetSequenceNumber(const FGuid& ClientId) const
+{
+	const int64* ClientSequenceNumber = ClientSequenceNumbers.Find(ClientId);
+	return ClientSequenceNumber ? *ClientSequenceNumber : DefaultSequenceNumber;
+}
+
+FGuid FWebSocketMessageHandler::GetTransactionGuid(const FGuid& ClientId, int32 TransactionId) const
+{
+	if (const TMap<FGuid, int32>* TransactionIdPairs = TransactionIdsByClientId.Find(ClientId))
+	{
+		for (const TPair<FGuid, int32>& Pair : *TransactionIdPairs)
+		{
+			if (TransactionId == Pair.Value)
+			{
+				return Pair.Key;
+			}
+		}
+	}
+
+	return FGuid();
+}
+
+int32 FWebSocketMessageHandler::GetClientTransactionId(const FGuid& ClientId, const FGuid& TransactionGuid) const
+{
+	if (const TMap<FGuid, int32>* TransactionIds = TransactionIdsByClientId.Find(ClientId))
+	{
+		if (const int32* ClientTransactionId = TransactionIds->Find(TransactionGuid))
+		{
+			return *ClientTransactionId;
+		}
+	}
+
+	return InvalidTransactionId;
+}
+
+#undef LOCTEXT_NAMESPACE
diff --git a/Engine/Plugins/VirtualProduction/RemoteControl/Source/WebRemoteControl/Private/WebSocketMessageHandler.h b/Engine/Plugins/VirtualProduction/RemoteControl/Source/WebRemoteControl/Private/WebSocketMessageHandler.h
--- a/Engine/Plugins/VirtualProduction/RemoteControl/Source/WebRemoteControl/Private/WebSocketMessageHandler.h
+++ b/Engine/Plugins/VirtualProduction/RemoteControl/Source/WebRemoteControl/Private/WebSocketMessageHandler.h
@@ -105,6 +105,12 @@ private:
 	/** Handles calling a Blueprint function */
 	void HandleWebSocketFunctionCall(const FRemoteControlWebSocketMessage& WebSocketMessage);
 
+	/** Handles beginning a manual editor transaction. */
+	void HandleWebSocketBeginEditorTransaction(const FRemoteControlWebSocketMessage& WebSocketMessage);
+
+	/** Handles ending a manual editor transaction. */
+	void HandleWebSocketEndEditorTransaction(const FRemoteControlWebSocketMessage& WebSocketMessage);
+
 	//Preset callbacks
 	void OnPresetExposedPropertiesModified(URemoteControlPreset* Owner, const TSet<FGuid>& ModifiedPropertyIds);
 	void OnPropertyExposed(URemoteControlPreset* Owner,  const FGuid& EntityId);
@@ -121,6 +127,9 @@ private:
 	/** End of frame callback to send cached property changed, preset changed messages */
 	void OnEndFrame();
 
+	/** Check if any transactions have timed out and end them. */
+	void TimeOutTransactions();
+
 	/** If properties have changed during the frame, send out notifications to listeners */
 	void ProcessChangedProperties();
 
@@ -214,6 +223,18 @@ private:
 	 */
 	void OnObjectTransacted(UObject* Object, const class FTransactionObjectEvent& TransactionEvent);
 
+#if WITH_EDITOR
+	/**
+	 * Called when the state of an editor transaction changes.
+	 */
+	void HandleTransactionStateChanged(const FTransactionContext& InTransactionContext, const ETransactionStateEventType InTransactionState);
+#endif
+
+	/**
+	 * Handle a transaction ending (either cancelled or finalized).
+	 */
+	void HandleTransactionEnded(const FGuid& TransactionGuid);
+
 	/**
 	 * Start watching an actor because it's a member of the given class.
 	 */
@@ -239,12 +260,45 @@ private:
 	 * Update the sequence number for a client when a new one is received.
 	 */
 	void UpdateSequenceNumber(const FGuid& ClientId, int64 NewSequenceNumber);
+	
+	/**
+	 * Get the current sequence number for a client.
+	 */
+	int64 GetSequenceNumber(const FGuid& ClientId) const;
+
+	/**
+	 * Indicate that a client is going to contribute to the transaction with the given ID.
+	 * This should be called whenever a change is made that will be part of the transaction, not just at the start.
+	 * Returns true if the client can contribute to a transaction with that ID.
+	 */
+	bool ContributeToTransaction(const FGuid& ClientId, int32 TransactionId);
+
+	/**
+	 * End the transaction with the given ID for the client with the given ID.
+	 */
+	void EndClientTransaction(const FGuid& ClientId, int32 TransactionId);
+
+	/**
+	 * Converts a client's transaction ID to its editor internal GUID.
+	 */
+	FGuid GetTransactionGuid(const FGuid& ClientId, int32 TransactionId) const;
+
+	/**
+	 * Converts an internal transaction GUID to the ID used by the given client to refer to it.
+	 */
+	int32 GetClientTransactionId(const FGuid& ClientId, const FGuid& TransactionGuid) const;
 
 private:
 
 	/** Default sequence number for a client that hasn't reported one yet. */
 	static const int64 DefaultSequenceNumber;
 
+	/** Client transaction ID for a transaction that doesn't exist. */
+	static const int32 InvalidTransactionId;
+
+	/** When this much time has passed since a client last contributed to a transaction, the transaction will automatically end. */
+	static const FTimespan TransactionTimeout;
+
 	/** Map type from class to Guids of clients listening for changes to actors of that class. */
 	typedef TMap<TWeakObjectPtr<UClass>, FWatchedClassData, FDefaultSetAllocator, TWeakObjectPtrMapKeyFuncs<TWeakObjectPtr<UClass>, FWatchedClassData>> FActorNotificationMap;
 
@@ -316,6 +370,12 @@ private:
 
 	/** Map from transient preset ID to clients which, when all disconnected, will automatically destroy the preset. */
 	TMap<FGuid, TArray<FGuid>> TransientPresetAutoDestroyClients;
+
+	/** Map from transasction ID to a map of (contributing client ID, time of last contribution to the transaction). */
+	TMap<FGuid, TMap<FGuid, FDateTime>> ClientsByTransactionGuid;
+
+	/** Map from client ID to pairs of editor GUID/client ID used to refer to the transaction. */
+	TMap<FGuid, TMap<FGuid, int32>> TransactionIdsByClientId;
 	
 	/** Holds the ID of the client currently making a request. Used to prevent sending back notifications to it. */
 	const FGuid& ActingClientId;
diff --git a/Engine/Plugins/VirtualProduction/RemoteControlInterception/Source/RemoteControlInterception/Public/IRemoteControlInterceptionCommands.h b/Engine/Plugins/VirtualProduction/RemoteControlInterception/Source/RemoteControlInterception/Public/IRemoteControlInterceptionCommands.h
--- a/Engine/Plugins/VirtualProduction/RemoteControlInterception/Source/RemoteControlInterception/Public/IRemoteControlInterceptionCommands.h
+++ b/Engine/Plugins/VirtualProduction/RemoteControlInterception/Source/RemoteControlInterception/Public/IRemoteControlInterceptionCommands.h
@@ -37,6 +37,7 @@ enum class ERCIAccess : uint8
 	READ_ACCESS,
 	WRITE_ACCESS,
 	WRITE_TRANSACTION_ACCESS,
+	WRITE_MANUAL_TRANSACTION_ACCESS,
 };
 
 /**
