[WebRemoteControl] Don't mix together struct fields when sending property updates

diff --git a/Engine/Plugins/VirtualProduction/RemoteControl/Source/WebRemoteControl/Private/WebSocketMessageHandler.cpp b/Engine/Plugins/VirtualProduction/RemoteControl/Source/WebRemoteControl/Private/WebSocketMessageHandler.cpp
--- a/Engine/Plugins/VirtualProduction/RemoteControl/Source/WebRemoteControl/Private/WebSocketMessageHandler.cpp
+++ b/Engine/Plugins/VirtualProduction/RemoteControl/Source/WebRemoteControl/Private/WebSocketMessageHandler.cpp
@@ -231,7 +231,7 @@ namespace WebSocketMessageHandlerMiscUtils
 #endif
 	}
 
-	uint64 GetPresetPropertyClassId(URemoteControlPreset* Preset, const FGuid& PropertyId)
+	void* GetPresetPropertyClassPointer(URemoteControlPreset* Preset, const FGuid& PropertyId)
 	{
 		FRCObjectReference ObjectRef;
 		if (TSharedPtr<FRemoteControlProperty> RCProperty = Preset->GetExposedEntity<FRemoteControlProperty>(PropertyId).Pin())
@@ -240,12 +240,17 @@ namespace WebSocketMessageHandlerMiscUtils
 			{
 				if (const FProperty* Property = RCProperty->GetProperty())
 				{
-					return Property->GetClass()->GetId();
+					if (const FStructProperty* StructProperty = CastField<FStructProperty>(Property))
+					{
+						return StructProperty->Struct;
+					}
+
+					return Property->GetClass();
 				}
 			}
 		}
 
-		return 0;
+		return nullptr;
 	}
 }
 
@@ -791,15 +796,15 @@ void FWebSocketMessageHandler::ProcessChangedProperties()
 		for (const TPair<FGuid, TSet<FGuid>>& ClientToEventsPair : Entry.Value)
 		{
 			// Categorize modified properties by type so we don't try to put multiple types of data in a single request
-			TMap<uint64, TSet<FGuid>> PropertyIdsByType;
+			TMap<void*, TSet<FGuid>> PropertyIdsByType;
 			for (const FGuid& Id : ClientToEventsPair.Value)
 			{
-				const uint64 ClassId = WebSocketMessageHandlerMiscUtils::GetPresetPropertyClassId(Preset, Id);
-				PropertyIdsByType.FindOrAdd(ClassId).Emplace(Id);
+				void* ClassPointer = WebSocketMessageHandlerMiscUtils::GetPresetPropertyClassPointer(Preset, Id);
+				PropertyIdsByType.FindOrAdd(ClassPointer).Emplace(Id);
 			}
 
 			// Send a property change event for each property type
-			for (const TPair<uint64, TSet<FGuid>>& ClassToEventsPair : PropertyIdsByType)
+			for (const TPair<void*, TSet<FGuid>>& ClassToEventsPair : PropertyIdsByType)
 			{
 				const int64 SequenceNumber = GetSequenceNumber(ClientToEventsPair.Key);
 
