[WebRemoteControl] Move manual transaction timeout to a CVar

diff --git a/Engine/Plugins/VirtualProduction/RemoteControl/Source/WebRemoteControl/Private/WebSocketMessageHandler.cpp b/Engine/Plugins/VirtualProduction/RemoteControl/Source/WebRemoteControl/Private/WebSocketMessageHandler.cpp
--- a/Engine/Plugins/VirtualProduction/RemoteControl/Source/WebRemoteControl/Private/WebSocketMessageHandler.cpp
+++ b/Engine/Plugins/VirtualProduction/RemoteControl/Source/WebRemoteControl/Private/WebSocketMessageHandler.cpp
@@ -23,10 +23,20 @@
 
 #define LOCTEXT_NAMESPACE "WebRemoteControl"
 
-static TAutoConsoleVariable<int32> CVarWebRemoteControlFramesBetweenPropertyNotifications(TEXT("WebControl.FramesBetweenPropertyNotifications"), 5, TEXT("The number of frames between sending batches of property notifications."));
+static TAutoConsoleVariable<int32> CVarWebRemoteControlFramesBetweenPropertyNotifications(
+	TEXT("WebControl.FramesBetweenPropertyNotifications"),
+	5,
+	TEXT("The number of frames between sending batches of property notifications.")
+);
+
+static TAutoConsoleVariable<float> CVarWebRemoteControlManualTransactionTimeout(
+	TEXT("WebControl.ManualTransactionTimeout"),
+	3.f,
+	TEXT("When this many seconds has passed since a client last contributed to a transaction, the transaction will automatically end.")
+);
+
 const int64 FWebSocketMessageHandler::DefaultSequenceNumber = -1;
 const int32 FWebSocketMessageHandler::InvalidTransactionId = -1;
-const FTimespan FWebSocketMessageHandler::TransactionTimeout = FTimespan::FromSeconds(3.f);
 
 namespace WebSocketMessageHandlerStructUtils
 {
@@ -1124,13 +1134,13 @@ void FWebSocketMessageHandler::OnEndFrame()
 void FWebSocketMessageHandler::TimeOutTransactions()
 {
 	TArray<TPair<FGuid, int32>> TransactionsToEnd;
-	const FDateTime Now = FDateTime::Now();
+	const FDateTime TimeoutCutoff = FDateTime::Now() - FTimespan::FromSeconds(CVarWebRemoteControlManualTransactionTimeout.GetValueOnGameThread());
 
 	for (const TPair<FGuid, TMap<FGuid, FDateTime>>& TransactionClientsPair : ClientsByTransactionGuid)
 	{
 		for (const TPair<FGuid, FDateTime>& ClientTimePair : TransactionClientsPair.Value)
 		{
-			if (Now >= ClientTimePair.Value + TransactionTimeout)
+			if (ClientTimePair.Value < TimeoutCutoff)
 			{
 				// This transaction has timed out; we should force it to end
 				TransactionsToEnd.Add({ ClientTimePair.Key, GetClientTransactionId(ClientTimePair.Key, TransactionClientsPair.Key) });
diff --git a/Engine/Plugins/VirtualProduction/RemoteControl/Source/WebRemoteControl/Private/WebSocketMessageHandler.h b/Engine/Plugins/VirtualProduction/RemoteControl/Source/WebRemoteControl/Private/WebSocketMessageHandler.h
--- a/Engine/Plugins/VirtualProduction/RemoteControl/Source/WebRemoteControl/Private/WebSocketMessageHandler.h
+++ b/Engine/Plugins/VirtualProduction/RemoteControl/Source/WebRemoteControl/Private/WebSocketMessageHandler.h
@@ -296,9 +296,6 @@ private:
 	/** Client transaction ID for a transaction that doesn't exist. */
 	static const int32 InvalidTransactionId;
 
-	/** When this much time has passed since a client last contributed to a transaction, the transaction will automatically end. */
-	static const FTimespan TransactionTimeout;
-
 	/** Map type from class to Guids of clients listening for changes to actors of that class. */
 	typedef TMap<TWeakObjectPtr<UClass>, FWatchedClassData, FDefaultSetAllocator, TWeakObjectPtrMapKeyFuncs<TWeakObjectPtr<UClass>, FWatchedClassData>> FActorNotificationMap;
 
