~~~ Remote Control Logic Core / Runtime changes ~~~

Added new function to URCVirtualPropertyBase class:

* CopyCompleteValue - Provides support for copying the controller value onto a given FProperty. This will be used the RC Web API's GetPresetController route.

diff --git a/Engine/Plugins/VirtualProduction/RemoteControl/Source/RemoteControl/Private/RCVirtualProperty.cpp b/Engine/Plugins/VirtualProduction/RemoteControl/Source/RemoteControl/Private/RCVirtualProperty.cpp
--- a/Engine/Plugins/VirtualProduction/RemoteControl/Source/RemoteControl/Private/RCVirtualProperty.cpp
+++ b/Engine/Plugins/VirtualProduction/RemoteControl/Source/RemoteControl/Private/RCVirtualProperty.cpp
@@ -145,6 +145,30 @@ bool URCVirtualPropertyBase::IsValueEqual(URCVirtualPropertyBase* InVirtualPrope
 	return ThisProperty->Identical(ThisValuePtr, CompareValuePtr);
 }
 
+bool URCVirtualPropertyBase::CopyCompleteValue(const FProperty* InTargetProperty, uint8* InTargetValuePtr)
+{
+	const FProperty* SourceProperty = GetProperty();
+
+	if (SourceProperty == nullptr || InTargetProperty == nullptr || InTargetValuePtr == nullptr)
+	{
+		ensureMsgf(false, TEXT("Invalid input passed to CopyCompleteValue"));
+
+		return false;
+	}
+
+	if (SourceProperty->GetClass() != InTargetProperty->GetClass())
+	{
+		ensureMsgf(false, TEXT("Invalid property type passed to CopyCompleteValue.\nExpected: %s, Found: %s"), *SourceProperty->GetClass()->GetName(), *InTargetProperty->GetClass()->GetName());
+
+		return false;
+	}
+
+	SourceProperty->CopyCompleteValue(InTargetValuePtr /*Dest*/, GetValuePtr() /*Source*/);
+
+	return true;
+
+}
+
 struct FRCVirtualPropertyCastHelpers
 {
 	template<typename TPropertyType, typename TValueType>
diff --git a/Engine/Plugins/VirtualProduction/RemoteControl/Source/RemoteControl/Public/RCVirtualProperty.h b/Engine/Plugins/VirtualProduction/RemoteControl/Source/RemoteControl/Public/RCVirtualProperty.h
--- a/Engine/Plugins/VirtualProduction/RemoteControl/Source/RemoteControl/Public/RCVirtualProperty.h
+++ b/Engine/Plugins/VirtualProduction/RemoteControl/Source/RemoteControl/Public/RCVirtualProperty.h
@@ -65,6 +65,13 @@ public:
 
 	/** Compare this virtual property value with given property value */
 	bool IsValueEqual(URCVirtualPropertyBase* InVirtualProperty);
+
+	/** Copy this virtual property's data onto a given FProperty
+	* 
+	* * @param InTargetProperty - The property onto which our value is to be copied
+	* * @param InTargetValuePtr - The memory location for the target property
+	*/
+	bool CopyCompleteValue(const FProperty* InTargetProperty, uint8* InTargetValuePtr);
 	
 	/** Get Bool value from Virtual Property */
 	UFUNCTION(BlueprintCallable, Category = "Remote Control Behaviour")
