Initialize TransactionIds
#rnx

diff --git a/Engine/Plugins/VirtualProduction/RemoteControl/Source/WebRemoteControl/Private/RemoteControlRequest.h b/Engine/Plugins/VirtualProduction/RemoteControl/Source/WebRemoteControl/Private/RemoteControlRequest.h
--- a/Engine/Plugins/VirtualProduction/RemoteControl/Source/WebRemoteControl/Private/RemoteControlRequest.h
+++ b/Engine/Plugins/VirtualProduction/RemoteControl/Source/WebRemoteControl/Private/RemoteControlRequest.h
@@ -789,7 +789,7 @@
 	 * The ID that will be used to refer to the transaction in future messages.
 	 */
 	UPROPERTY()
-	int32 TransactionId;
+	int32 TransactionId = -1;
 };
 
 
@@ -815,5 +815,5 @@
 	 * The ID of the transaction. If this doesn't match the current editor transaction, it won't be ended.
 	 */
 	UPROPERTY()
-	int32 TransactionId;
+	int32 TransactionId = -1;
 };
diff --git a/Engine/Plugins/VirtualProduction/RemoteControl/Source/WebRemoteControl/Private/RemoteControlResponse.h b/Engine/Plugins/VirtualProduction/RemoteControl/Source/WebRemoteControl/Private/RemoteControlResponse.h
--- a/Engine/Plugins/VirtualProduction/RemoteControl/Source/WebRemoteControl/Private/RemoteControlResponse.h
+++ b/Engine/Plugins/VirtualProduction/RemoteControl/Source/WebRemoteControl/Private/RemoteControlResponse.h
@@ -468,7 +468,7 @@
 	 * The client-specific ID of the transaction that was ended.
 	 */
 	UPROPERTY()
-	int32 TransactionId;
+	int32 TransactionId = -1;
 
 	/**
 	 * The highest sequence number received from the receiving client at the time that the transaction ended.
