1. Remote control preset UI revamp

2. Rework grouping UI

diff --git a/Engine/Plugins/VirtualProduction/RemoteControl/Source/RemoteControl/Private/RemoteControlPreset.cpp b/Engine/Plugins/VirtualProduction/RemoteControl/Source/RemoteControl/Private/RemoteControlPreset.cpp
--- a/Engine/Plugins/VirtualProduction/RemoteControl/Source/RemoteControl/Private/RemoteControlPreset.cpp
+++ b/Engine/Plugins/VirtualProduction/RemoteControl/Source/RemoteControl/Private/RemoteControlPreset.cpp
@@ -45,7 +45,7 @@ static TAutoConsoleVariable<int32> CVarRemoteControlFramesBetweenPropertyWatch(T
 namespace
 {
 	FGuid DefaultGroupId = FGuid(0x5DFBC958, 0xF3B311EA, 0x9A3F00EE, 0xFB2CA371);
-	FName NAME_DefaultLayoutGroup = FName("Default Group");
+	FName NAME_DefaultLayoutGroup = FName("All");
 	FName NAME_DefaultNewGroup = FName("New Group");
 	const FString DefaultObjectPrefix = TEXT("Default__");
 
@@ -149,6 +149,26 @@ FRemoteControlPresetGroup& FRemoteControlPresetLayout::GetDefaultGroup()
 	}
 }
 
+bool FRemoteControlPresetLayout::IsDefaultGroup(FGuid GroupId) const
+{
+	return GroupId == DefaultGroupId;
+}
+
+FLinearColor FRemoteControlPresetLayout::GetTagColor(FGuid GroupId)
+{
+	if (IsDefaultGroup(GroupId))
+	{
+		return FLinearColor::Transparent;
+	}
+
+	if (FRemoteControlPresetGroup* Group = GetGroup(GroupId))
+	{
+		return Group->TagColor;
+	}
+
+	return FLinearColor::White;
+}
+
 FRemoteControlPresetGroup* FRemoteControlPresetLayout::GetGroup(FGuid GroupId)
 {
 	return Groups.FindByPredicate([GroupId](const FRemoteControlPresetGroup& Group) { return Group.Id == GroupId; });
diff --git a/Engine/Plugins/VirtualProduction/RemoteControl/Source/RemoteControl/Public/RemoteControlPreset.h b/Engine/Plugins/VirtualProduction/RemoteControl/Source/RemoteControl/Public/RemoteControlPreset.h
--- a/Engine/Plugins/VirtualProduction/RemoteControl/Source/RemoteControl/Public/RemoteControlPreset.h
+++ b/Engine/Plugins/VirtualProduction/RemoteControl/Source/RemoteControl/Public/RemoteControlPreset.h
@@ -101,6 +101,7 @@ struct REMOTECONTROL_API FRemoteControlPresetGroup
 	FRemoteControlPresetGroup(FName InName, FGuid InId)
 		: Name(InName)
 		, Id(MoveTemp(InId))
+		, TagColor(MakeSimilarSaturatedColor())
 	{}
 
 	/** Get the fields under this group. */
@@ -113,6 +114,18 @@ struct REMOTECONTROL_API FRemoteControlPresetGroup
 	{
 		return LHS.Id == RHS.Id;
 	}
+
+private:
+
+	/**
+	* Makes a random but similar saturated color.
+	*/
+	static FLinearColor MakeSimilarSaturatedColor(float InSaturationLevel = 0.5f)
+	{
+		const uint8 Hue = (uint8)(FMath::FRand() * 255.f);
+		const uint8 Saturation = (uint8)(255.f * InSaturationLevel);
+		return FLinearColor::MakeFromHSV8(Hue, Saturation, 255);
+	}
  
 public:
 	/** Name of this group. */
@@ -123,6 +136,10 @@ public:
 	UPROPERTY()
 	FGuid Id;
 
+	/* Color Tag for this group. */
+	UPROPERTY()
+	FLinearColor TagColor;
+
 private:
 	/** The list of exposed fields under this group. */
 	UPROPERTY()
@@ -150,6 +167,12 @@ struct REMOTECONTROL_API FRemoteControlPresetLayout
 	/** Get or create the default group. */
 	FRemoteControlPresetGroup& GetDefaultGroup();
 
+	/** Returns true when the given group id is a default one. */
+	bool IsDefaultGroup(FGuid GroupId) const;
+	
+	/** Returns the tag color of the given group id. */
+	FLinearColor GetTagColor(FGuid GroupId);
+
 	/**
 	 * Get a group by searching by ID.
 	 * @param GroupId the id to use.
diff --git a/Engine/Plugins/VirtualProduction/RemoteControl/Source/RemoteControlUI/Private/Commands/RemoteControlCommands.cpp b/Engine/Plugins/VirtualProduction/RemoteControl/Source/RemoteControlUI/Private/Commands/RemoteControlCommands.cpp
--- a/Engine/Plugins/VirtualProduction/RemoteControl/Source/RemoteControlUI/Private/Commands/RemoteControlCommands.cpp
+++ b/Engine/Plugins/VirtualProduction/RemoteControl/Source/RemoteControlUI/Private/Commands/RemoteControlCommands.cpp
@@ -36,6 +36,12 @@ void FRemoteControlCommands::RegisterCommands()
 
 	// Toggle Logic Editor
 	UI_COMMAND(ToggleLogicEditor, "Logic", "View the logic applied to active selection.", EUserInterfaceActionType::Button, FInputChord());
+
+	// Delete Entity
+	UI_COMMAND(DeleteEntity, "Delete", "Delete the selected  group/exposed entity from the list.", EUserInterfaceActionType::Button, FInputChord(EKeys::Delete));
+
+	// Rename Entity
+	UI_COMMAND(RenameEntity, "Rename", "Rename the selected  group/exposed entity.", EUserInterfaceActionType::Button, FInputChord(EKeys::F2));
 }
 
 #undef LOCTEXT_NAMESPACE
diff --git a/Engine/Plugins/VirtualProduction/RemoteControl/Source/RemoteControlUI/Private/Commands/RemoteControlCommands.h b/Engine/Plugins/VirtualProduction/RemoteControl/Source/RemoteControlUI/Private/Commands/RemoteControlCommands.h
--- a/Engine/Plugins/VirtualProduction/RemoteControl/Source/RemoteControlUI/Private/Commands/RemoteControlCommands.h
+++ b/Engine/Plugins/VirtualProduction/RemoteControl/Source/RemoteControlUI/Private/Commands/RemoteControlCommands.h
@@ -50,4 +50,14 @@ public:
 	 * Holds the information about UI Command that brings up a panel which enables the RC Logical Behaviour.
 	 */
 	TSharedPtr<FUICommandInfo> ToggleLogicEditor;
+
+	/**
+	 * Holds the information about UI Command that deletes currently selected group/exposed entity.
+	 */
+	TSharedPtr< FUICommandInfo > DeleteEntity;
+
+	/**
+	 * Holds the information about UI Command that  renames selected group/exposed entity.
+	 */
+	TSharedPtr< FUICommandInfo > RenameEntity;
 };
diff --git a/Engine/Plugins/VirtualProduction/RemoteControl/Source/RemoteControlUI/Private/UI/SRCPanelExposedEntitiesList.cpp b/Engine/Plugins/VirtualProduction/RemoteControl/Source/RemoteControlUI/Private/UI/SRCPanelExposedEntitiesList.cpp
--- a/Engine/Plugins/VirtualProduction/RemoteControl/Source/RemoteControlUI/Private/UI/SRCPanelExposedEntitiesList.cpp
+++ b/Engine/Plugins/VirtualProduction/RemoteControl/Source/RemoteControlUI/Private/UI/SRCPanelExposedEntitiesList.cpp
@@ -4,6 +4,7 @@
 
 #include "Algo/ForEach.h"
 #include "Editor.h"
+#include "EditorFontGlyphs.h"
 #include "ISettingsModule.h"
 #include "RCPanelWidgetRegistry.h"
 #include "GameFramework/Actor.h"
@@ -24,6 +25,15 @@
 
 #define LOCTEXT_NAMESPACE "RemoteControlPanelEntitiesList"
 
+namespace FRemoteControlPresetColumns
+{
+	const FName DragDropHandle = TEXT("DragDropHandle");
+	const FName IsBound = TEXT("IsBound");
+	const FName Description = TEXT("Description");
+	const FName Value = TEXT("Value");
+	const FName Groups = TEXT("Groups");
+}
+
 void SRCPanelExposedEntitiesList::Construct(const FArguments& InArgs, URemoteControlPreset* InPreset, TWeakPtr<FRCPanelWidgetRegistry> InWidgetRegistry)
 {
 	bIsInEditMode = InArgs._EditMode;
@@ -37,14 +47,132 @@ void SRCPanelExposedEntitiesList::Construct(const FArguments& InArgs, URemoteCon
 
 	ChildSlot
 	[
-		SAssignNew(TreeView, STreeView<TSharedPtr<SRCPanelTreeNode>>)
-			.TreeItemsSource(reinterpret_cast<TArray<TSharedPtr<SRCPanelTreeNode>>*>(&FieldGroups))
-			.ItemHeight(24.0f)
-			.OnGenerateRow(this, &SRCPanelExposedEntitiesList::OnGenerateRow)
-			.OnGetChildren(this, &SRCPanelExposedEntitiesList::OnGetNodeChildren)
-			.OnSelectionChanged(this, &SRCPanelExposedEntitiesList::OnSelectionChanged)
-			.OnContextMenuOpening(this, &SRCPanelExposedEntitiesList::OnContextMenuOpening)
-			.ClearSelectionOnClick(true)
+		SNew(SVerticalBox)
+		
+		+SVerticalBox::Slot()
+		.FillHeight(1.f)
+		[
+			SNew(SHorizontalBox)
+			
+			+SHorizontalBox::Slot()
+			.Padding(2.f, 4.f)
+			.FillWidth(0.25f)
+			[
+				SNew(SBorder)
+				.BorderImage(FAppStyle::Get().GetBrush("ToolPanel.DarkGroupBorder"))
+				[
+					SAssignNew(GroupsListView, SListView<TSharedPtr<SRCPanelTreeNode>>)
+					.ItemHeight(24.f)
+					.OnGenerateRow(this, &SRCPanelExposedEntitiesList::OnGenerateRow)
+					.OnSelectionChanged(this, &SRCPanelExposedEntitiesList::OnSelectionChanged)
+					.SelectionMode(ESelectionMode::Single)
+					.ListItemsSource(reinterpret_cast<TArray<TSharedPtr<SRCPanelTreeNode>>*>(&FieldGroups))
+					.OnContextMenuOpening(this, &SRCPanelExposedEntitiesList::OnContextMenuOpening, SRCPanelTreeNode::Group)
+					.ClearSelectionOnClick(true)
+					.HeaderRow(
+						SNew(SHeaderRow)
+
+						+ SHeaderRow::Column(FRemoteControlPresetColumns::Groups)
+						.DefaultLabel(LOCTEXT("RCPresetGroupsColumnHeader", "Groups"))
+						.FillWidth(1.f)
+						.ToolTipText(LOCTEXT("NewGroupToolTip", "Create new group."))
+						[
+							SNew(SHorizontalBox)
+				
+							+SHorizontalBox::Slot()
+							.Padding(3.f)
+							.VAlign(VAlign_Center)
+							.AutoWidth()
+							[
+								SNew(STextBlock)
+								.ColorAndOpacity(FColor::White)
+								.Font(FCoreStyle::GetDefaultFontStyle("Bold", 12.f))
+								.Text(LOCTEXT("GroupsLabel", "Groups"))
+							]
+	
+							+SHorizontalBox::Slot()
+							.Padding(3.f)
+							.VAlign(VAlign_Center)
+							.FillWidth(1.f)
+							[
+								SNew(SSpacer)
+							]
+	
+							+ SHorizontalBox::Slot()
+							.Padding(3.f)
+							.HAlign(HAlign_Center)
+							.VAlign(VAlign_Center)
+							.AutoWidth()
+							[
+								SNew(SButton)
+								.Visibility_Lambda([this]() { return bIsInEditMode.Get() ? EVisibility::Visible : EVisibility::Collapsed; })
+								.ButtonStyle(FAppStyle::Get(), "ToggleButton")
+								.ToolTipText(LOCTEXT("NewGroupToolTip", "Create new group."))
+								.OnClicked(this, &SRCPanelExposedEntitiesList::OnCreateGroup)
+								[
+									SNew(SOverlay)
+
+									+SOverlay::Slot()
+									.HAlign(HAlign_Center)
+									.VAlign(VAlign_Center)
+									[
+										SNew(STextBlock)
+										.TextStyle(FRemoteControlPanelStyle::Get(), "RemoteControlPanel.Button.TextStyle")
+										.ColorAndOpacity(FColor::White)
+										.Font(FAppStyle::Get().GetFontStyle("FontAwesome.14"))
+										.Text(FEditorFontGlyphs::Circle_O)
+									]
+
+									+SOverlay::Slot()
+									.HAlign(HAlign_Center)
+									.VAlign(VAlign_Center)
+									[
+										SNew(STextBlock)
+										.TextStyle(FRemoteControlPanelStyle::Get(), "RemoteControlPanel.Button.TextStyle")
+										.ColorAndOpacity(FColor::White)
+										.Font(FAppStyle::Get().GetFontStyle("FontAwesome.7"))
+										.Text(FEditorFontGlyphs::Plus)
+									]
+								]
+							]
+						]
+					)
+				]
+			]
+
+			+SHorizontalBox::Slot()
+			.Padding(2.f, 4.f)
+			.FillWidth(0.75f)
+			[
+				SNew(SBorder)
+				.BorderImage(FAppStyle::Get().GetBrush("ToolPanel.DarkGroupBorder"))
+				[
+					SAssignNew(FieldsListView, SListView<TSharedPtr<SRCPanelTreeNode>>)
+					.ItemHeight(24.f)
+					.OnGenerateRow(this, &SRCPanelExposedEntitiesList::OnGenerateRow)
+					.OnSelectionChanged(this, &SRCPanelExposedEntitiesList::OnSelectionChanged)
+					.SelectionMode(ESelectionMode::Single)
+					.ListItemsSource(&FieldEntities)
+					.OnContextMenuOpening(this, &SRCPanelExposedEntitiesList::OnContextMenuOpening, SRCPanelTreeNode::Field)
+					.ClearSelectionOnClick(true)
+					.HeaderRow(
+						SNew(SHeaderRow)
+
+						+ SHeaderRow::Column(FRemoteControlPresetColumns::DragDropHandle)
+						.DefaultLabel(LOCTEXT("RCPresetDragDropHandleColumnHeader", ""))
+						.FixedWidth(36.0f)
+
+						+ SHeaderRow::Column(FRemoteControlPresetColumns::Description)
+						.DefaultLabel(LOCTEXT("RCPresetDescColumnHeader", "Description"))
+						.FillWidth(0.7f)
+
+						+ SHeaderRow::Column(FRemoteControlPresetColumns::Value)
+						.DefaultLabel(LOCTEXT("RCPresetValueColumnHeader", "Value"))
+						.FillWidth(0.7f)
+					)
+				]
+			]
+		]
 	];
 
 	RegisterEvents();
@@ -57,30 +185,49 @@ SRCPanelExposedEntitiesList::~SRCPanelExposedEntitiesList()
 	UnregisterEvents();
 }
 
-TSharedPtr<SRCPanelTreeNode> SRCPanelExposedEntitiesList::GetSelection() const
+TSharedPtr<SRCPanelTreeNode> SRCPanelExposedEntitiesList::GetSelectedGroup() const
 {
 	TArray<TSharedPtr<SRCPanelTreeNode>> SelectedNodes;
-	TreeView->GetSelectedItems(SelectedNodes);
+
+	GroupsListView->GetSelectedItems(SelectedNodes);
+
 	if (SelectedNodes.Num())
 	{
 		return SelectedNodes[0];
 	}
+
 	return nullptr;
 }
 
-void SRCPanelExposedEntitiesList::SetSelection(const TSharedPtr<SRCPanelTreeNode>& Node)
+TSharedPtr<SRCPanelTreeNode> SRCPanelExposedEntitiesList::GetSelectedEntity() const
+{
+	TArray<TSharedPtr<SRCPanelTreeNode>> SelectedNodes;
+
+	FieldsListView->GetSelectedItems(SelectedNodes);
+
+	if (SelectedNodes.Num())
+	{
+		return SelectedNodes[0];
+	}
+
+	return nullptr;
+}
+
+void SRCPanelExposedEntitiesList::SetSelection(const TSharedPtr<SRCPanelTreeNode>& Node, const bool bForceMouseClick)
 {
 	if (Node)
 	{
+		const ESelectInfo::Type SelectInfo = bForceMouseClick ? ESelectInfo::OnMouseClick : ESelectInfo::Direct;
+
 		if (TSharedPtr<SRCPanelTreeNode>* FoundTreeNode = FieldWidgetMap.Find(Node->GetRCId()))
 		{
-			TreeView->SetSelection(*FoundTreeNode);
+			FieldsListView->SetSelection(*FoundTreeNode, SelectInfo);
 			return;
 		}
 
 		if (TSharedPtr<SRCPanelGroup> SRCGroup = FindGroupById(Node->GetRCId()))
 		{
-			TreeView->SetSelection(SRCGroup);
+			GroupsListView->SetSelection(SRCGroup, SelectInfo);
 		}	
 	}
 }
@@ -132,22 +279,23 @@ void SRCPanelExposedEntitiesList::OnObjectPropertyChange(UObject* InObject, FPro
 		}
 	}
 
-	TreeView->RequestListRefresh();
+	GroupsListView->RequestListRefresh();
+	FieldsListView->RequestListRefresh();
 }
 
 void SRCPanelExposedEntitiesList::Refresh()
 {
 	GenerateListWidgets();
+
 	RefreshGroups();
 
-	for (const TSharedPtr<SRCPanelGroup>& Group : FieldGroups)
+	if (Preset.IsValid())
 	{
-		TreeView->SetItemExpansion(Group, true);
-	}
+		constexpr bool bForceMouseClick = true;
 
-	for (const TTuple<FGuid, TSharedPtr<SRCPanelTreeNode>>& FieldTuple : FieldWidgetMap)
-	{
-		TreeView->SetItemExpansion(FieldTuple.Value, false);
+		FRemoteControlPresetGroup& DefaultGroup = Preset->Layout.GetDefaultGroup();
+		GenerateListWidgets(DefaultGroup);
+		SetSelection(FindGroupById(DefaultGroup.Id), bForceMouseClick);
 	}
 }
 
@@ -171,6 +319,29 @@ void SRCPanelExposedEntitiesList::GenerateListWidgets()
 	}
 }
 
+void SRCPanelExposedEntitiesList::GenerateListWidgets(const FRemoteControlPresetGroup& FromGroup)
+{
+	FieldEntities.Reset();
+
+	if (FieldWidgetMap.IsEmpty())
+	{
+		FieldsListView->RequestListRefresh();
+
+		return;
+	}
+
+	if (Preset->Layout.IsDefaultGroup(FromGroup.Id))
+	{
+		FieldWidgetMap.GenerateValueArray(FieldEntities);
+	}
+	else if (TSharedPtr<SRCPanelGroup> SRCGroup = FindGroupById(FromGroup.Id))
+	{
+		SRCGroup->GetNodeChildren(FieldEntities);
+	}
+
+	FieldsListView->RequestListRefresh();
+}
+
 void SRCPanelExposedEntitiesList::RefreshGroups()
 {
 	FieldGroups.Reset(Preset->Layout.GetGroups().Num());
@@ -197,7 +368,7 @@ void SRCPanelExposedEntitiesList::RefreshGroups()
 		}
 	}
 
-	TreeView->RequestListRefresh();
+	GroupsListView->RequestListRefresh();
 }
 
 TSharedRef<ITableRow> SRCPanelExposedEntitiesList::OnGenerateRow(TSharedPtr<SRCPanelTreeNode> Node, const TSharedRef<STableViewBase>& OwnerTable)
@@ -228,7 +399,7 @@ TSharedRef<ITableRow> SRCPanelExposedEntitiesList::OnGenerateRow(TSharedPtr<SRCP
 	if (Node->GetRCType() == SRCPanelTreeNode::Group)
 	{
 		return SNew(STableRow<TSharedPtr<SWidget>>, OwnerTable)
-			.Padding(FMargin(0.5f, 0.5f, 0.5f, 0.5f))
+			.Padding(1.f)
 			.Style(&FRemoteControlPanelStyle::Get()->GetWidgetStyle<FTableRowStyle>("RemoteControlPanel.GroupRow"))
 			[
 				NodeWidget
@@ -264,25 +435,35 @@ void SRCPanelExposedEntitiesList::OnSelectionChanged(TSharedPtr<SRCPanelTreeNode
 		return;
 	}
 
-	// todo call handler on node itself
-	if (TSharedPtr<FRemoteControlField> RCField = Preset->GetExposedEntity<FRemoteControlField>(Node->GetRCId()).Pin())
+	if (Node->GetRCType() == SRCPanelTreeNode::Group)
 	{
-		TSet<UObject*> Objects = TSet<UObject*>{ RCField->GetBoundObjects() };
-
-		TArray<UObject*> OwnerActors;
-		for (UObject* Object : Objects)
+		if (FRemoteControlPresetGroup* RCGroup = Preset->Layout.GetGroup(Node->GetRCId()))
 		{
-			if (Object->IsA<AActor>())
-			{
-				OwnerActors.Add(Object);
-			}
-			else
+			GenerateListWidgets(*RCGroup);
+		}
+	}
+	else
+	{
+		// todo call handler on node itself
+		if (TSharedPtr<FRemoteControlField> RCField = Preset->GetExposedEntity<FRemoteControlField>(Node->GetRCId()).Pin())
+		{
+			TSet<UObject*> Objects = TSet<UObject*>{ RCField->GetBoundObjects() };
+
+			TArray<UObject*> OwnerActors;
+			for (UObject* Object : Objects)
 			{
-				OwnerActors.Add(Object->GetTypedOuter<AActor>());
+				if (Object->IsA<AActor>())
+				{
+					OwnerActors.Add(Object);
+				}
+				else
+				{
+					OwnerActors.Add(Object->GetTypedOuter<AActor>());
+				}
 			}
-		}
 
-		SelectActorsInlevel(OwnerActors);
+			SelectActorsInlevel(OwnerActors);
+		}
 	}
 
 	OnSelectionChangeDelegate.Broadcast(Node);
@@ -294,6 +475,12 @@ FReply SRCPanelExposedEntitiesList::OnDropOnGroup(const TSharedPtr<FDragDropOper
 
 	if (DragDropOperation->IsOfType<FExposedEntityDragDrop>())
 	{
+		if (Preset->Layout.IsDefaultGroup(DragTargetGroup->GetRCId()))
+		{
+			// We do not add fields to the default group.
+			return FReply::Unhandled();
+		}
+
 		if (TSharedPtr<FExposedEntityDragDrop> DragDropOp = StaticCastSharedPtr<FExposedEntityDragDrop>(DragDropOperation))
 		{
 			FGuid DragOriginGroupId = GetGroupId(DragDropOp->GetId());
@@ -323,6 +510,8 @@ FReply SRCPanelExposedEntitiesList::OnDropOnGroup(const TSharedPtr<FDragDropOper
 			FScopedTransaction Transaction(LOCTEXT("MoveField", "Move exposed field"));
 			Preset->Modify();
 			Preset->Layout.SwapFields(Args);
+			constexpr bool bForceMouseClick = true;
+			SetSelection(FindGroupById(Args.TargetGroupId), bForceMouseClick);
 			return FReply::Handled();
 		}
 	}
@@ -360,6 +549,14 @@ FGuid SRCPanelExposedEntitiesList::GetGroupId(const FGuid& EntityId)
 	return GroupId;
 }
 
+FReply SRCPanelExposedEntitiesList::OnCreateGroup()
+{
+	FScopedTransaction Transaction(LOCTEXT("CreateGroup", "Create Group"));
+	Preset->Modify();
+	Preset->Layout.CreateGroup();
+	return FReply::Handled();
+}
+
 void SRCPanelExposedEntitiesList::OnDeleteGroup(const FGuid& GroupId)
 {
 	FScopedTransaction Transaction(LOCTEXT("DeleteGroup", "Delete Group"));
@@ -414,9 +611,9 @@ TSharedPtr<SRCPanelGroup> SRCPanelExposedEntitiesList::FindGroupById(const FGuid
 	return TargetGroup;
 }
 
-TSharedPtr<SWidget> SRCPanelExposedEntitiesList::OnContextMenuOpening()
+TSharedPtr<SWidget> SRCPanelExposedEntitiesList::OnContextMenuOpening(SRCPanelTreeNode::ENodeType InType)
 {
-	if (TSharedPtr<SRCPanelTreeNode> SelectedNode = SRCPanelExposedEntitiesList::GetSelection())
+	if (TSharedPtr<SRCPanelTreeNode> SelectedNode = InType == SRCPanelTreeNode::Group ? SRCPanelExposedEntitiesList::GetSelectedGroup() : SRCPanelExposedEntitiesList::GetSelectedEntity())
 	{
 		return SelectedNode->GetContextMenu();
 	}
@@ -464,7 +661,8 @@ void SRCPanelExposedEntitiesList::OnEntityAdded(const FGuid& InEntityId)
 			if (TSharedPtr<SRCPanelGroup> SRCGroup = FindGroupById(GetGroupId(InEntityId)))
 			{
 				SRCGroup->GetNodes().Add(MoveTemp(Node));
-				TreeView->SetItemExpansion(SRCGroup, true);
+
+				GenerateListWidgets(*Preset->Layout.GetGroup(GetGroupId(InEntityId)));
 			}
 		}
 	};
@@ -477,8 +675,6 @@ void SRCPanelExposedEntitiesList::OnEntityAdded(const FGuid& InEntityId)
 	Args.Entity = Preset->GetExposedEntity(InEntityId).Pin();
 
 	ExposeEntity(FRemoteControlUIModule::Get().GenerateEntityWidget(Args));
-	
-	TreeView->RequestListRefresh();
 }
 
 void SRCPanelExposedEntitiesList::OnEntityRemoved(const FGuid& InGroupId, const FGuid& InEntityId)
@@ -492,7 +688,7 @@ void SRCPanelExposedEntitiesList::OnEntityRemoved(const FGuid& InGroupId, const
 		}
 	}
 
-	if (const TSharedPtr<SRCPanelTreeNode> Node = GetSelection())
+	if (const TSharedPtr<SRCPanelTreeNode> Node = GetSelectedEntity())
 	{
 		if (Node->GetRCId() == InEntityId)
 		{
@@ -501,7 +697,8 @@ void SRCPanelExposedEntitiesList::OnEntityRemoved(const FGuid& InGroupId, const
 	}
 
 	FieldWidgetMap.Remove(InEntityId);
-	TreeView->RequestListRefresh();
+	
+	GenerateListWidgets(*Preset->Layout.GetGroup(InGroupId));
 }
 
 void SRCPanelExposedEntitiesList::OnGroupAdded(const FRemoteControlPresetGroup& Group)
@@ -527,16 +724,16 @@ void SRCPanelExposedEntitiesList::OnGroupAdded(const FRemoteControlPresetGroup&
 	}
 
 	FieldGroup->EnterRenameMode();
-	TreeView->SetSelection(FieldGroup);
-	TreeView->ScrollToBottom();
-	TreeView->RequestListRefresh();
+	GroupsListView->SetSelection(FieldGroup, ESelectInfo::OnMouseClick);
+	GroupsListView->ScrollToBottom();
+	GroupsListView->RequestListRefresh();
 }
 
 void SRCPanelExposedEntitiesList::OnGroupDeleted(FRemoteControlPresetGroup DeletedGroup)
 {
 	int32 Index = FieldGroups.IndexOfByPredicate([&DeletedGroup](const TSharedPtr<SRCPanelGroup>& Group) { return Group->GetRCId() == DeletedGroup.Id; });
 
-	if (TSharedPtr<SRCPanelTreeNode> Node = GetSelection())
+	if (TSharedPtr<SRCPanelTreeNode> Node = GetSelectedEntity())
 	{
 		if (DeletedGroup.GetFields().Contains(Node->GetRCId()))
 		{
@@ -547,7 +744,7 @@ void SRCPanelExposedEntitiesList::OnGroupDeleted(FRemoteControlPresetGroup Delet
 	if (Index != INDEX_NONE)
 	{
 		FieldGroups.RemoveAt(Index);
-		TreeView->RequestListRefresh();
+		GroupsListView->RequestListRefresh();
 	}
 }
 
@@ -567,7 +764,7 @@ void SRCPanelExposedEntitiesList::OnGroupOrderChanged(const TArray<FGuid>& Group
 	};
 
 	FieldGroups.Sort(SortFunc);
-	TreeView->RequestListRefresh();
+	GroupsListView->RequestListRefresh();
 }
 
 void SRCPanelExposedEntitiesList::OnGroupRenamed(const FGuid& GroupId, FName NewName)
@@ -611,7 +808,8 @@ void SRCPanelExposedEntitiesList::OnFieldOrderChanged(const FGuid& GroupId, cons
 		}
 	}
 
-	TreeView->RequestListRefresh();
+	GroupsListView->RequestListRefresh();
+	FieldsListView->RequestListRefresh();
 }
 
 void SRCPanelExposedEntitiesList::OnEntitiesUpdated(URemoteControlPreset*, const TSet<FGuid>& UpdatedEntities)
@@ -630,7 +828,8 @@ void SRCPanelExposedEntitiesList::OnEntitiesUpdated(URemoteControlPreset*, const
 	{
 		if (TSharedPtr<SRCPanelExposedEntitiesList> ListPtr = WeakListPtr.Pin())
 		{
-			ListPtr->TreeView->RequestListRefresh();
+			ListPtr->GroupsListView->RequestListRefresh();
+			ListPtr->FieldsListView->RequestListRefresh();
 		}
 	}));
 
diff --git a/Engine/Plugins/VirtualProduction/RemoteControl/Source/RemoteControlUI/Private/UI/SRCPanelExposedEntitiesList.h b/Engine/Plugins/VirtualProduction/RemoteControl/Source/RemoteControlUI/Private/UI/SRCPanelExposedEntitiesList.h
--- a/Engine/Plugins/VirtualProduction/RemoteControl/Source/RemoteControlUI/Private/UI/SRCPanelExposedEntitiesList.h
+++ b/Engine/Plugins/VirtualProduction/RemoteControl/Source/RemoteControlUI/Private/UI/SRCPanelExposedEntitiesList.h
@@ -59,11 +59,14 @@ public:
 
 	~SRCPanelExposedEntitiesList();
 
-	/** Get the currently selected group or exposed entity. */
-	TSharedPtr<SRCPanelTreeNode> GetSelection() const;
+	/** Get the currently selected group. */
+	TSharedPtr<SRCPanelTreeNode> GetSelectedGroup() const;
+	
+	/** Get the currently selected exposed entity. */
+	TSharedPtr<SRCPanelTreeNode> GetSelectedEntity() const;
 	
 	/** Set the currently selected group or exposed entity. */
-	void SetSelection(const TSharedPtr<SRCPanelTreeNode>& Node);
+	void SetSelection(const TSharedPtr<SRCPanelTreeNode>& Node, const bool bForceMouseClick = false);
 
 	/** Recreate everything in the panel. */
 	void Refresh();
@@ -79,6 +82,8 @@ private:
 	void OnObjectPropertyChange(UObject* InObject, FPropertyChangedEvent& InChangeEvent);
 	/** Create exposed entity widgets. */
 	void GenerateListWidgets();
+	/** Create exposed entity widgets. */
+	void GenerateListWidgets(const FRemoteControlPresetGroup& FromGroup);
 	/** Generate the groups using the preset's layout data. */
 	void RefreshGroups();
 	/** Generate row widgets for groups and exposed entities. */
@@ -91,6 +96,8 @@ private:
 	FReply OnDropOnGroup(const TSharedPtr<FDragDropOperation>& DragDropOperation, const TSharedPtr<SRCPanelTreeNode>& TargetEntity, const TSharedPtr<SRCPanelTreeNode>& DragTargetGroup);
 	/** Get the id of the group that holds a particular widget. */
 	FGuid GetGroupId(const FGuid& EntityId);
+	/** Handles creating a new group. */
+	FReply OnCreateGroup();
 	/** Handles group deletion. */
 	void OnDeleteGroup(const FGuid& GroupId);
 	/** Select actors in the current level. */
@@ -108,7 +115,7 @@ private:
 	TSharedPtr<SRCPanelGroup> FindGroupById(const FGuid& Id);
 
 	/** Handle context menu opening on a row. */
-	TSharedPtr<SWidget> OnContextMenuOpening();
+	TSharedPtr<SWidget> OnContextMenuOpening(SRCPanelTreeNode::ENodeType InType);
 
 	//~ Register and handle preset delegates.
 	void RegisterPresetDelegates();
@@ -125,10 +132,14 @@ private:
 	void OnEntitiesUpdated(URemoteControlPreset*, const TSet<FGuid>& UpdatedEntities);
 
 private:
-	/** Holds the fields list view. */
-	TSharedPtr<STreeView<TSharedPtr<SRCPanelTreeNode>>> TreeView;
+	/** Holds the Groups list view. */
+	TSharedPtr<SListView<TSharedPtr<SRCPanelTreeNode>>> GroupsListView;
+	/** Holds the Fields list view. */
+	TSharedPtr<SListView<TSharedPtr<SRCPanelTreeNode>>> FieldsListView;
 	/** Holds all the field groups. */
 	TArray<TSharedPtr<SRCPanelGroup>> FieldGroups;
+	/** Holds all the field entities. */
+	TArray<TSharedPtr<SRCPanelTreeNode>> FieldEntities;
 	/** Map of field ids to field widgets. */
 	TMap<FGuid, TSharedPtr<SRCPanelTreeNode>> FieldWidgetMap;
 	/** Whether the panel is in edit mode. */
diff --git a/Engine/Plugins/VirtualProduction/RemoteControl/Source/RemoteControlUI/Private/UI/SRCPanelExposedEntity.cpp b/Engine/Plugins/VirtualProduction/RemoteControl/Source/RemoteControlUI/Private/UI/SRCPanelExposedEntity.cpp
--- a/Engine/Plugins/VirtualProduction/RemoteControl/Source/RemoteControlUI/Private/UI/SRCPanelExposedEntity.cpp
+++ b/Engine/Plugins/VirtualProduction/RemoteControl/Source/RemoteControlUI/Private/UI/SRCPanelExposedEntity.cpp
@@ -3,6 +3,7 @@
 #include "SRCPanelExposedEntity.h"
 
 #include "ActorTreeItem.h"
+#include "Commands/RemoteControlCommands.h"
 #include "Editor.h"
 #include "EditorFontGlyphs.h"
 #include "Engine/Selection.h"
@@ -10,6 +11,7 @@
 #include "Framework/Application/SlateApplication.h"
 #include "Framework/MultiBox/MultiBoxBuilder.h"
 #include "GameFramework/Actor.h"
+#include "Interfaces/IMainFrameModule.h"
 #include "RemoteControlBinding.h"
 #include "RemoteControlEntity.h"
 #include "RemoteControlPreset.h"
@@ -55,7 +57,18 @@ TSharedPtr<FRemoteControlEntity> SRCPanelExposedEntity::GetEntity() const
 
 TSharedPtr<SWidget> SRCPanelExposedEntity::GetContextMenu()
 {
-	FMenuBuilder MenuBuilder(true, TSharedPtr<const FUICommandList>());
+	IMainFrameModule& MainFrame = FModuleManager::Get().LoadModuleChecked<IMainFrameModule>("MainFrame");
+
+	FMenuBuilder MenuBuilder(true, MainFrame.GetMainFrameCommandBindings());
+
+	MenuBuilder.BeginSection("Common");
+
+	MenuBuilder.AddMenuEntry(FRemoteControlCommands::Get().RenameEntity, NAME_None, TAttribute<FText>(), TAttribute<FText>(), FSlateIcon(FAppStyle::GetAppStyleSetName(), TEXT("GenericCommands.Rename")));
+	MenuBuilder.AddMenuEntry(FRemoteControlCommands::Get().DeleteEntity, NAME_None, TAttribute<FText>(), TAttribute<FText>(), FSlateIcon(FAppStyle::GetAppStyleSetName(), TEXT("GenericCommands.Delete")));
+
+	MenuBuilder.EndSection();
+
+	MenuBuilder.AddSeparator();
 
 	constexpr bool bNoIndent = true;
 	MenuBuilder.AddWidget(CreateUseContextCheckbox(), LOCTEXT("UseContextLabel", "Use Context"), bNoIndent);
@@ -89,6 +102,11 @@ TSharedPtr<SWidget> SRCPanelExposedEntity::GetContextMenu()
 	return MenuBuilder.MakeWidget();
 }
 
+void SRCPanelExposedEntity::EnterRenameMode()
+{
+	bNeedsRename = true;
+}
+
 void SRCPanelExposedEntity::Initialize(const FGuid& InEntityId, URemoteControlPreset* InPreset, const TAttribute<bool>& InbEditMode)
 {
 	EntityId = InEntityId;
@@ -309,6 +327,7 @@ TSharedRef<SWidget> SRCPanelExposedEntity::CreateEntityWidget(TSharedPtr<SWidget
             .Text(FEditorFontGlyphs::Exclamation_Triangle)
 		]
 		+ SHorizontalBox::Slot()
+		.Padding(4.f, 0.f, 2.0f, 0.f)
 		.AutoWidth()
 		[
 			SAssignNew(NameTextBox, SInlineEditableTextBlock)
@@ -318,32 +337,11 @@ TSharedRef<SWidget> SRCPanelExposedEntity::CreateEntityWidget(TSharedPtr<SWidget
 			.IsReadOnly_Lambda([this]() { return !bEditMode.Get(); })
 		];
 
-	Args.RenameButton = SNew(SButton)
-		.Visibility_Raw(this, &SRCPanelExposedEntity::GetVisibilityAccordingToEditMode, EVisibility::Collapsed)
-		.ButtonStyle(FAppStyle::Get(), "FlatButton")
-		.OnClicked_Lambda([this]() {
-			bNeedsRename = true;
-			return FReply::Handled();
-		})
-		[
-			SNew(STextBlock)
-				.TextStyle(FRemoteControlPanelStyle::Get(), "RemoteControlPanel.Button.TextStyle")
-				.Font(FAppStyle::Get().GetFontStyle("FontAwesome.10"))
-				.Text(FText::FromString(FString(TEXT("\xf044"))) /*fa-edit*/)
-		];
+	Args.RenameButton = SNullWidget::NullWidget;
 
 	Args.ValueWidget = ValueWidget;
 
-	Args.UnexposeButton = SNew(SButton)
-		.Visibility_Raw(this, &SRCPanelExposedEntity::GetVisibilityAccordingToEditMode, EVisibility::Collapsed)
-		.OnPressed_Raw(this, &SRCPanelExposedEntity::HandleUnexposeEntity)
-		.ButtonStyle(FRemoteControlPanelStyle::Get(), "RemoteControlPanel.UnexposeButton")
-		[
-			SNew(STextBlock)
-			.TextStyle(FRemoteControlPanelStyle::Get(), "RemoteControlPanel.Button.TextStyle")
-			.Font(FAppStyle::Get().GetFontStyle("FontAwesome.10"))
-			.Text(FText::FromString(FString(TEXT("\xf00d"))) /*fa-times*/)
-		];
+	Args.UnexposeButton = SNullWidget::NullWidget;
 
 	Widget->SetContent(MakeNodeWidget(Args));
 	return Widget;
diff --git a/Engine/Plugins/VirtualProduction/RemoteControl/Source/RemoteControlUI/Private/UI/SRCPanelExposedEntity.h b/Engine/Plugins/VirtualProduction/RemoteControl/Source/RemoteControlUI/Private/UI/SRCPanelExposedEntity.h
--- a/Engine/Plugins/VirtualProduction/RemoteControl/Source/RemoteControlUI/Private/UI/SRCPanelExposedEntity.h
+++ b/Engine/Plugins/VirtualProduction/RemoteControl/Source/RemoteControlUI/Private/UI/SRCPanelExposedEntity.h
@@ -26,6 +26,8 @@ struct SRCPanelExposedEntity : public SRCPanelTreeNode
 	TSharedPtr<FRemoteControlEntity> GetEntity() const;
 	virtual TSharedPtr<SWidget> GetContextMenu() override;
 	virtual FGuid GetRCId() const override final { return EntityId; }
+	/** Make the group name's text box editable. */
+	virtual void EnterRenameMode() override;
 
 protected:
 	void Initialize(const FGuid& InEntityId, URemoteControlPreset* InPreset, const TAttribute<bool>& InbEditMode);
diff --git a/Engine/Plugins/VirtualProduction/RemoteControl/Source/RemoteControlUI/Private/UI/SRCPanelFieldGroup.cpp b/Engine/Plugins/VirtualProduction/RemoteControl/Source/RemoteControlUI/Private/UI/SRCPanelFieldGroup.cpp
--- a/Engine/Plugins/VirtualProduction/RemoteControl/Source/RemoteControlUI/Private/UI/SRCPanelFieldGroup.cpp
+++ b/Engine/Plugins/VirtualProduction/RemoteControl/Source/RemoteControlUI/Private/UI/SRCPanelFieldGroup.cpp
@@ -2,6 +2,8 @@
 
 #include "SRCPanelFieldGroup.h"
 
+#include "Commands/RemoteControlCommands.h"
+#include "Interfaces/IMainFrameModule.h"
 #include "RemoteControlPanelStyle.h"
 #include "RemoteControlPreset.h"
 #include "ScopedTransaction.h"
@@ -47,15 +49,25 @@ void SRCPanelGroup::Construct(const FArguments& InArgs, URemoteControlPreset* In
 
 	TSharedRef<SWidget> LeftColumn = 
 		SNew(SHorizontalBox)
+		// Color band
+		+ SHorizontalBox::Slot()
+		.VAlign(VAlign_Fill)
+		.HAlign(HAlign_Center)
+		.Padding(0.f, 0.f, 2.f, 0.f)
+		.AutoWidth()
+		[
+			SNew(SBorder)
+			.BorderImage(this, &SRCPanelGroup::HandleGroupColor)
+		]
 		// Drag and drop handle
 		+ SHorizontalBox::Slot()
 		.VAlign(VAlign_Fill)
 		.HAlign(HAlign_Center)
-		.Padding(FMargin(4.0f, 0.0f))
+		.Padding(0.f, 4.f)
 		.AutoWidth()
 		[
 			SNew(SBox)
-			.Padding(FMargin(0.0f, 2.0f) )
+			.Padding(2.f, 0.f)
 			.Visibility(this, &SRCPanelGroup::GetVisibilityAccordingToEditMode, EVisibility::Collapsed)
 			[
 				SNew(SRCPanelDragHandle<FFieldGroupDragDropOp>, Id)
@@ -64,61 +76,20 @@ void SRCPanelGroup::Construct(const FArguments& InArgs, URemoteControlPreset* In
 		]
 		+ SHorizontalBox::Slot()
 		// Group name
-		.FillWidth(1.0f)
-		.VAlign(VAlign_Fill)
-		.Padding(FMargin(0.f, 0.f, 0.f, 2.f))
+		.FillWidth(1.f)
+		.VAlign(VAlign_Center)
+		.Padding(6.f, 4.f)
 		.AutoWidth()
 		[
 			SAssignNew(NameTextBox, SInlineEditableTextBlock)
 			.ColorAndOpacity(this, &SRCPanelGroup::GetGroupNameTextColor)
-			.Font(FAppStyle::Get().GetFontStyle("DetailsView.CategoryFontStyle"))
+			.Font(FAppStyle::Get().GetFontStyle("PropertyWindow.BoldFont"))
 			.Text(FText::FromName(Name))
 			.OnTextCommitted(this, &SRCPanelGroup::OnLabelCommitted)
 			.OnVerifyTextChanged(this, &SRCPanelGroup::OnVerifyItemLabelChanged)
 			.IsReadOnly_Lambda([this]() { return !bEditMode.Get(); })
-		]
-		+ SHorizontalBox::Slot()
-		// Rename button
-		.AutoWidth()
-		.VAlign(VAlign_Center)
-		.HAlign(HAlign_Left)
-		[
-			SNew(SButton)
-			.Visibility(this, &SRCPanelGroup::GetVisibilityAccordingToEditMode, EVisibility::Collapsed)
-			.ButtonStyle(FAppStyle::Get(), "FlatButton")
-			.OnClicked_Lambda([this] () 
-				{
-					bNeedsRename = true;
-					return FReply::Handled();	
-				})
-			[
-				SNew(STextBlock)
-				.TextStyle(FRemoteControlPanelStyle::Get(), "RemoteControlPanel.Button.TextStyle")
-				.Font(FAppStyle::Get().GetFontStyle("FontAwesome.10"))
-				.Text(FText::FromString(FString(TEXT("\xf044"))) /*fa-edit*/)
-			]
 		];
 
-	TSharedRef<SWidget> RightColumn =
-		SNew(SHorizontalBox)
-		+SHorizontalBox::Slot()
-		.VAlign(VAlign_Top)
-		.HAlign(HAlign_Right)
-		.Padding(0, 2.0f)
-		.FillWidth(1.f)
-		[
-			SNew(SButton)
-			.Visibility_Raw(this, &SRCPanelGroup::GetVisibilityAccordingToEditMode, EVisibility::Hidden)
-			.OnClicked(this, &SRCPanelGroup::HandleDeleteGroup)
-			.ButtonStyle(FRemoteControlPanelStyle::Get(), "RemoteControlPanel.UnexposeButton")
-			[
-				SNew(STextBlock)
-				.TextStyle(FRemoteControlPanelStyle::Get(), "RemoteControlPanel.Button.TextStyle")
-				.Font(FAppStyle::Get().GetFontStyle("FontAwesome.10"))
-				.Text(FText::FromString(FString(TEXT("\xf00d"))) /*fa-times*/)
-			]
-		];
-	
 	ChildSlot
 	[
 		SNew(SBorder)
@@ -126,14 +97,18 @@ void SRCPanelGroup::Construct(const FArguments& InArgs, URemoteControlPreset* In
 		.BorderImage(this, &SRCPanelGroup::GetBorderImage)
 		.VAlign(VAlign_Fill)
 		[
-			SNew(SDropTarget)
-			.VerticalImage(FRemoteControlPanelStyle::Get()->GetBrush("RemoteControlPanel.VerticalDash"))
-			.HorizontalImage(FRemoteControlPanelStyle::Get()->GetBrush("RemoteControlPanel.HorizontalDash"))
-			.OnDropped_Lambda([this](const FGeometry& InGeometry, const FDragDropEvent& InDragDropEvent) { return OnFieldDropGroup(InDragDropEvent.GetOperation(), nullptr); })
-			.OnAllowDrop(this, &SRCPanelGroup::OnAllowDropFromOtherGroup)
-			.OnIsRecognized(this, &SRCPanelGroup::OnAllowDropFromOtherGroup)
+			SNew(SBox)
+			.HeightOverride(32.f)
 			[
-				MakeSplitRow(LeftColumn, RightColumn)
+				SNew(SDropTarget)
+				.VerticalImage(FRemoteControlPanelStyle::Get()->GetBrush("RemoteControlPanel.VerticalDash"))
+				.HorizontalImage(FRemoteControlPanelStyle::Get()->GetBrush("RemoteControlPanel.HorizontalDash"))
+				.OnDropped_Lambda([this](const FGeometry& InGeometry, const FDragDropEvent& InDragDropEvent) { return OnFieldDropGroup(InDragDropEvent.GetOperation(), nullptr); })
+				.OnAllowDrop(this, &SRCPanelGroup::OnAllowDropFromOtherGroup)
+				.OnIsRecognized(this, &SRCPanelGroup::OnAllowDropFromOtherGroup)
+				[
+					LeftColumn
+				]
 			]
 		]
 	];
@@ -160,6 +135,22 @@ void SRCPanelGroup::GetNodeChildren(TArray<TSharedPtr<SRCPanelTreeNode>>& OutChi
 	OutChildren.Append(Nodes);
 }
 
+TSharedPtr<SWidget> SRCPanelGroup::GetContextMenu()
+{
+	IMainFrameModule& MainFrame = FModuleManager::Get().LoadModuleChecked<IMainFrameModule>("MainFrame");
+
+	FMenuBuilder MenuBuilder(true, MainFrame.GetMainFrameCommandBindings());
+
+	MenuBuilder.BeginSection("Common");
+
+	MenuBuilder.AddMenuEntry(FRemoteControlCommands::Get().RenameEntity, NAME_None, TAttribute<FText>(), TAttribute<FText>(), FSlateIcon(FAppStyle::GetAppStyleSetName(), TEXT("GenericCommands.Rename")));
+	MenuBuilder.AddMenuEntry(FRemoteControlCommands::Get().DeleteEntity, NAME_None, TAttribute<FText>(), TAttribute<FText>(), FSlateIcon(FAppStyle::GetAppStyleSetName(), TEXT("GenericCommands.Delete")));
+
+	MenuBuilder.EndSection();
+
+	return MenuBuilder.MakeWidget();
+}
+
 FGuid SRCPanelGroup::GetRCId() const
 {
 	return Id;
@@ -170,6 +161,18 @@ SRCPanelTreeNode::ENodeType SRCPanelGroup::GetRCType() const
 	return SRCPanelTreeNode::Group;
 }
 
+FReply SRCPanelGroup::OnMouseButtonDoubleClick(const FGeometry& InMyGeometry, const FPointerEvent& InMouseEvent)
+{
+	check(Preset.IsValid());
+
+	if (Preset->Layout.IsDefaultGroup(Id))
+	{
+		return SCompoundWidget::OnMouseButtonDoubleClick(InMyGeometry, InMouseEvent);
+	}
+
+	return SRCPanelTreeNode::OnMouseButtonDoubleClick(InMyGeometry, InMouseEvent);
+}
+
 FReply SRCPanelGroup::OnFieldDropGroup(const FDragDropEvent& Event, TSharedPtr<SRCPanelTreeNode> TargetField)
 {
 	if (const TSharedPtr<FExposedEntityDragDrop> DragDropOp = Event.GetOperationAs<FExposedEntityDragDrop>())
@@ -265,4 +268,11 @@ void SRCPanelGroup::OnLabelCommitted(const FText& InLabel, ETextCommit::Type InC
 	Preset->Layout.RenameGroup(Id, FName(*InLabel.ToString()));
 }
 
+const FSlateBrush* SRCPanelGroup::HandleGroupColor() const
+{
+	check(Preset.IsValid());
+
+	return new FSlateColorBrush(Preset->Layout.GetTagColor(Id));
+}
+
 #undef LOCTEXT_NAMESPACE
diff --git a/Engine/Plugins/VirtualProduction/RemoteControl/Source/RemoteControlUI/Private/UI/SRCPanelFieldGroup.h b/Engine/Plugins/VirtualProduction/RemoteControl/Source/RemoteControlUI/Private/UI/SRCPanelFieldGroup.h
--- a/Engine/Plugins/VirtualProduction/RemoteControl/Source/RemoteControlUI/Private/UI/SRCPanelFieldGroup.h
+++ b/Engine/Plugins/VirtualProduction/RemoteControl/Source/RemoteControlUI/Private/UI/SRCPanelFieldGroup.h
@@ -57,12 +57,14 @@ public:
 	TArray<TSharedPtr<SRCPanelTreeNode>>& GetNodes() { return Nodes; }
 
 	/** Make the group name's text box editable. */
-	void EnterRenameMode();
+	virtual void EnterRenameMode() override;
 
 	//~ SRCPanelTreeNode Interface
 	virtual void GetNodeChildren(TArray<TSharedPtr<SRCPanelTreeNode>>& OutChildren) const override;
+	virtual TSharedPtr<SWidget> GetContextMenu() override;
 	virtual FGuid GetRCId() const override;
 	virtual ENodeType GetRCType() const override;
+	virtual FReply OnMouseButtonDoubleClick(const FGeometry& InMyGeometry, const FPointerEvent& InMouseEvent) override;
 
 private:
 	//~ Handle drag/drop events
@@ -83,6 +85,8 @@ private:
 
 	void OnLabelCommitted(const FText& InLabel, ETextCommit::Type InCommitInfo);
 
+	const FSlateBrush* HandleGroupColor() const;
+
 private:
 	/** Event called when something is dropped on this group. */
 	FOnFieldDropEvent OnFieldDropEvent;
diff --git a/Engine/Plugins/VirtualProduction/RemoteControl/Source/RemoteControlUI/Private/UI/SRCPanelTreeNode.h b/Engine/Plugins/VirtualProduction/RemoteControl/Source/RemoteControlUI/Private/UI/SRCPanelTreeNode.h
--- a/Engine/Plugins/VirtualProduction/RemoteControl/Source/RemoteControlUI/Private/UI/SRCPanelTreeNode.h
+++ b/Engine/Plugins/VirtualProduction/RemoteControl/Source/RemoteControlUI/Private/UI/SRCPanelTreeNode.h
@@ -41,6 +41,10 @@ struct SRCPanelTreeNode : public SCompoundWidget
 	virtual TSharedPtr<SWidget> GetContextMenu() { return nullptr; }
 	/** Set whether this widget is currently hovered when drag and dropping. */
 	virtual void SetIsHovered(bool bIsBeingHovered) {}
+	/** Make the node name's text box editable. */
+	virtual void EnterRenameMode() {};
+
+	virtual FReply OnMouseButtonDoubleClick(const FGeometry& InMyGeometry, const FPointerEvent& InMouseEvent) override { if (InMouseEvent.GetEffectingButton() == EKeys::LeftMouseButton) { EnterRenameMode(); } return SCompoundWidget::OnMouseButtonDoubleClick(InMyGeometry, InMouseEvent); };
 
 protected:
 	struct FMakeNodeWidgetArgs
diff --git a/Engine/Plugins/VirtualProduction/RemoteControl/Source/RemoteControlUI/Private/UI/SRemoteControlPanel.cpp b/Engine/Plugins/VirtualProduction/RemoteControl/Source/RemoteControlUI/Private/UI/SRemoteControlPanel.cpp
--- a/Engine/Plugins/VirtualProduction/RemoteControl/Source/RemoteControlUI/Private/UI/SRemoteControlPanel.cpp
+++ b/Engine/Plugins/VirtualProduction/RemoteControl/Source/RemoteControlUI/Private/UI/SRemoteControlPanel.cpp
@@ -64,7 +64,6 @@
 #include "Widgets/Layout/SWidgetSwitcher.h"
 #include "Widgets/SBoxPanel.h"
 #include "Widgets/Text/STextBlock.h"
-#include "Styling/AppStyle.h"
 
 #define LOCTEXT_NAMESPACE "RemoteControlPanel"
 
@@ -247,6 +246,29 @@ void SRemoteControlPanel::Construct(const FArguments& InArgs, URemoteControlPres
 
 	BindRemoteControlCommands();
 
+	// Rebind All
+	AddToolbarWidget(SNew(SButton)
+		.Visibility_Lambda([this]() { return bShowRebindButton ? EVisibility::Visible : EVisibility::Collapsed; })
+		.OnClicked_Raw(this, &SRemoteControlPanel::OnClickRebindAllButton)
+		[
+			SNew(STextBlock)
+			.ToolTipText(LOCTEXT("RebindButtonToolTip", "Attempt to rebind all unbound entites of the preset."))
+			.Text(LOCTEXT("RebindButtonText", "Rebind All"))
+		]);
+
+	// Show Log
+	AddToolbarWidget(SNew(SCheckBox)
+		.Style(FRemoteControlPanelStyle::Get(), "RemoteControlPanel.Switch")
+		.ToolTipText(LOCTEXT("ShowLogTooltip", "Show/Hide remote control log."))
+		.IsChecked_Lambda([]() { return FRemoteControlLogger::Get().IsEnabled() ? ECheckBoxState::Checked : ECheckBoxState::Unchecked; })
+		.OnCheckStateChanged(this, &SRemoteControlPanel::OnLogCheckboxToggle)
+		.Padding(4.f)
+		[
+			SNew(STextBlock)
+			.Text(LOCTEXT("ShowLogLabel", "Log"))
+			.Justification(ETextJustify::Center)
+		]);
+
 	// Edit Mode
 	AddToolbarWidget(SNew(SCheckBox)
 		.Style(FRemoteControlPanelStyle::Get(), "RemoteControlPanel.Switch")
@@ -267,38 +289,6 @@ void SRemoteControlPanel::Construct(const FArguments& InArgs, URemoteControlPres
 		AddToolbarWidget(Widget);
 	}
 
-	// Settings
-	AddToolbarWidget(SNew(SButton)
-		.ContentPadding(2.0f)
-		.ButtonStyle(FAppStyle::Get(), "FlatButton")
-		.TextStyle(FRemoteControlPanelStyle::Get(), "RemoteControlPanel.Button.TextStyle")
-		.OnClicked_Raw(this, &SRemoteControlPanel::OnClickSettingsButton)
-		.ToolTipText(LOCTEXT("OpenRemoteControlSettings", "Open Remote Control settings."))
-		[
-			SNew(SHorizontalBox)
-
-			+SHorizontalBox::Slot()
-			.HAlign(HAlign_Center)
-			.VAlign(VAlign_Center)
-			.AutoWidth()
-			.Padding(0.f, 2.f, 4.f, 2.f)
-			[
-				SNew(SImage)
-				.Image(FAppStyle::Get().GetBrush("Icons.Toolbar.Settings"))
-				.ColorAndOpacity(FSlateColor::UseForeground())
-			]
-
-			+ SHorizontalBox::Slot()
-			.HAlign(HAlign_Fill)
-			.VAlign(VAlign_Center)
-			.FillWidth(1.f)
-			.Padding(4.f, 2.f, 0.f, 2.f)
-			[
-				SNew(STextBlock)
-				.Text(LOCTEXT("SettingsLabel", "Settings"))
-			]
-		]);
-
 	GenerateToolbar();
 	GenerateAuxiliaryToolbar();
 
@@ -311,7 +301,7 @@ void SRemoteControlPanel::Construct(const FArguments& InArgs, URemoteControlPres
 	EntityList = SNew(SRCPanelExposedEntitiesList, Preset.Get(), WidgetRegistry)
 		.OnEntityListUpdated_Lambda([this] ()
 		{
-			UpdateEntityDetailsView(EntityList->GetSelection());
+			UpdateEntityDetailsView(EntityList->GetSelectedEntity());
 			UpdateRebindButtonVisibility();
 			CachedExposedPropertyArgs.Reset();
 		})
@@ -388,6 +378,25 @@ void SRemoteControlPanel::Construct(const FArguments& InArgs, URemoteControlPres
 					.ColorAndOpacity(FSlateColor::UseForeground())
 				]
 			]
+			
+			// Settings
+			+ SHorizontalBox::Slot()
+			.VAlign(VAlign_Center)
+			.AutoWidth()
+			.Padding(3.f, 3.f, 3.f, 3.f)
+			[
+				SNew(SButton)
+				.ContentPadding(2.0f)
+				.ButtonStyle(FAppStyle::Get(), "FlatButton")
+				.TextStyle(FRemoteControlPanelStyle::Get(), "RemoteControlPanel.Button.TextStyle")
+				.OnClicked_Raw(this, &SRemoteControlPanel::OnClickSettingsButton)
+				.ToolTipText(LOCTEXT("OpenRemoteControlSettings", "Open Remote Control settings."))
+				[
+					SNew(SImage)
+					.Image(FAppStyle::Get().GetBrush("Icons.Toolbar.Settings"))
+					.ColorAndOpacity(FSlateColor::UseForeground())
+				]
+			]
 
 			// Separator
 			+SHorizontalBox::Slot()
@@ -832,12 +841,9 @@ void SRemoteControlPanel::ToggleProperty(const FRCExposesPropertyArgs& InPropert
 
 FGuid SRemoteControlPanel::GetSelectedGroup() const
 {
-	if (TSharedPtr<SRCPanelTreeNode> Node = EntityList->GetSelection())
+	if (TSharedPtr<SRCPanelTreeNode> Node = EntityList->GetSelectedGroup())
 	{
-		if (Node->GetRCType() == SRCPanelTreeNode::Group)
-		{
-			return Node->GetRCId();		
-		}
+		return Node->GetRCId();
 	}
 	return FGuid();
 }
@@ -1141,6 +1147,16 @@ void SRemoteControlPanel::BindRemoteControlCommands()
 		Commands.ToggleLogicEditor,
 		FExecuteAction::CreateLambda([this] { OnToggleLogicEditor(); }),
 		FCanExecuteAction::CreateLambda([] { return true; }));
+
+	ActionList.MapAction(
+		Commands.DeleteEntity,
+		FExecuteAction::CreateSP(this, &SRemoteControlPanel::DeleteEntity_Execute),
+		FCanExecuteAction::CreateSP(this, &SRemoteControlPanel::CanDeleteEntity));
+
+	ActionList.MapAction(
+		Commands.RenameEntity,
+		FExecuteAction::CreateSP(this, &SRemoteControlPanel::RenameEntity_Execute),
+		FCanExecuteAction::CreateSP(this, &SRemoteControlPanel::CanRenameEntity));
 }
 
 void SRemoteControlPanel::UnbindRemoteControlCommands()
@@ -1158,6 +1174,8 @@ void SRemoteControlPanel::UnbindRemoteControlCommands()
 		ActionList.UnmapAction(Commands.ToggleExposeFunctions);
 		ActionList.UnmapAction(Commands.ToggleProtocolMappings);
 		ActionList.UnmapAction(Commands.ToggleLogicEditor);
+		ActionList.UnmapAction(Commands.DeleteEntity);
+		ActionList.UnmapAction(Commands.RenameEntity);
 	}
 }
 
@@ -1396,7 +1414,7 @@ TSharedRef<SWidget> SRemoteControlPanel::CreateEntityDetailsView()
 
 	EntityDetailsView = FModuleManager::GetModuleChecked<FPropertyEditorModule>("PropertyEditor").CreateStructureDetailView(MoveTemp(Args), FStructureDetailsViewArgs(), nullptr);
 
-	UpdateEntityDetailsView(EntityList->GetSelection());
+	UpdateEntityDetailsView(EntityList->GetSelectedEntity());
 	if (ensure(EntityDetailsView && EntityDetailsView->GetWidget()))
 	{
 		return EntityDetailsView->GetWidget().ToSharedRef();
@@ -1409,6 +1427,8 @@ void SRemoteControlPanel::UpdateEntityDetailsView(const TSharedPtr<SRCPanelTreeN
 {
 	TSharedPtr<FStructOnScope> SelectedEntityPtr;
 
+	SelectedEntity = SelectedNode;
+
 	if (SelectedNode)
 	{
 		if (SelectedNode->GetRCType() != SRCPanelTreeNode::Group &&
@@ -1575,7 +1595,7 @@ void SRemoteControlPanel::GenerateToolbar()
 				[
 					SNew(SSeparator)
 					.SeparatorImage(FAppStyle::Get().GetBrush("Separator"))
-					.Thickness(2.f)
+					.Thickness(1.5f)
 					.Orientation(EOrientation::Orient_Vertical)
 				];
 
@@ -1608,22 +1628,6 @@ void SRemoteControlPanel::GenerateToolbar()
 			.Text(this, &SRemoteControlPanel::HandlePresetName)
 		]
 		+ SHorizontalBox::Slot()
-		.Padding(5.f, 0.f)
-		.HAlign(HAlign_Fill)
-		.VAlign(VAlign_Center)
-		.AutoWidth()
-		[
-			SNew(SCheckBox)
-			.IsChecked_Lambda([]() { return FRemoteControlLogger::Get().IsEnabled() ? ECheckBoxState::Checked : ECheckBoxState::Unchecked; })
-			.OnCheckStateChanged(this, &SRemoteControlPanel::OnLogCheckboxToggle)
-			.Padding(FMargin(5.f, 0.f))
-			[
-				SNew(STextBlock)
-				.Text(LOCTEXT("ShowLogLabel", "Show Log"))
-				.Justification(ETextJustify::Center)
-			]
-		]
-		+ SHorizontalBox::Slot()
 		.HAlign(HAlign_Right)
 		.AutoWidth()
 		[
@@ -1799,4 +1803,59 @@ void SRemoteControlPanel::OnToggleLogicEditor() const
 {
 }
 
+void SRemoteControlPanel::DeleteEntity_Execute() const
+{
+	if (SelectedEntity->GetRCType() == SRCPanelTreeNode::FieldChild) // Field Child does not contain entity ID, that is why it should not be processed
+	{
+		return;
+	}
+
+	if (SelectedEntity->GetRCType() == SRCPanelTreeNode::Group)
+	{
+		FScopedTransaction Transaction(LOCTEXT("DeleteGroup", "Delete Group"));
+		Preset->Modify();
+		Preset->Layout.DeleteGroup(SelectedEntity->GetRCId());
+	}
+	else
+	{
+		FScopedTransaction Transaction(LOCTEXT("UnexposeFunction", "Unexpose remote control entity"));
+		Preset->Modify();
+		Preset->Unexpose(SelectedEntity->GetRCId());
+	}
+
+	EntityList->Refresh();
+}
+
+bool SRemoteControlPanel::CanDeleteEntity() const
+{
+	if (SelectedEntity.IsValid() && Preset.IsValid() && bIsInEditMode)
+	{
+		// Do not allow default group to be deleted.
+		return !Preset->Layout.IsDefaultGroup(SelectedEntity->GetRCId());
+	}
+
+	return false;
+}
+
+void SRemoteControlPanel::RenameEntity_Execute() const
+{
+	if (SelectedEntity->GetRCType() == SRCPanelTreeNode::FieldChild) // Field Child does not contain entity ID, that is why it should not be processed
+	{
+		return;
+	}
+
+	SelectedEntity->EnterRenameMode();
+}
+
+bool SRemoteControlPanel::CanRenameEntity() const
+{
+	if (SelectedEntity.IsValid() && Preset.IsValid() && bIsInEditMode)
+	{
+		// Do not allow default group to be renamed.
+		return !Preset->Layout.IsDefaultGroup(SelectedEntity->GetRCId());
+	}
+
+	return false;
+}
+
 #undef LOCTEXT_NAMESPACE /*RemoteControlPanel*/
diff --git a/Engine/Plugins/VirtualProduction/RemoteControl/Source/RemoteControlUI/Private/UI/SRemoteControlPanel.h b/Engine/Plugins/VirtualProduction/RemoteControl/Source/RemoteControlUI/Private/UI/SRemoteControlPanel.h
--- a/Engine/Plugins/VirtualProduction/RemoteControl/Source/RemoteControlUI/Private/UI/SRemoteControlPanel.h
+++ b/Engine/Plugins/VirtualProduction/RemoteControl/Source/RemoteControlUI/Private/UI/SRemoteControlPanel.h
@@ -248,6 +248,18 @@ private:
 	void OnToggleProtocolMappings() const;
 
 	void OnToggleLogicEditor() const;
+	
+	/** Called when user attempts to delete a group/exposed entity. */
+	void DeleteEntity_Execute() const;
+
+	/** Called to test if user is able to delete a group/exposed entity. */
+	bool CanDeleteEntity() const;
+	
+	/** Called when user attempts to rename a group/exposed entity. */
+	void RenameEntity_Execute() const;
+
+	/** Called to test if user is able to rename a group/exposed entity. */
+	bool CanRenameEntity() const;
 
 private:
 	static const FName DefaultRemoteControlPanelToolBarName;
@@ -302,6 +314,8 @@ private:
 	TArray<TSharedRef<SWidget>> ToolbarWidgets;
 	/** The text box used to search for tags. */
 	TSharedPtr<SSearchBox> SearchBoxPtr;
+	/** Holds a shared pointer reference to the active entity that is selected. */
+	TSharedPtr<SRCPanelTreeNode> SelectedEntity;
 
 public:
 	static const float MinimumPanelWidth;
