Fix for un-initialised variable CIS compiler issue

diff --git a/Engine/Plugins/VirtualProduction/RemoteControl/Source/RemoteControl/Public/RemoteControlPreset.h b/Engine/Plugins/VirtualProduction/RemoteControl/Source/RemoteControl/Public/RemoteControlPreset.h
--- a/Engine/Plugins/VirtualProduction/RemoteControl/Source/RemoteControl/Public/RemoteControlPreset.h
+++ b/Engine/Plugins/VirtualProduction/RemoteControl/Source/RemoteControl/Public/RemoteControlPreset.h
@@ -96,7 +96,9 @@ struct REMOTECONTROL_API FRemoteControlPresetGroup
 {
 	GENERATED_BODY()
 
-	FRemoteControlPresetGroup() = default;
+	FRemoteControlPresetGroup()
+		:TagColor(MakeSimilarSaturatedColor())
+	{}
 
 	FRemoteControlPresetGroup(FName InName, FGuid InId)
 		: Name(InName)
