Fix CIS

diff --git a/Engine/Plugins/VirtualProduction/RemoteControl/Source/RemoteControlUI/Private/UI/SRCPanelFieldGroup.cpp b/Engine/Plugins/VirtualProduction/RemoteControl/Source/RemoteControlUI/Private/UI/SRCPanelFieldGroup.cpp
--- a/Engine/Plugins/VirtualProduction/RemoteControl/Source/RemoteControlUI/Private/UI/SRCPanelFieldGroup.cpp
+++ b/Engine/Plugins/VirtualProduction/RemoteControl/Source/RemoteControlUI/Private/UI/SRCPanelFieldGroup.cpp
@@ -3,6 +3,7 @@
 #include "SRCPanelFieldGroup.h"
 
 #include "Commands/RemoteControlCommands.h"
+#include "Framework/MultiBox/MultiBoxBuilder.h"
 #include "Interfaces/IMainFrameModule.h"
 #include "RemoteControlPanelStyle.h"
 #include "RemoteControlPreset.h"
