Fix non-deterministic initialization of FRemoteControlPresetGroup, causing warnings when running initialization tests.

diff --git a/Engine/Plugins/VirtualProduction/RemoteControl/Source/RemoteControl/Public/RemoteControlPreset.h b/Engine/Plugins/VirtualProduction/RemoteControl/Source/RemoteControl/Public/RemoteControlPreset.h
--- a/Engine/Plugins/VirtualProduction/RemoteControl/Source/RemoteControl/Public/RemoteControlPreset.h
+++ b/Engine/Plugins/VirtualProduction/RemoteControl/Source/RemoteControl/Public/RemoteControlPreset.h
@@ -97,7 +97,7 @@ struct REMOTECONTROL_API FRemoteControlPresetGroup
 	GENERATED_BODY()
 
 	FRemoteControlPresetGroup()
-		:TagColor(MakeSimilarSaturatedColor())
+		:TagColor(ForceInit)
 	{}
 
 	FRemoteControlPresetGroup(FName InName, FGuid InId)
