Fix BP event not callable with raw remote control API

diff --git a/Engine/Plugins/VirtualProduction/RemoteControl/Source/RemoteControl/Private/RemoteControlModule.cpp b/Engine/Plugins/VirtualProduction/RemoteControl/Source/RemoteControl/Private/RemoteControlModule.cpp
--- a/Engine/Plugins/VirtualProduction/RemoteControl/Source/RemoteControl/Private/RemoteControlModule.cpp
+++ b/Engine/Plugins/VirtualProduction/RemoteControl/Source/RemoteControl/Private/RemoteControlModule.cpp
@@ -658,7 +658,7 @@ bool FRemoteControlModule::ResolveCall(const FString& ObjectPath, const FString&
 				ErrorText = FString::Printf(TEXT("Function: %s does not exist on object: %s"), *FunctionName, *ObjectPath);
 				bSuccess = false;
 			}
-			else if (!Function->HasAllFunctionFlags(FUNC_BlueprintCallable | FUNC_Public)
+			else if ((!Function->HasAllFunctionFlags(FUNC_BlueprintCallable | FUNC_Public) && !Function->HasAllFunctionFlags(FUNC_BlueprintEvent))
 #if WITH_EDITOR
 					|| Function->HasMetaData(RemoteControlUtil::NAME_DeprecatedFunction)
 					|| Function->HasMetaData(RemoteControlUtil::NAME_ScriptNoExport)
