~~~ Remote Control Logic User Interface ~~~

* Base implementation of UI for Controllers, Behavior, Actions
* Additional changes for UI Conformance to Figma
* Further revisions will bring implementation of any UI functionality missing or incomplete in the original prototype.

diff --git a/Engine/Plugins/VirtualProduction/RemoteControl/Config/BaseRemoteControl.ini b/Engine/Plugins/VirtualProduction/RemoteControl/Config/BaseRemoteControl.ini
--- a/Engine/Plugins/VirtualProduction/RemoteControl/Config/BaseRemoteControl.ini
+++ b/Engine/Plugins/VirtualProduction/RemoteControl/Config/BaseRemoteControl.ini
@@ -1,3 +1,17 @@
 [CoreRedirects]
 +PropertyRedirects=(OldName="RemoteControlField.Label",NewName="RemoteControlEntity.Label")
-+PropertyRedirects=(OldName="RemoteControlField.Id",NewName="RemoteControlEntity.Id")
\ No newline at end of file
++PropertyRedirects=(OldName="RemoteControlField.Id",NewName="RemoteControlEntity.Id")
+
+[/Script/RemoteControlLogic.RemoteControlLogicConfig]
++SupportedControllerTypes=EPropertyBagPropertyType::Bool
++SupportedControllerTypes=EPropertyBagPropertyType::Byte
++SupportedControllerTypes=EPropertyBagPropertyType::Int32
++SupportedControllerTypes=EPropertyBagPropertyType::Float
++SupportedControllerTypes=EPropertyBagPropertyType::Double
++SupportedControllerTypes=EPropertyBagPropertyType::Name
++SupportedControllerTypes=EPropertyBagPropertyType::String
++SupportedControllerTypes=EPropertyBagPropertyType::Text
++SupportedControllerTypes=EPropertyBagPropertyType::Struct
++SupportedStructTypes="FVector"
++SupportedStructTypes="FColor"
++SupportedStructTypes="FRotator"
\ No newline at end of file
diff --git a/Engine/Plugins/VirtualProduction/RemoteControl/Source/RemoteControl/Private/RCVirtualPropertyContainer.cpp b/Engine/Plugins/VirtualProduction/RemoteControl/Source/RemoteControl/Private/RCVirtualPropertyContainer.cpp
--- a/Engine/Plugins/VirtualProduction/RemoteControl/Source/RemoteControl/Private/RCVirtualPropertyContainer.cpp
+++ b/Engine/Plugins/VirtualProduction/RemoteControl/Source/RemoteControl/Private/RCVirtualPropertyContainer.cpp
@@ -182,15 +182,15 @@ FName URCVirtualPropertyContainerBase::GenerateUniquePropertyName(const FName& I
 	};
 
 	const UEnum* Enum = FindObjectChecked<UEnum>(nullptr, TEXT("/Script/StructUtils.EPropertyBagPropertyType"));
-	const FName DefaultName = Enum->GetValueAsName(InValueType);
+	FString DefaultName = Enum->GetValueAsName(InValueType).ToString();	
+	DefaultName.RemoveFromStart("EPropertyBagPropertyType::");
 	
 	int32 Index = 0;
-	const FString InitialName = InPropertyName.IsNone() ? DefaultName.ToString() : InPropertyName.ToString();
+	const FString InitialName = InPropertyName.IsNone() ? DefaultName : InPropertyName.ToString();
 	FString FinalName = InitialName;
-
 	
+	// Recursively search for an available name by incrementing suffix till we find one.
 	const FPropertyBagPropertyDesc* PropertyDesc = InContainer->Bag.FindPropertyDescByName(*FinalName);
-
 	while (PropertyDesc)
 	{
 		++Index;
diff --git a/Engine/Plugins/VirtualProduction/RemoteControl/Source/RemoteControl/Private/RemoteControlPreset.cpp b/Engine/Plugins/VirtualProduction/RemoteControl/Source/RemoteControl/Private/RemoteControlPreset.cpp
--- a/Engine/Plugins/VirtualProduction/RemoteControl/Source/RemoteControl/Private/RemoteControlPreset.cpp
+++ b/Engine/Plugins/VirtualProduction/RemoteControl/Source/RemoteControl/Private/RemoteControlPreset.cpp
@@ -11,6 +11,7 @@
 #include "HAL/IConsoleManager.h"
 #include "IRemoteControlModule.h"
 #include "RCVirtualPropertyContainer.h"
+#include "RCVirtualProperty.h"
 #include "Misc/CoreDelegates.h"
 #include "Misc/ITransaction.h"
 #include "Misc/Optional.h"
@@ -780,6 +781,87 @@ FName URemoteControlPreset::GenerateUniqueLabel(const FName InDesiredName) const
 	return Registry->GenerateUniqueLabel(InDesiredName);
 }
 
+URCVirtualPropertyBase* URemoteControlPreset::GetVirtualProperty(const FName& InPropertyName) const
+{
+	if (!ensure(ControllerContainer))
+	{
+		return nullptr;
+	}
+
+	return ControllerContainer->GetVirtualProperty(InPropertyName);
+}
+
+URCVirtualPropertyInContainer* URemoteControlPreset::AddVirtualProperty(TSubclassOf<URCVirtualPropertyInContainer> InPropertyClass, const EPropertyBagPropertyType InValueType, UObject* InValueTypeObject /*= nullptr*/, const FName InPropertyName /*= NAME_None*/)
+{
+	if (!ensure(ControllerContainer))
+	{
+		return nullptr;
+	}
+
+	// New Property Name
+	FName NewPropertyName = InPropertyName;
+	if (NewPropertyName.IsNone())
+	{
+		NewPropertyName = URCVirtualPropertyContainerBase::GenerateUniquePropertyName(TEXT(""), InValueType, ControllerContainer);
+	}	
+
+	return ControllerContainer->AddProperty(NewPropertyName, InPropertyClass, InValueType, InValueTypeObject);
+}
+
+bool URemoteControlPreset::RemoveVirtualProperty(const FName& InPropertyName)
+{
+	if (ensure(ControllerContainer))
+	{
+		return ControllerContainer->RemoveProperty(InPropertyName);
+	}
+
+	return false;
+	
+}
+
+void URemoteControlPreset::ResetVirtualProperties()
+{
+	if (ensure(ControllerContainer))
+	{
+		return ControllerContainer->Reset();
+	}
+}
+
+int32 URemoteControlPreset::GetNumVirtualProperties() const
+{
+	if (ensure(ControllerContainer))
+	{
+		return ControllerContainer->GetNumVirtualProperties();
+	}
+
+	return 0;
+}
+
+TSharedPtr<FStructOnScope> URemoteControlPreset::GetControllerContainerStructOnScope()
+{
+	if (ensure(ControllerContainer))
+	{
+		return ControllerContainer->CreateStructOnScope();
+	}
+
+	return nullptr;
+}
+
+void URemoteControlPreset::SetControllerContainer(URCVirtualPropertyContainerBase* InControllerContainer)
+{
+	ControllerContainer = InControllerContainer;
+}
+
+#if WITH_EDITOR
+void URemoteControlPreset::OnModifyVirtualProperty(const FPropertyChangedEvent& PropertyChangedEvent)
+{
+	if (ensure(ControllerContainer))
+	{
+		ControllerContainer->OnModifyPropertyValue(PropertyChangedEvent);
+	}	
+}
+#endif
+
 TWeakPtr<FRemoteControlProperty> URemoteControlPreset::ExposeProperty(UObject* Object, FRCFieldPathInfo FieldPath, FRemoteControlPresetExposeArgs Args)
 {
 	if (!Object)
diff --git a/Engine/Plugins/VirtualProduction/RemoteControl/Source/RemoteControl/Public/RemoteControlPreset.h b/Engine/Plugins/VirtualProduction/RemoteControl/Source/RemoteControl/Public/RemoteControlPreset.h
--- a/Engine/Plugins/VirtualProduction/RemoteControl/Source/RemoteControl/Public/RemoteControlPreset.h
+++ b/Engine/Plugins/VirtualProduction/RemoteControl/Source/RemoteControl/Public/RemoteControlPreset.h
@@ -12,11 +12,11 @@
 
 #include "RemoteControlPreset.generated.h"
 
-class URCVirtualPropertyContainerBase;
 class AActor;
 class IStructSerializerBackend;
 class IStructDeserializerBackend;
 enum class EPackageReloadPhase : uint8;
+enum class EPropertyBagPropertyType : uint8;
 struct FPropertyChangedEvent;
 struct FRCFieldPathInfo;
 struct FRemoteControlActor;
@@ -24,6 +24,9 @@ struct FRemoteControlPresetLayout;
 class FTransactionObjectEvent;
 class FRemoteControlPresetRebindingManager;
 class UBlueprint;
+class URCVirtualPropertyBase;
+class URCVirtualPropertyContainerBase;
+class URCVirtualPropertyInContainer;
 class URemoteControlExposeRegistry;
 class URemoteControlBinding;
 class URemoteControlPreset;
@@ -837,11 +840,49 @@ public:
 	UPROPERTY(EditAnywhere, Category = "Remote Control Preset")
 	TArray<URemoteControlBinding*> Bindings;
 
+	/** ~~~Virtual Property Wrapper Functions ~~~
+	* 
+	* The goal is to hide the Controller Container and provide a simple interface for Controller access to UI and Web.
+	*/
+
+	/** Fetches a virtual property by specified name. */
+	URCVirtualPropertyBase* GetVirtualProperty(const FName& InPropertyName) const;
+
+	/** Adds a Virtual Property (Controller) to the Remote Control Preset */
+	URCVirtualPropertyInContainer* AddVirtualProperty(TSubclassOf<URCVirtualPropertyInContainer> InPropertyClass, const EPropertyBagPropertyType InValueType, UObject* InValueTypeObject = nullptr, const FName InPropertyName = NAME_None);
+
+	/** Removes a given Virtual Property (by Name) from the Remote Control preset */
+	bool RemoveVirtualProperty(const FName& InPropertyName);
+
+	/** Removes all virtual properties held by this Remote Control preset*/
+	void ResetVirtualProperties();
+
+	/** Returns the number of Virtual Properties contained in this Remote Control preset*/
+	int32 GetNumVirtualProperties() const;
+
+	/** Returns the Struct On Scope of the Controller Container (value ptr of the virtual properties)*
+	* Currently used by the UI class SRCControllerPanelList for user value entry via the RC Controllers panel*/
+	TSharedPtr<FStructOnScope> GetControllerContainerStructOnScope();
+
+	/** Sets the Controller Container (holds all Virtual Properties) */
+	void SetControllerContainer(URCVirtualPropertyContainerBase* InControllerContainer);
+
+	/** Checks whether the Controller Container (which holds all Virtual Properties) is currently valid*/
+	bool IsControllerContainerValid()
+	{
+		return ControllerContainer != nullptr;
+	}
+
+#if WITH_EDITOR
+	/** Called when a virtual property is modified. This call is routed to the Controller for evaluating associated Logic & Behaviours*/
+	void OnModifyVirtualProperty(const FPropertyChangedEvent& PropertyChangedEvent);
+#endif
+
+private:
 	/** Holds virtual controllers properties, behaviours and actions */
 	UPROPERTY(Instanced)
 	TObjectPtr<URCVirtualPropertyContainerBase> ControllerContainer;
 
-private:
 	/** Generate a unique alias for this target. */
 	FName GenerateAliasForObjects(const TArray<UObject*>& Objects);
 	
diff --git a/Engine/Plugins/VirtualProduction/RemoteControl/Source/RemoteControlLogic/Private/RemoteControlLogicModule.cpp b/Engine/Plugins/VirtualProduction/RemoteControl/Source/RemoteControlLogic/Private/RemoteControlLogicModule.cpp
--- a/Engine/Plugins/VirtualProduction/RemoteControl/Source/RemoteControlLogic/Private/RemoteControlLogicModule.cpp
+++ b/Engine/Plugins/VirtualProduction/RemoteControl/Source/RemoteControlLogic/Private/RemoteControlLogicModule.cpp
@@ -27,12 +27,12 @@ private:
 	void OnPostInitPropertiesRemoteControlPreset(URemoteControlPreset* InPreset) const
 	{
 		// Create a new controller for preset if that does not exists
-		if (InPreset && InPreset->ControllerContainer.IsNull())
+		if (InPreset && !InPreset->IsControllerContainerValid())
 		{
 			URCControllerContainer* ControllerContainer = NewObject<URCControllerContainer>(InPreset);
 			ControllerContainer->PresetWeakPtr = InPreset;
 
-			InPreset->ControllerContainer = ControllerContainer;
+			InPreset->SetControllerContainer(ControllerContainer);
 		}
 	}
 };
diff --git a/Engine/Plugins/VirtualProduction/RemoteControl/Source/RemoteControlLogic/Private/Tests/RemoteControlLogicTests.cpp b/Engine/Plugins/VirtualProduction/RemoteControl/Source/RemoteControlLogic/Private/Tests/RemoteControlLogicTests.cpp
--- a/Engine/Plugins/VirtualProduction/RemoteControl/Source/RemoteControlLogic/Private/Tests/RemoteControlLogicTests.cpp
+++ b/Engine/Plugins/VirtualProduction/RemoteControl/Source/RemoteControlLogic/Private/Tests/RemoteControlLogicTests.cpp
@@ -45,17 +45,17 @@ bool FRemoteControlLogicTest::RunTest(const FString& Parameters)
  
 	// 3. Add Controllers
 	// 3.1.1 Create Float Property
-	URCController* FloatController =  Cast<URCController>(Preset->ControllerContainer->AddProperty(TEXT("FloatController"), URCController::StaticClass(), EPropertyBagPropertyType::Float));
+	URCController* FloatController =  Cast<URCController>(Preset->AddVirtualProperty(URCController::StaticClass(), EPropertyBagPropertyType::Float, nullptr, TEXT("FloatController")) );
 	// 3.1.2 Create Second Float Property
-	URCController* FloatController1 = Cast<URCController>(Preset->ControllerContainer->AddProperty(TEXT("FloatController1"), URCController::StaticClass(), EPropertyBagPropertyType::Float));
+	URCController* FloatController1 = Cast<URCController>(Preset->AddVirtualProperty(URCController::StaticClass(), EPropertyBagPropertyType::Float, nullptr, TEXT("FloatController1")) );
 	// 3.1.3 Create Bool Property
-	URCController* BoolController = Cast<URCController>(Preset->ControllerContainer->AddProperty(TEXT("BoolController"), URCController::StaticClass(), EPropertyBagPropertyType::Bool));
+	URCController* BoolController = Cast<URCController>(Preset->AddVirtualProperty(URCController::StaticClass(), EPropertyBagPropertyType::Bool, nullptr, TEXT("BoolController")));
 	// 3.1.4 Create Int Controller
-	URCController* IntController =  Cast<URCController>(Preset->ControllerContainer->AddProperty(TEXT("IntController"), URCController::StaticClass(), EPropertyBagPropertyType::Int32));
+	URCController* IntController =  Cast<URCController>(Preset->AddVirtualProperty(URCController::StaticClass(), EPropertyBagPropertyType::Int32, nullptr, TEXT("IntController")));
 	// 3.1.5 Create String Controller
-	URCController* StrController =  Cast<URCController>(Preset->ControllerContainer->AddProperty(TEXT("StrController"), URCController::StaticClass(), EPropertyBagPropertyType::String));
+	URCController* StrController =  Cast<URCController>(Preset->AddVirtualProperty(URCController::StaticClass(), EPropertyBagPropertyType::String, nullptr, TEXT("StrController")));
 	// 3.1.6 Create Vector Struct Controller
-	URCController* VectorStructController =  Cast<URCController>(Preset->ControllerContainer->AddProperty(TEXT("VectorProperty"), URCController::StaticClass(), EPropertyBagPropertyType::Struct, TBaseStructure<FVector>::Get()));
+	URCController* VectorStructController =  Cast<URCController>(Preset->AddVirtualProperty(URCController::StaticClass(), EPropertyBagPropertyType::Struct, TBaseStructure<FVector>::Get(), TEXT("VectorProperty")));
  
 	// 3.2 Set Float Property Value to controller
 	constexpr float FloatValue = 0.65f;
@@ -151,19 +151,19 @@ bool FRemoteControlLogicTest::RunTest(const FString& Parameters)
 	TestEqual(TEXT("After empty actions the count should be = 0"), FloatController->Behaviours.Num(), 0);
  
 	// 9. Remove Controllers
-	const int32 NumPropertiesBeforeRemove = Preset->ControllerContainer->GetNumVirtualProperties();
+	const int32 NumPropertiesBeforeRemove = Preset->GetNumVirtualProperties();
 	float OutFloatValue = 0.f;
 	float OutFloatValue1 = 0.f;
 	FloatController->GetValueFloat(OutFloatValue);
-	Preset->ControllerContainer->RemoveProperty(FloatController1->PropertyName);
+	Preset->RemoveVirtualProperty(FloatController1->PropertyName);
 	FloatController->GetValueFloat(OutFloatValue1);
 	TestEqual(TEXT("After remove properties the old value should be the same"), OutFloatValue1, OutFloatValue);
-	Preset->ControllerContainer->RemoveProperty(FloatController->PropertyName);
-	Preset->ControllerContainer->RemoveProperty(BoolController->PropertyName);
-	Preset->ControllerContainer->RemoveProperty(IntController->PropertyName);
-	TestEqual(TEXT("After remove 4 properties the count should be =-4"), Preset->ControllerContainer->GetNumVirtualProperties(), NumPropertiesBeforeRemove - 4);
-	Preset->ControllerContainer->Reset();
-	TestEqual(TEXT("After empty controllers could should be equal 0"), Preset->ControllerContainer->GetNumVirtualProperties(), 0);
+	Preset->RemoveVirtualProperty(FloatController->PropertyName);
+	Preset->RemoveVirtualProperty(BoolController->PropertyName);
+	Preset->RemoveVirtualProperty(IntController->PropertyName);
+	TestEqual(TEXT("After remove 4 properties the count should be =-4"), Preset->GetNumVirtualProperties(), NumPropertiesBeforeRemove - 4);
+	Preset->ResetVirtualProperties();
+	TestEqual(TEXT("After empty controllers could should be equal 0"), Preset->GetNumVirtualProperties(), 0);
  
 	return true;
 }
diff --git a/Engine/Plugins/VirtualProduction/RemoteControl/Source/RemoteControlLogic/Public/RemoteControlLogicConfig.h b/Engine/Plugins/VirtualProduction/RemoteControl/Source/RemoteControlLogic/Public/RemoteControlLogicConfig.h
new file mode 100644
--- /dev/null
+++ b/Engine/Plugins/VirtualProduction/RemoteControl/Source/RemoteControlLogic/Public/RemoteControlLogicConfig.h
@@ -0,0 +1,26 @@
+// Copyright Epic Games, Inc. All Rights Reserved.
+
+#pragma once
+
+#include "PropertyBag.h"
+
+#include "RemoteControlLogicConfig.generated.h"
+
+/** Config class for Remote Control Logic related settings.
+* Supports configuration of a valid list of Controller types.
+*/
+UCLASS(config=RemoteControl)
+class REMOTECONTROLLOGIC_API URemoteControlLogicConfig: public UObject
+{
+	GENERATED_BODY()
+
+public:
+	/** PropertyBag types that are supported for use as Controllers */
+	UPROPERTY(config)
+	TArray<EPropertyBagPropertyType> SupportedControllerTypes;
+
+	/** Certain types such as FVector / FRotator / FColor are specializations of UStructProperty
+	* This config represents the list of such types that may be used as Controllers*/
+	UPROPERTY(config)
+	TArray<FName> SupportedControllerStructTypes;
+};
diff --git a/Engine/Plugins/VirtualProduction/RemoteControl/Source/RemoteControlUI/Private/Factories/RCBehaviourNodeFactory.cpp b/Engine/Plugins/VirtualProduction/RemoteControl/Source/RemoteControlUI/Private/Factories/RCBehaviourNodeFactory.cpp
new file mode 100644
--- /dev/null
+++ b/Engine/Plugins/VirtualProduction/RemoteControl/Source/RemoteControlUI/Private/Factories/RCBehaviourNodeFactory.cpp
@@ -0,0 +1,42 @@
+﻿// Copyright Epic Games, Inc. All Rights Reserved.
+
+#include "RCBehaviourNodeFactory.h"
+
+#include "AssetToolsModule.h"
+#include "Behaviour/RCBehaviourBlueprintNode.h"
+#include "Kismet2/KismetEditorUtilities.h"
+
+#define LOCTEXT_NAMESPACE "LevelSnapshotFilterFactory"
+
+URCBehaviourNodeFactory::URCBehaviourNodeFactory()
+{
+	SupportedClass = UBlueprint::StaticClass();
+	ParentClass = URCBehaviourBlueprintNode::StaticClass();
+
+	bCreateNew = true;
+	bEditAfterNew = true;
+}
+
+FText URCBehaviourNodeFactory::GetDisplayName() const
+{
+	return LOCTEXT("DisplayName", "Remote Control Behaviour Node");
+}
+
+UObject* URCBehaviourNodeFactory::FactoryCreateNew(UClass* Class, UObject* InParent, FName Name, EObjectFlags Flags, UObject* Context, FFeedbackContext* Warn)
+{
+	UBlueprint* NodeBlueprint = nullptr;
+	if (ensure(SupportedClass == Class))
+	{
+		ensure(0 != (RF_Public & Flags));
+		NodeBlueprint = FKismetEditorUtilities::CreateBlueprint(ParentClass, InParent, Name, BPTYPE_Normal, UBlueprint::StaticClass(), UBlueprintGeneratedClass::StaticClass(), NAME_None);
+	}
+	return NodeBlueprint;
+}
+
+uint32 URCBehaviourNodeFactory::GetMenuCategories() const
+{
+	IAssetTools& AssetTools = FModuleManager::LoadModuleChecked<FAssetToolsModule>("AssetTools").Get();
+	return AssetTools.RegisterAdvancedAssetCategory("Remote Control", LOCTEXT("AssetCategoryName", "Remote Control Behaviour Node"));
+}
+
+#undef LOCTEXT_NAMESPACE
\ No newline at end of file
diff --git a/Engine/Plugins/VirtualProduction/RemoteControl/Source/RemoteControlUI/Private/Factories/RCBehaviourNodeFactory.h b/Engine/Plugins/VirtualProduction/RemoteControl/Source/RemoteControlUI/Private/Factories/RCBehaviourNodeFactory.h
new file mode 100644
--- /dev/null
+++ b/Engine/Plugins/VirtualProduction/RemoteControl/Source/RemoteControlUI/Private/Factories/RCBehaviourNodeFactory.h
@@ -0,0 +1,30 @@
+﻿// Copyright Epic Games, Inc. All Rights Reserved.
+
+#pragma once
+
+#include "CoreMinimal.h"
+#include "Factories/Factory.h"
+
+#include "RCBehaviourNodeFactory.generated.h"
+
+/** Factory for Remote Control Logic's custom blueprint behavior feature
+* The associated Blueprint class is URCBehaviourBlueprintNode
+*/
+UCLASS(hidecategories = Object)
+class URCBehaviourNodeFactory : public UFactory
+{
+	GENERATED_BODY()
+public:
+	URCBehaviourNodeFactory();
+	
+	/** The parent class of the created blueprint */
+	UPROPERTY(VisibleAnywhere, Category="RCBehaviourNodeFactory", meta=(AllowAbstract = "", BlueprintBaseOnly = ""))
+	TSubclassOf<UObject> ParentClass;
+	
+	//~ Begin UFactory Interface
+	virtual FText GetDisplayName() const override;
+	virtual UObject* FactoryCreateNew(UClass* Class, UObject* InParent, FName Name, EObjectFlags Flags, UObject* Context, FFeedbackContext* Warn) override;
+	virtual uint32 GetMenuCategories() const override;
+	//~ Begin UFactory Interface
+	
+};
\ No newline at end of file
diff --git a/Engine/Plugins/VirtualProduction/RemoteControl/Source/RemoteControlUI/Private/UI/Action/RCActionModel.cpp b/Engine/Plugins/VirtualProduction/RemoteControl/Source/RemoteControlUI/Private/UI/Action/RCActionModel.cpp
new file mode 100644
--- /dev/null
+++ b/Engine/Plugins/VirtualProduction/RemoteControl/Source/RemoteControlUI/Private/UI/Action/RCActionModel.cpp
@@ -0,0 +1,156 @@
+﻿// Copyright Epic Games, Inc. All Rights Reserved.
+
+#include "RCActionModel.h"
+
+#include "Action/RCFunctionAction.h"
+#include "Action/RCPropertyAction.h"
+#include "Controller/RCController.h"
+#include "IDetailTreeNode.h"
+#include "IPropertyRowGenerator.h"
+#include "RCVirtualProperty.h"
+#include "RCVirtualPropertyContainer.h"
+#include "RemoteControlField.h"
+#include "RemoteControlPreset.h"
+#include "UI/RCUIHelpers.h"
+#include "UI/RemoteControlPanelStyle.h"
+#include "UI/SRCPanelDragHandle.h"
+#include "UI/SRCPanelExposedEntity.h"
+
+FRCActionModel::FRCActionModel(URCAction* InAction)
+	: ActionWeakPtr(InAction)
+{
+}
+
+URCAction* FRCActionModel::GetAction() const
+{
+	return ActionWeakPtr.Get();
+}
+
+FRCPropertyActionModel::FRCPropertyActionModel(URCPropertyAction* InPropertyAction)
+	: FRCActionModel(InPropertyAction)
+	, PropertyActionWeakPtr(InPropertyAction)
+{
+	FPropertyRowGeneratorArgs Args;
+	Args.bShouldShowHiddenProperties = true;
+	PropertyRowGenerator = FModuleManager::GetModuleChecked<FPropertyEditorModule>("PropertyEditor").CreatePropertyRowGenerator(Args);
+
+	if (InPropertyAction)
+	{
+		if (const TSharedPtr<FStructOnScope> StructOnScope = InPropertyAction->PropertySelfContainer->CreateStructOnScope())
+		{
+			PropertyRowGenerator->SetStructure(StructOnScope);
+
+			for (const TSharedRef<IDetailTreeNode>& CategoryNode : PropertyRowGenerator->GetRootTreeNodes())
+			{
+				TArray<TSharedRef<IDetailTreeNode>> Children;
+				CategoryNode->GetChildren(Children);
+				for (const TSharedRef<IDetailTreeNode>& Child : Children)
+				{
+					DetailTreeNodeWeakPtr = Child;
+					break;
+				}
+			}
+		}
+	}
+}
+
+const FName& FRCPropertyActionModel::GetPropertyName() const
+{
+	return PropertyActionWeakPtr.Get()->PropertySelfContainer->PropertyName;
+}
+
+TSharedRef<SWidget> FRCPropertyActionModel::GetWidget() const
+{
+	if (DetailTreeNodeWeakPtr.IsValid())
+	{
+		const FNodeWidgets NodeWidgets = DetailTreeNodeWeakPtr.Pin()->CreateNodeWidgets();
+
+		const TSharedRef<SHorizontalBox> FieldWidget = SNew(SHorizontalBox);
+	
+		if (NodeWidgets.ValueWidget)
+		{
+			FieldWidget->AddSlot()
+				.Padding(FMargin(3.0f, 2.0f))
+				.HAlign(HAlign_Right)
+				.AutoWidth()
+				[
+					NodeWidgets.ValueWidget.ToSharedRef()
+				];
+		}
+		else if (NodeWidgets.WholeRowWidget)
+		{
+			FieldWidget->AddSlot()
+				.Padding(FMargin(3.0f, 2.0f))
+				.AutoWidth()
+				[
+					NodeWidgets.WholeRowWidget.ToSharedRef()
+				];
+		}
+
+		TSharedRef<SBorder> BorderWidget = SNew(SBorder)
+			.Padding(0.0f)
+			.BorderImage(FRemoteControlPanelStyle::Get()->GetBrush("RemoteControlPanel.ExposedFieldBorder"));
+
+		FLinearColor TypeColor;
+		if (PropertyActionWeakPtr.IsValid())
+		{
+			TypeColor = UE::RCUIHelpers::GetFieldClassTypeColor(PropertyActionWeakPtr->PropertySelfContainer->GetProperty());
+		}
+	
+		return SNew(SHorizontalBox)
+			.Clipping(EWidgetClipping::OnDemand)
+
+			// Variable Color Bar
+			+ SHorizontalBox::Slot()
+			.VAlign(VAlign_Fill)
+			.AutoWidth()
+			.Padding(FMargin(3.f))
+			[
+				SNew(SBorder)
+				.Visibility(EVisibility::HitTestInvisible)
+				.BorderImage(FAppStyle::Get().GetBrush("NumericEntrySpinBox.NarrowDecorator"))
+				.BorderBackgroundColor(TypeColor)
+				.Padding(FMargin(5.0f, 0.0f, 0.0f, 0.0f))
+			]
+			// Field name
+			+SHorizontalBox::Slot()
+			.VAlign(VAlign_Center)
+			.AutoWidth()
+			[
+				NodeWidgets.NameWidget.ToSharedRef()
+			]
+			// Value Widget
+			+ SHorizontalBox::Slot()
+			.VAlign(VAlign_Center)
+			.AutoWidth()
+			[
+				FieldWidget
+			];
+	}
+	
+	return FRCActionModel::GetWidget();
+}
+
+FRCFunctionActionModel::FRCFunctionActionModel(URCFunctionAction* InFunctionAction)
+	: FRCActionModel(InFunctionAction)
+	, FunctionActionWeakPtr(InFunctionAction)
+{
+}
+
+TSharedRef<SWidget> FRCFunctionActionModel::GetWidget() const
+{
+	if (FunctionActionWeakPtr.IsValid())
+	{
+		if (URemoteControlPreset* Preset = FunctionActionWeakPtr->PresetWeakPtr.Get())
+		{
+			if (const TSharedPtr<FRemoteControlFunction> RemoteControlField = Preset->GetExposedEntity<FRemoteControlFunction>(FunctionActionWeakPtr->ExposedFieldId).Pin())
+			{
+				const FString FinalNameStr = RemoteControlField->GetLabel().ToString();
+
+				return SNew(STextBlock).Text(FText::FromString(FinalNameStr));
+			}
+		}
+	}
+	
+	return FRCActionModel::GetWidget();
+}
\ No newline at end of file
diff --git a/Engine/Plugins/VirtualProduction/RemoteControl/Source/RemoteControlUI/Private/UI/Action/RCActionModel.h b/Engine/Plugins/VirtualProduction/RemoteControl/Source/RemoteControlUI/Private/UI/Action/RCActionModel.h
new file mode 100644
--- /dev/null
+++ b/Engine/Plugins/VirtualProduction/RemoteControl/Source/RemoteControlUI/Private/UI/Action/RCActionModel.h
@@ -0,0 +1,82 @@
+﻿// Copyright Epic Games, Inc. All Rights Reserved.
+
+#pragma once
+
+#include "UI/BaseLogicUI/RCLogicModeBase.h"
+#include "UObject/WeakFieldPtr.h"
+
+class IDetailTreeNode;
+class IPropertyRowGenerator;
+class SWidget;
+class URCAction;
+class URCPropertyAction;
+class URCVirtualProperty;
+class URCFunctionAction;
+
+/* 
+* ~ FRCActionModel ~
+*
+* UI model for representing an Action.
+* Contains a row widget with Action related metadata and a generic value widget
+*/
+class FRCActionModel : public FRCLogicModeBase
+{
+public:
+	FRCActionModel(URCAction* InAction);
+	
+	/** Fetch the Action (data model) associated with us*/
+	URCAction* GetAction() const;
+
+private:
+	/** The Action (data model) associated with us*/
+	TWeakObjectPtr<URCAction> ActionWeakPtr;
+};
+
+/* 
+* ~ FRCPropertyActionModel ~
+*
+* UI model for Property based Actions
+*/
+class FRCPropertyActionModel : public FRCActionModel
+{
+public:
+	FRCPropertyActionModel(URCPropertyAction* InPropertyAction);
+
+	/** Property Name associated with this Action */
+	const FName& GetPropertyName() const;
+
+	/** The widget to be rendered for this Property 
+	* Used to represent a single row when added to the Actions Panel List
+	*/
+	virtual TSharedRef<SWidget> GetWidget() const override;
+
+private:
+	/** The Property Action (data model) associated with us*/
+	TWeakObjectPtr<URCPropertyAction> PropertyActionWeakPtr;
+
+	/** The row generator used to represent this widget as a row, when used with SListView */
+	TSharedPtr<IPropertyRowGenerator> PropertyRowGenerator;
+
+	/** Used to create a generic Value Widget for the property row widget*/
+	TWeakPtr<IDetailTreeNode> DetailTreeNodeWeakPtr;
+};
+
+/*
+* ~ FRCFunctionActionModel ~
+*
+* UI model for Function based Actions
+*/
+class FRCFunctionActionModel : public FRCActionModel
+{
+public:
+	FRCFunctionActionModel(URCFunctionAction* InAction);
+
+	/** The widget to be rendered for this Function
+	* Used to represent a single row when added to the Actions Panel List
+	*/
+	virtual TSharedRef<SWidget> GetWidget() const override;
+
+private:
+	/** The Function Action (data model) associated with us*/
+	TWeakObjectPtr<URCFunctionAction> FunctionActionWeakPtr;
+};
diff --git a/Engine/Plugins/VirtualProduction/RemoteControl/Source/RemoteControlUI/Private/UI/Action/SRCActionPanel.cpp b/Engine/Plugins/VirtualProduction/RemoteControl/Source/RemoteControlUI/Private/UI/Action/SRCActionPanel.cpp
new file mode 100644
--- /dev/null
+++ b/Engine/Plugins/VirtualProduction/RemoteControl/Source/RemoteControlUI/Private/UI/Action/SRCActionPanel.cpp
@@ -0,0 +1,111 @@
+﻿// Copyright Epic Games, Inc. All Rights Reserved.
+
+#include "SRCActionPanel.h"
+
+#include "SlateOptMacros.h"
+#include "UI/SRemoteControlPanel.h"
+#include "SRCActionPanelHeader.h"
+#include "SRCActionPanelList.h"
+#include "SRCActionPanelFooter.h"
+
+#include "UI/Behaviour/RCBehaviourModel.h"
+
+#define LOCTEXT_NAMESPACE "SRCActionPanel"
+
+BEGIN_SLATE_FUNCTION_BUILD_OPTIMIZATION
+
+void SRCActionPanel::Construct(const FArguments& InArgs, const TSharedRef<SRemoteControlPanel>& InPanel)
+{
+	SRCLogicPanelBase::Construct(SRCLogicPanelBase::FArguments(), InPanel);
+	
+	WrappedBoxWidget = SNew(SBox);
+	UpdateWrappedWidget();
+	
+	ChildSlot
+	[
+		WrappedBoxWidget.ToSharedRef()
+	];
+
+	// Register delegates
+	InPanel->OnBehaviourSelectionChanged.AddSP(this, &SRCActionPanel::OnBehaviourSelectionChanged);
+}
+
+END_SLATE_FUNCTION_BUILD_OPTIMIZATION
+
+void SRCActionPanel::OnBehaviourSelectionChanged(TSharedPtr<FRCBehaviourModel> InBehaviourItem)
+{
+	SelectedBehaviourItemWeakPtr = InBehaviourItem;
+	UpdateWrappedWidget(InBehaviourItem);
+}
+
+void SRCActionPanel::UpdateWrappedWidget(TSharedPtr<FRCBehaviourModel> InBehaviourItem)
+{
+	if (InBehaviourItem.IsValid())
+	{
+		WrappedBoxWidget->SetContent(
+		SNew(SVerticalBox)
+		// Header
+		+ SVerticalBox::Slot()
+		.Padding(5.0f, 5.0f)
+		.AutoHeight()
+		[
+			SNew(SRCActionPanelHeader, SharedThis(this), InBehaviourItem.ToSharedRef())
+		]
+
+		// Border (separating Header and "If Value Is" panel)
+		+ SVerticalBox::Slot()
+		.AutoHeight()
+		[
+			SNew(SBorder)
+			.BorderImage(FAppStyle::Get().GetBrush("ToolPanel.GroupBorder"))
+			.BorderBackgroundColor(FColor::Black)
+			.Padding(FMargin(0.f, 1.f))
+			[
+				SNullWidget::NullWidget
+			]
+		]
+
+		// Behaviour Specific Details Panel
+		+ SVerticalBox::Slot()
+		.Padding(20.0f, 5.0f)
+		.AutoHeight()
+		[
+			InBehaviourItem->GetBehaviourDetailsWidget()			
+		]
+
+		// Actions Mapping List
+		+ SVerticalBox::Slot()
+		.AutoHeight()
+		.Padding(FMargin(5.f))
+		[
+			SAssignNew(ActionPanelList, SRCActionPanelList, SharedThis(this), InBehaviourItem)
+		]
+
+		// Footer
+		+ SVerticalBox::Slot()
+		.Padding(5.0f, 5.0f)
+		.AutoHeight()
+		[
+			SNew(SRCActionPanelFooter, SharedThis(this), InBehaviourItem.ToSharedRef())
+		]
+	);
+	}
+	else
+	{
+		WrappedBoxWidget->SetContent(
+			SNew(STextBlock).Text(LOCTEXT("Unselected", "No Behaviour Selected"))
+		);
+	}
+}
+
+bool SRCActionPanel::IsListFocused() const
+{
+	return ActionPanelList.IsValid() && ActionPanelList->IsListFocused();
+}
+
+void SRCActionPanel::DeleteSelectedPanelItem()
+{
+	ActionPanelList->DeleteSelectedPanelItem();
+}
+
+#undef LOCTEXT_NAMESPACE
\ No newline at end of file
diff --git a/Engine/Plugins/VirtualProduction/RemoteControl/Source/RemoteControlUI/Private/UI/Action/SRCActionPanel.h b/Engine/Plugins/VirtualProduction/RemoteControl/Source/RemoteControlUI/Private/UI/Action/SRCActionPanel.h
new file mode 100644
--- /dev/null
+++ b/Engine/Plugins/VirtualProduction/RemoteControl/Source/RemoteControlUI/Private/UI/Action/SRCActionPanel.h
@@ -0,0 +1,53 @@
+﻿// Copyright Epic Games, Inc. All Rights Reserved.
+
+#pragma once
+#include "UI/BaseLogicUI/SRCLogicPanelBase.h"
+#include "Widgets/Layout/SBox.h"
+
+class FRCBehaviourModel;
+class SRemoteControlPanel;
+
+/*
+* ~ SRCActionPanel ~
+*
+* UI Widget for Action Panel.
+* Contains a header, footer, list of actions and a behaviour specific panel
+*/
+class REMOTECONTROLUI_API SRCActionPanel : public SRCLogicPanelBase
+{
+public:
+	SLATE_BEGIN_ARGS(SRCActionPanel)
+		{
+		}
+
+	SLATE_END_ARGS()
+	
+	/** Constructs this widget with InArgs */
+	void Construct(const FArguments& InArgs, const TSharedRef<SRemoteControlPanel>& InPanel);
+
+private:
+	/** 
+	* Behaviour selection change listener.
+	* Updates the list of actions from newly selected Behaviour
+	*/
+	void OnBehaviourSelectionChanged(TSharedPtr<FRCBehaviourModel> InBehaviourItem);
+
+	/* Rebuilds the Action Panel for a newly selected Behaviour*/
+	void UpdateWrappedWidget(TSharedPtr<FRCBehaviourModel> InBehaviourItem = nullptr);
+
+	/** The parent Behaviour that this Action panel is associated with */
+	TWeakPtr<FRCBehaviourModel> SelectedBehaviourItemWeakPtr;
+	
+	/** Parent Box for the entire widget */
+	TSharedPtr<SBox> WrappedBoxWidget;
+
+	/** Widget representing List of Actions */
+	TSharedPtr<class SRCActionPanelList> ActionPanelList;
+
+public:
+	/** Whether the Actions list widget currently has focus.*/
+	bool IsListFocused() const;
+
+	/** Delete Item UI command implementation for this panel */
+	virtual void DeleteSelectedPanelItem() override;
+};
diff --git a/Engine/Plugins/VirtualProduction/RemoteControl/Source/RemoteControlUI/Private/UI/Action/SRCActionPanelFooter.cpp b/Engine/Plugins/VirtualProduction/RemoteControl/Source/RemoteControlUI/Private/UI/Action/SRCActionPanelFooter.cpp
new file mode 100644
--- /dev/null
+++ b/Engine/Plugins/VirtualProduction/RemoteControl/Source/RemoteControlUI/Private/UI/Action/SRCActionPanelFooter.cpp
@@ -0,0 +1,220 @@
+// Copyright Epic Games, Inc. All Rights Reserved.
+
+#include "SRCActionPanelFooter.h"
+
+#include "Action/RCActionContainer.h"
+#include "Action/RCFunctionAction.h"
+#include "Action/RCPropertyAction.h"
+#include "Behaviour/RCBehaviour.h"
+#include "Controller/RCController.h"
+#include "EdGraphSchema_K2.h"
+#include "EdGraph/EdGraphPin.h"
+#include "Framework/MultiBox/MultiBoxBuilder.h"
+#include "RemoteControlPreset.h"
+#include "SlateOptMacros.h"
+#include "SPositiveActionButton.h"
+#include "SRCActionPanel.h"
+#include "Styling/CoreStyle.h"
+#include "UI/SRemoteControlPanel.h"
+#include "UI/Behaviour/RCBehaviourModel.h"
+#include "UI/BaseLogicUI/RCLogicHelpers.h"
+#include "UI/RemoteControlPanelStyle.h"
+#include "Widgets/Input/SButton.h"
+#include "Widgets/Layout/SSeparator.h"
+
+
+#define LOCTEXT_NAMESPACE "SRCActionPanelFooter"
+
+class SPositiveActionButton;
+
+BEGIN_SLATE_FUNCTION_BUILD_OPTIMIZATION
+
+void SRCActionPanelFooter::Construct(const FArguments& InArgs, TSharedRef<SRCActionPanel> InActionPanel, TSharedRef<FRCBehaviourModel> InBehaviourItem)
+{
+	SRCLogicPanelHeaderBase::Construct(SRCLogicPanelHeaderBase::FArguments());
+
+	ActionPanelWeakPtr = InActionPanel;
+	BehaviourItemWeakPtr = InBehaviourItem;
+	
+	const TSharedRef<SPositiveActionButton> AddNewMenu = SNew(SPositiveActionButton)
+		.AddMetaData<FTagMetaData>(FTagMetaData(TEXT("Add Action")))
+		.Icon(FAppStyle::Get().GetBrush("Icons.Plus"))
+		.Text(LOCTEXT("AddActionLabel", "Add Action"))
+		.OnGetMenuContent(this, &SRCActionPanelFooter::GetActionMenuContentWidget);
+	
+	ChildSlot
+	[
+		SNew(SHorizontalBox)
+		// Add new
+		+ SHorizontalBox::Slot()
+		.VAlign(VAlign_Center)
+		.AutoWidth()
+		[
+			AddNewMenu
+		]
+
+		// Add All
+		+ SHorizontalBox::Slot()
+		.VAlign(VAlign_Center)
+		.AutoWidth()
+		[
+			SNew(SButton)
+			.ContentPadding(2.0f)
+			.ButtonStyle(FAppStyle::Get(), "FlatButton")
+			.OnClicked(this, &SRCActionPanelFooter::OnAddAllField)
+			[
+				SNew(STextBlock)
+				.ColorAndOpacity(FColor::White)
+				.Text(LOCTEXT("AddAllLabel", "Add All"))
+			]
+		]
+		
+		// Empty
+		+ SHorizontalBox::Slot()
+		.VAlign(VAlign_Center)
+		.AutoWidth()
+		[
+			SNew(SButton)
+			.ContentPadding(2.0f)
+			.ButtonStyle(FAppStyle::Get(), "FlatButton")
+			.OnClicked(this, &SRCActionPanelFooter::OnClickEmptyButton)
+			[
+				SNew(STextBlock)
+				.ColorAndOpacity(FColor::White)
+				.Text(LOCTEXT("EmptyActionsButtonLabel", "Empty"))
+			]
+		]
+	];
+}
+
+END_SLATE_FUNCTION_BUILD_OPTIMIZATION
+
+TSharedRef<SWidget> SRCActionPanelFooter::GetActionMenuContentWidget()
+{
+	constexpr  bool bShouldCloseWindowAfterMenuSelection = true;
+	FMenuBuilder MenuBuilder(bShouldCloseWindowAfterMenuSelection, nullptr);
+
+	// List of exposed entities
+	if (URemoteControlPreset* Preset = ActionPanelWeakPtr.Pin()->GetPreset())
+	{
+		const TArray<TWeakPtr<FRemoteControlField>>& RemoteControlFields = Preset->GetExposedEntities<FRemoteControlField>();
+		for (const TWeakPtr<FRemoteControlField>& RemoteControlFieldWeakPtr : RemoteControlFields)
+		{
+			if (const TSharedPtr<FRemoteControlField> RemoteControlField = RemoteControlFieldWeakPtr.Pin())
+			{
+				if (const URCBehaviour* Behaviour = BehaviourItemWeakPtr.Pin()->GetBehaviour())
+				{
+					// Skip if we already have an Action created for this exposed entity
+					if (Behaviour->ActionContainer->FindActionByFieldId(RemoteControlField->GetId()))
+					{
+						continue;
+					}
+
+					// Create menu entry
+					FUIAction Action(FExecuteAction::CreateRaw(this, &SRCActionPanelFooter::OnAddActionClicked, RemoteControlField));
+					MenuBuilder.AddMenuEntry(
+						FText::Format( LOCTEXT("AddAction", "{0}"), FText::FromName(RemoteControlField->GetLabel())),
+						FText::Format( LOCTEXT( "AddActionTooltip", "{0}"), FText::FromName(RemoteControlField->GetLabel())),
+						FSlateIcon(),
+						MoveTemp(Action));
+				}
+			}
+		}
+	}
+	
+	return MenuBuilder.MakeWidget();
+}
+
+void SRCActionPanelFooter::OnAddActionClicked(TSharedPtr<FRemoteControlField> InRemoteControlField)
+{
+	if (!ensure(BehaviourItemWeakPtr.IsValid()))
+	{
+		return;
+	}
+
+	if (const URCBehaviour* Behaviour = BehaviourItemWeakPtr.Pin()->GetBehaviour())
+	{
+		URCAction* NewAction = nullptr;
+		
+		if (InRemoteControlField->FieldType == EExposedFieldType::Property)
+		{
+			NewAction = Behaviour->ActionContainer->AddAction(StaticCastSharedPtr<FRemoteControlProperty>(InRemoteControlField));
+		}
+		else if (InRemoteControlField->FieldType == EExposedFieldType::Function)
+		{
+			NewAction = Behaviour->ActionContainer->AddAction(StaticCastSharedPtr<FRemoteControlFunction>(InRemoteControlField));
+		}
+		
+		if (const TSharedPtr<SRCActionPanel> ActionPanel = ActionPanelWeakPtr.Pin())
+		{
+			if (const TSharedPtr<SRemoteControlPanel> RemoteControlPanel = ActionPanel->GetRemoteControlPanel())
+			{
+				RemoteControlPanel->OnActionAdded.Broadcast(NewAction);
+			}
+		}
+	}
+}
+
+FReply SRCActionPanelFooter::OnClickEmptyButton()
+{
+	if (const TSharedPtr<FRCBehaviourModel> BehaviourItem = BehaviourItemWeakPtr.Pin())
+	{
+		if (const URCBehaviour* Behaviour = BehaviourItem->GetBehaviour())
+		{
+			Behaviour->ActionContainer->EmptyActions();
+		}
+	}
+
+	const TSharedPtr<SRemoteControlPanel> RemoteControlPanel = ActionPanelWeakPtr.Pin()->GetRemoteControlPanel();
+	RemoteControlPanel->OnEmptyActions.Broadcast();
+	
+	return FReply::Handled();
+}
+
+FReply SRCActionPanelFooter::OnAddAllField()
+{
+	if (!ActionPanelWeakPtr.IsValid() || !BehaviourItemWeakPtr.IsValid())
+	{
+		return FReply::Handled();
+	}
+
+	TSharedPtr<SRCActionPanel> ActionPanel = ActionPanelWeakPtr.Pin();
+
+	if (URemoteControlPreset* Preset = ActionPanel->GetPreset())
+	{
+		if (const URCBehaviour* Behaviour = BehaviourItemWeakPtr.Pin()->GetBehaviour())
+		{
+			const TArray<TWeakPtr<FRemoteControlField>>& RemoteControlFields = Preset->GetExposedEntities<FRemoteControlField>();
+
+			// Enumerate the list of Exposed Entities and Functions available in this Preset for our use as Actions
+			for (const TWeakPtr<FRemoteControlField>& RemoteControlFieldWeakPtr : RemoteControlFields)
+			{
+				if (const TSharedPtr<FRemoteControlField> RemoteControlField = RemoteControlFieldWeakPtr.Pin())
+				{
+					URCAction* NewAction = nullptr;
+
+					// Property Action
+					if (RemoteControlField->FieldType == EExposedFieldType::Property)
+					{
+						NewAction = Behaviour->ActionContainer->AddAction(StaticCastSharedPtr<FRemoteControlProperty>(RemoteControlField));
+					}
+					// Function Action
+					else if (RemoteControlField->FieldType == EExposedFieldType::Function)
+					{
+						NewAction = Behaviour->ActionContainer->AddAction(StaticCastSharedPtr<FRemoteControlFunction>(RemoteControlField));
+					}
+
+					// Broadcast New Action
+					if (const TSharedPtr<SRemoteControlPanel> RemoteControlPanel = ActionPanel->GetRemoteControlPanel())
+					{
+						RemoteControlPanel->OnActionAdded.Broadcast(NewAction);
+					}
+				}
+			}
+		}
+	}
+
+	return FReply::Handled();
+}
+
+#undef LOCTEXT_NAMESPACE
\ No newline at end of file
diff --git a/Engine/Plugins/VirtualProduction/RemoteControl/Source/RemoteControlUI/Private/UI/Action/SRCActionPanelFooter.h b/Engine/Plugins/VirtualProduction/RemoteControl/Source/RemoteControlUI/Private/UI/Action/SRCActionPanelFooter.h
new file mode 100644
--- /dev/null
+++ b/Engine/Plugins/VirtualProduction/RemoteControl/Source/RemoteControlUI/Private/UI/Action/SRCActionPanelFooter.h
@@ -0,0 +1,45 @@
+// Copyright Epic Games, Inc. All Rights Reserved.
+
+#pragma once
+#include "RemoteControlField.h"
+#include "UI/BaseLogicUI/SRCLogicPanelHeaderBase.h"
+
+class SRCActionPanel;
+class FRCBehaviourModel;
+
+/**
+ *  Action panel UI Footer widget.
+ * Contains Add/ Add All / Remove/Empty buttons for managing the Actions list
+ */
+class REMOTECONTROLUI_API SRCActionPanelFooter : public SRCLogicPanelHeaderBase
+{
+public:
+	SLATE_BEGIN_ARGS(SRCActionPanelFooter)
+		{
+		}
+
+	SLATE_END_ARGS()
+
+	/** Constructs this widget with InArgs */
+	void Construct(const FArguments& InArgs, TSharedRef<SRCActionPanel> InActionPanel, TSharedRef<FRCBehaviourModel> InBehaviourItem);
+
+private:
+	/** Builds a menu containing the list of all possible Actions
+	* These are derived from the list of Exposed entities of the Remote Control Preset associated with us */
+	TSharedRef<SWidget> GetActionMenuContentWidget();
+
+	/** Handles click event for Add Action button*/
+	void OnAddActionClicked(TSharedPtr<FRemoteControlField> InRemoteControlField);
+
+	/** Handles click event for Empty button; clear all Actions from the panel*/
+	FReply OnClickEmptyButton();
+	
+	/** Handles click event for Add All button; Adds all possible actions for the active Remote Control Preset*/
+	FReply OnAddAllField();
+
+	/** The parent Action panel holding this Footer widget*/
+	TWeakPtr<SRCActionPanel> ActionPanelWeakPtr;
+	
+	/** The active Behaviour (UI Model) associated with the Actions being shown*/
+	TWeakPtr<FRCBehaviourModel> BehaviourItemWeakPtr;
+};
diff --git a/Engine/Plugins/VirtualProduction/RemoteControl/Source/RemoteControlUI/Private/UI/Action/SRCActionPanelHeader.cpp b/Engine/Plugins/VirtualProduction/RemoteControl/Source/RemoteControlUI/Private/UI/Action/SRCActionPanelHeader.cpp
new file mode 100644
--- /dev/null
+++ b/Engine/Plugins/VirtualProduction/RemoteControl/Source/RemoteControlUI/Private/UI/Action/SRCActionPanelHeader.cpp
@@ -0,0 +1,148 @@
+﻿// Copyright Epic Games, Inc. All Rights Reserved.
+
+#include "SRCActionPanelHeader.h"
+
+#include "Action/RCActionContainer.h"
+#include "Action/RCFunctionAction.h"
+#include "Action/RCPropertyAction.h"
+#include "Behaviour/RCBehaviour.h"
+#include "Controller/RCController.h"
+#include "RemoteControlPreset.h"
+#include "SlateOptMacros.h"
+#include "SRCActionPanel.h"
+#include "Styling/CoreStyle.h"
+#include "UI/Behaviour/RCBehaviourModel.h"
+#include "UI/BaseLogicUI/RCLogicHelpers.h"
+#include "UI/RCUIHelpers.h"
+#include "UI/RemoteControlPanelStyle.h"
+#include "UI/SRemoteControlPanel.h"
+#include "Widgets/Input/SButton.h"
+#include "Widgets/Layout/SSeparator.h"
+
+#define LOCTEXT_NAMESPACE "SRCActionPanelHeader"
+
+BEGIN_SLATE_FUNCTION_BUILD_OPTIMIZATION
+
+void SRCActionPanelHeader::Construct(const FArguments& InArgs, TSharedRef<SRCActionPanel> InActionPanel, TSharedRef<FRCBehaviourModel> InBehaviourItem)
+{
+	SRCLogicPanelHeaderBase::Construct(SRCLogicPanelHeaderBase::FArguments());
+
+	ActionPanelWeakPtr = InActionPanel;
+	BehaviourItemWeakPtr = InBehaviourItem;
+
+	URCBehaviour* Behaviour = BehaviourItemWeakPtr.Pin()->GetBehaviour();
+	const FText BehaviourDisplayName = Behaviour->GetDisplayName().ToUpper();
+	const FText BehaviourDescription = Behaviour->GetBehaviorDescription();
+
+	const FSlateFontInfo& FontBehaviorDesc = FRemoteControlPanelStyle::Get()->GetFontStyle("RemoteControlPanel.Behaviours.BehaviorDescription");
+
+	FLinearColor TypeColor;
+	FText TypeDisplayName;
+	if (URCController* Controller = Behaviour->ControllerWeakPtr.Get())
+	{
+		TypeColor = UE::RCUIHelpers::GetFieldClassTypeColor(Behaviour->ControllerWeakPtr->GetProperty());
+		TypeDisplayName = UE::RCUIHelpers::GetFieldClassDisplayName(Behaviour->ControllerWeakPtr->GetProperty());
+	}
+	
+	ChildSlot
+	[
+		SNew(SVerticalBox)
+		+ SVerticalBox::Slot()
+		.HAlign(HAlign_Fill)
+		.AutoHeight()
+		[
+			SNew(SHorizontalBox)
+			+ SHorizontalBox::Slot()
+			.HAlign(HAlign_Left)
+			.AutoWidth()
+			.Padding(FMargin(5.f))
+			[
+				SNew(SVerticalBox)
+				+ SVerticalBox::Slot()
+				.HAlign(HAlign_Center)
+				.AutoHeight()
+				[
+					SNew(STextBlock)
+					.ColorAndOpacity(FLinearColor(1, 1, 1, 0.4f))
+					.Text(BehaviourDisplayName)
+					.Font(FRemoteControlPanelStyle::Get()->GetFontStyle("RemoteControlPanel.Behaviours.Title"))
+				]
+
+				+ SVerticalBox::Slot()
+				.AutoHeight()
+				.Padding(FMargin(3.f))
+				[
+					SNew(SHorizontalBox)
+					+ SHorizontalBox::Slot()
+					.VAlign(VAlign_Center)
+					.AutoWidth()
+					.Padding(FMargin(5.f))
+					[
+						SNew(STextBlock)
+						.ColorAndOpacity(FLinearColor(1, 1, 1, 0.4f))
+						.Text(LOCTEXT("InputTitle", "Input"))
+						.Font(FontBehaviorDesc)
+					]
+
+					+ SHorizontalBox::Slot()
+					.VAlign(VAlign_Center)
+					.AutoWidth()
+					.Padding(FMargin(5.f, 0.f))
+					[
+						SNew(SBorder)
+						.BorderImage(FAppStyle::Get().GetBrush("SettingsEditor.CheckoutWarningBorder"))
+						.Padding(FMargin(10.f, 0.f))
+						.BorderBackgroundColor(TypeColor)
+						[
+							SNew(STextBlock)
+							.Text(TypeDisplayName)
+							.Font(FontBehaviorDesc)
+						]
+					]
+				]
+			]
+
+			+ SHorizontalBox::Slot()
+			.Padding(FMargin(10.f))
+			.VAlign(VAlign_Top)
+			[
+				SNew(STextBlock)
+				.ColorAndOpacity(FLinearColor(1, 1, 1, 0.4f))
+				.Text(BehaviourDescription)
+				.Font(FontBehaviorDesc)
+				.AutoWrapText(true) // note: use of autowidth on this slot prevents wrapping from working!
+			]
+
+			+ SHorizontalBox::Slot()
+			.VAlign(VAlign_Center)
+			.HAlign(HAlign_Right)
+			.AutoWidth()
+			.Padding(FMargin(10.f))
+			[
+				SNew(SButton)
+				.ButtonStyle(FAppStyle::Get(), "FlatButton.DarkGrey")
+				.OnClicked(this, &SRCActionPanelHeader::OnClickOverrideBlueprintButton)
+				.ContentPadding(FMargin(0.f))
+				[
+					SNew(SImage)
+					.Image(FRemoteControlPanelStyle::Get()->GetBrush("RemoteControlPanel.Behaviours.CustomBlueprint"))
+				]
+			]
+		]
+	];
+}
+
+END_SLATE_FUNCTION_BUILD_OPTIMIZATION
+
+FReply SRCActionPanelHeader::OnClickOverrideBlueprintButton()
+{
+	TSharedPtr<FRCBehaviourModel> Behaviour = BehaviourItemWeakPtr.Pin();
+	if (!ensure(Behaviour))
+		return FReply::Unhandled();
+
+	Behaviour->OnOverrideBlueprint();
+
+	return FReply::Handled();
+}
+
+#undef LOCTEXT_NAMESPACE
\ No newline at end of file
diff --git a/Engine/Plugins/VirtualProduction/RemoteControl/Source/RemoteControlUI/Private/UI/Action/SRCActionPanelHeader.h b/Engine/Plugins/VirtualProduction/RemoteControl/Source/RemoteControlUI/Private/UI/Action/SRCActionPanelHeader.h
new file mode 100644
--- /dev/null
+++ b/Engine/Plugins/VirtualProduction/RemoteControl/Source/RemoteControlUI/Private/UI/Action/SRCActionPanelHeader.h
@@ -0,0 +1,37 @@
+﻿// Copyright Epic Games, Inc. All Rights Reserved.
+
+#pragma once
+#include "RemoteControlField.h"
+#include "UI/BaseLogicUI/SRCLogicPanelHeaderBase.h"
+
+class SRCActionPanel;
+class FRCBehaviourModel;
+
+/**
+ *  Action panel UI Header widget.
+ * Contains metdata, Add/Remove buttons, etc;
+ */
+class REMOTECONTROLUI_API SRCActionPanelHeader : public SRCLogicPanelHeaderBase
+{
+public:
+	SLATE_BEGIN_ARGS(SRCActionPanelHeader)
+		{
+		}
+
+	SLATE_END_ARGS()
+
+	/** Constructs this widget with InArgs */
+	void Construct(const FArguments& InArgs, TSharedRef<SRCActionPanel> InActionPanel, TSharedRef<FRCBehaviourModel> InBehaviourItem);
+
+private:
+
+	/** Handles click event for Open Behaviour Blueprint button*/
+	FReply OnClickOverrideBlueprintButton();
+
+private:
+	/** The parent Action panel UI widget holding this Header*/
+	TWeakPtr<SRCActionPanel> ActionPanelWeakPtr;
+	
+	/** The active Behaviour (UI Model) associated with the Actions being shown*/
+	TWeakPtr<FRCBehaviourModel> BehaviourItemWeakPtr;
+};
diff --git a/Engine/Plugins/VirtualProduction/RemoteControl/Source/RemoteControlUI/Private/UI/Action/SRCActionPanelList.cpp b/Engine/Plugins/VirtualProduction/RemoteControl/Source/RemoteControlUI/Private/UI/Action/SRCActionPanelList.cpp
new file mode 100644
--- /dev/null
+++ b/Engine/Plugins/VirtualProduction/RemoteControl/Source/RemoteControlUI/Private/UI/Action/SRCActionPanelList.cpp
@@ -0,0 +1,165 @@
+﻿// Copyright Epic Games, Inc. All Rights Reserved.
+
+#include "SRCActionPanelList.h"
+
+#include "Action/RCAction.h"
+#include "Action/RCActionContainer.h"
+#include "Action/RCFunctionAction.h"
+#include "Action/RCPropertyAction.h"
+#include "RCActionModel.h"
+#include "SlateOptMacros.h"
+#include "SRCActionPanel.h"
+#include "UI/Behaviour/RCBehaviourModel.h"
+#include "UI/RCUIHelpers.h"
+#include "UI/RemoteControlPanelStyle.h"
+#include "UI/SRemoteControlPanel.h"
+#include "Widgets/Views/SListView.h"
+#include "Widgets/Views/STableRow.h"
+#include "Widgets/Views/STableViewBase.h"
+
+#define LOCTEXT_NAMESPACE "SRCActionPanelList"
+
+BEGIN_SLATE_FUNCTION_BUILD_OPTIMIZATION
+
+void SRCActionPanelList::Construct(const FArguments& InArgs, const TSharedRef<SRCActionPanel> InActionPanel, TSharedPtr<FRCBehaviourModel> InBehaviourItem)
+{
+	SRCLogicPanelListBase::Construct(SRCLogicPanelListBase::FArguments());
+
+	ActionPanelWeakPtr = InActionPanel;
+	BehaviourItemWeakPtr = InBehaviourItem;
+	
+	ListView = SNew(SListView<TSharedPtr<FRCActionModel>>)
+		.ListItemsSource( &ActionItems )
+		.OnGenerateRow(this, &SRCActionPanelList::OnGenerateWidgetForList )
+		.HeaderRow
+		(
+			SNew(SHeaderRow)
+
+			+ SHeaderRow::Column("Description")
+			.DefaultLabel(LOCTEXT("Description", "Description"))
+
+			+ SHeaderRow::Column("Value")
+			.DefaultLabel(LOCTEXT("Value", "Value"))
+		);
+	
+	ChildSlot
+	[
+		SNew(SVerticalBox)
+		+ SVerticalBox::Slot()
+		.HAlign(HAlign_Fill)
+		.AutoHeight()
+		.Padding(FMargin(0.f, 5.f))
+		[
+			SNew(SBorder)
+			.BorderImage(FAppStyle::Get().GetBrush("ToolPanel.GroupBorder"))
+			.BorderBackgroundColor(FSlateColor::UseSubduedForeground())
+			.Padding(FMargin(3.f))
+			[
+				SNew(STextBlock)
+				.Text(LOCTEXT("ActionsMapping", "Actions Mapping"))
+			]
+		]
+
+		+ SVerticalBox::Slot()
+		.AutoHeight()
+		[
+			ListView.ToSharedRef()
+		]
+	];
+
+	// Add delegates
+	const TSharedPtr<SRemoteControlPanel> RemoteControlPanel = InActionPanel->GetRemoteControlPanel();
+	if (ensure(RemoteControlPanel))
+	{
+		RemoteControlPanel->OnActionAdded.AddSP(this, &SRCActionPanelList::OnActionAdded);
+		RemoteControlPanel->OnEmptyActions.AddSP(this, &SRCActionPanelList::OnEmptyActions);
+	}
+
+	Reset();
+}
+
+END_SLATE_FUNCTION_BUILD_OPTIMIZATION
+
+void SRCActionPanelList::Reset()
+{
+	ActionItems.Empty();
+
+	if (TSharedPtr<FRCBehaviourModel> BehaviourItem = BehaviourItemWeakPtr.Pin())
+	{
+		if (URCBehaviour* Behaviour = Cast<URCBehaviour>(BehaviourItem->GetBehaviour()))
+		{
+			for (URCAction* Action : Behaviour->ActionContainer->Actions)
+			{
+				if (URCPropertyAction* PropertyAction = Cast<URCPropertyAction>(Action))
+				{
+					ActionItems.Add(MakeShared<FRCPropertyActionModel>(PropertyAction));
+				}
+				else if (URCFunctionAction* FunctionAction = Cast<URCFunctionAction>(Action))
+				{
+					ActionItems.Add(MakeShared<FRCFunctionActionModel>(FunctionAction));
+				}
+			}
+		}
+	}
+
+	ListView->RebuildList();
+}
+
+TSharedRef<ITableRow> SRCActionPanelList::OnGenerateWidgetForList(TSharedPtr<FRCActionModel> InItem, const TSharedRef<STableViewBase>& OwnerTable)
+{
+	return SNew(STableRow<TSharedPtr<FString>>, OwnerTable)
+	[
+		InItem->GetWidget()
+	];
+}
+
+void SRCActionPanelList::OnActionAdded(URCAction* InAction)
+{
+	Reset();
+}
+
+void SRCActionPanelList::OnEmptyActions()
+{
+	Reset();
+}
+
+URemoteControlPreset* SRCActionPanelList::GetPreset()
+{
+	if (ActionPanelWeakPtr.IsValid())
+	{
+		return ActionPanelWeakPtr.Pin()->GetPreset();
+	}
+
+	return nullptr;
+}
+
+int32 SRCActionPanelList::RemoveModel(const TSharedPtr<FRCLogicModeBase> InModel)
+{
+	if (const TSharedPtr<FRCBehaviourModel> BehaviourModel = BehaviourItemWeakPtr.Pin())
+	{
+		if (const URCBehaviour* Behaviour = BehaviourModel->GetBehaviour())
+		{
+			if(const TSharedPtr<FRCActionModel> SelectedAction = StaticCastSharedPtr<FRCActionModel>(InModel))
+			{
+				// Remove Model from Data Container
+				const int32 RemoveCount = Behaviour->ActionContainer->RemoveAction(SelectedAction->GetAction());
+
+				return RemoveCount;
+			}
+		}
+	}
+
+	return 0;
+}
+
+bool SRCActionPanelList::IsListFocused() const
+{
+	return ListView->HasAnyUserFocus().IsSet();
+}
+
+void SRCActionPanelList::DeleteSelectedPanelItem()
+{
+	DeleteItemFromLogicPanel<FRCActionModel>(ActionItems, ListView->GetSelectedItems());
+}
+
+#undef LOCTEXT_NAMESPACE
\ No newline at end of file
diff --git a/Engine/Plugins/VirtualProduction/RemoteControl/Source/RemoteControlUI/Private/UI/Action/SRCActionPanelList.h b/Engine/Plugins/VirtualProduction/RemoteControl/Source/RemoteControlUI/Private/UI/Action/SRCActionPanelList.h
new file mode 100644
--- /dev/null
+++ b/Engine/Plugins/VirtualProduction/RemoteControl/Source/RemoteControlUI/Private/UI/Action/SRCActionPanelList.h
@@ -0,0 +1,78 @@
+﻿// Copyright Epic Games, Inc. All Rights Reserved.
+
+#pragma once
+#include "UI/BaseLogicUI/SRCLogicPanelListBase.h"
+
+
+class FRCBehaviourModel;
+class FRCActionModel;
+class ITableRow;
+class STableViewBase;
+class SRCActionPanel;
+class URCAction;
+class URemoteControlPreset;
+template <typename ItemType> class SListView;
+
+/*
+* ~ SRCActionPanelList ~
+*
+* UI Widget for Actions List
+* Used as part of the RC Logic Actions Panel.
+*/
+class REMOTECONTROLUI_API SRCActionPanelList : public SRCLogicPanelListBase
+{
+public:
+	SLATE_BEGIN_ARGS(SRCActionPanelList)
+		{
+		}
+
+	SLATE_END_ARGS()
+
+	/** Constructs this widget with InArgs */
+	void Construct(const FArguments& InArgs, const TSharedRef<SRCActionPanel> InActionPanel, TSharedPtr<FRCBehaviourModel> InBehaviourItem);
+
+private:
+	/** OnGenerateRow delegate for the Actions List View*/
+	TSharedRef<ITableRow> OnGenerateWidgetForList( TSharedPtr<FRCActionModel> InItem, const TSharedRef<STableViewBase>& OwnerTable );
+
+	/** Responds to the selection of a newly created action. Resets UI state*/
+	void OnActionAdded(URCAction* InAction);
+
+	/** Responds to the removal of all actions. Rests UI state*/
+	void OnEmptyActions();
+
+private:
+	/** The parent Action Panel widget*/
+	TWeakPtr<SRCActionPanel> ActionPanelWeakPtr;
+	
+	/** The Behaviour (UI model) associated with us*/
+	TWeakPtr<FRCBehaviourModel> BehaviourItemWeakPtr;
+
+	/** List of Actions (UI model) active in this widget */
+	TArray<TSharedPtr<FRCActionModel>> ActionItems;
+	
+	/** List View widget for representing our Actions List*/
+	TSharedPtr<SListView<TSharedPtr<FRCActionModel>>> ListView;
+
+private:
+	/** Refreshes the list from the latest state of the data model*/
+	virtual void Reset() override;
+
+	/** Handles broadcasting of a successful remove item operation.*/
+	virtual void BroadcastOnItemRemoved() override {}
+
+	/** Fetches the Remote Control preset associated with the parent panel */
+	virtual URemoteControlPreset* GetPreset() override;
+
+	/** Removes the given Action UI model item from the list of UI models*/
+	virtual int32 RemoveModel(const TSharedPtr<FRCLogicModeBase> InModel) override;
+
+public:
+	/** Whether the Actions List View currently has focus.*/
+	bool IsListFocused() const;
+
+	/** Deletes currently selected items from the list view*/
+	void DeleteSelectedPanelItem();
+
+};
+
diff --git a/Engine/Plugins/VirtualProduction/RemoteControl/Source/RemoteControlUI/Private/UI/BaseLogicUI/RCLogicHelpers.cpp b/Engine/Plugins/VirtualProduction/RemoteControl/Source/RemoteControlUI/Private/UI/BaseLogicUI/RCLogicHelpers.cpp
new file mode 100644
--- /dev/null
+++ b/Engine/Plugins/VirtualProduction/RemoteControl/Source/RemoteControlUI/Private/UI/BaseLogicUI/RCLogicHelpers.cpp
@@ -0,0 +1,77 @@
+﻿// Copyright Epic Games, Inc. All Rights Reserved.
+
+#include "RCLogicHelpers.h"
+
+#include "AssetRegistry/AssetRegistryModule.h"
+#include "Behaviour/RCBehaviourNode.h"
+#include "BlueprintEditorModule.h"
+#include "ContentBrowserModule.h"
+#include "IContentBrowserSingleton.h"
+#include "Kismet2/BlueprintEditorUtils.h"
+#include "Kismet2/KismetEditorUtilities.h"
+
+#define LOCTEXT_NAMESPACE "RCLogicHelpers"
+
+UBlueprint* UE::RCLogicHelpers::CreateBlueprintWithDialog(UClass* InBlueprintParentClass, const UPackage* InSourcePackage, TSubclassOf<UBlueprint> InBlueprintClassType, TSubclassOf<UBlueprintGeneratedClass> InBlueprintGeneratedClassType)
+{
+	UBlueprint* NewBlueprint = nullptr;
+	
+	if (!InBlueprintParentClass)
+	{
+		return NewBlueprint;
+	}
+
+	
+	const FString ClassName = FBlueprintEditorUtils::GetClassNameWithoutSuffix(InBlueprintParentClass);
+
+	FString PathName = InSourcePackage->GetPathName();
+	PathName = FPaths::GetPath(PathName);
+
+	// Now that we've generated some reasonable default locations/names for the package, allow the user to have the final say
+	// before we create the package and initialize the blueprint inside of it.
+	FSaveAssetDialogConfig SaveAssetDialogConfig;
+	SaveAssetDialogConfig.DialogTitleOverride = LOCTEXT("SaveAssetDialogTitle", "Save Asset As");
+	SaveAssetDialogConfig.DefaultPath = PathName;
+	SaveAssetDialogConfig.DefaultAssetName = ClassName + TEXT("_New");
+	SaveAssetDialogConfig.ExistingAssetPolicy = ESaveAssetDialogExistingAssetPolicy::Disallow;
+			
+	const FContentBrowserModule& ContentBrowserModule = FModuleManager::LoadModuleChecked<FContentBrowserModule>("ContentBrowser");
+	const FString SaveObjectPath = ContentBrowserModule.Get().CreateModalSaveAssetDialog(SaveAssetDialogConfig);
+	if (!SaveObjectPath.IsEmpty())
+	{
+		const FString SavePackageName = FPackageName::ObjectPathToPackageName(SaveObjectPath);
+		const FString SavePackagePath = FPaths::GetPath(SavePackageName);
+		const FString SaveAssetName = FPaths::GetBaseFilename(SavePackageName);
+
+		UPackage* Package = CreatePackage(*SavePackageName);
+		if (ensure(Package))
+		{
+			NewBlueprint = FKismetEditorUtilities::CreateBlueprint(InBlueprintParentClass, Package, FName(*SaveAssetName), BPTYPE_Normal, InBlueprintClassType, InBlueprintGeneratedClassType);
+
+			if (NewBlueprint)
+			{
+				// Notify the asset registry
+				FAssetRegistryModule::AssetCreated(NewBlueprint);
+
+				// Mark the package dirty...
+				Package->MarkPackageDirty();
+			}
+		}
+	}
+
+	return NewBlueprint;
+}
+
+bool UE::RCLogicHelpers::OpenBlueprintEditor(UBlueprint* InBlueprint)
+{
+	if (InBlueprint)
+	{
+		FBlueprintEditorModule& BlueprintEditorModule = FModuleManager::LoadModuleChecked<FBlueprintEditorModule>( "Kismet" );
+		BlueprintEditorModule.CreateBlueprintEditor(  EToolkitMode::Standalone, TSharedPtr<IToolkitHost>(), InBlueprint, false );
+		return true;
+	}
+
+	return false;
+}
+
+#undef LOCTEXT_NAMESPACE
diff --git a/Engine/Plugins/VirtualProduction/RemoteControl/Source/RemoteControlUI/Private/UI/BaseLogicUI/RCLogicHelpers.h b/Engine/Plugins/VirtualProduction/RemoteControl/Source/RemoteControlUI/Private/UI/BaseLogicUI/RCLogicHelpers.h
new file mode 100644
--- /dev/null
+++ b/Engine/Plugins/VirtualProduction/RemoteControl/Source/RemoteControlUI/Private/UI/BaseLogicUI/RCLogicHelpers.h
@@ -0,0 +1,20 @@
+﻿// Copyright Epic Games, Inc. All Rights Reserved.
+
+#pragma once
+
+#include "CoreMinimal.h"
+#include "Templates/SubclassOf.h"
+
+class URCBehaviourBlueprint;
+class UBlueprint;
+class UBlueprintGeneratedClass;
+
+/** Helpers for "Blueprint Override" functionality - for creating custom BP Behaviours*/
+namespace UE::RCLogicHelpers
+{
+	/** Launches the Create Blueprint dialog for creating custom Behaviour blueprints */
+	UBlueprint* CreateBlueprintWithDialog(UClass* InBlueprintParentClass, const UPackage* InSourcePackage, TSubclassOf<UBlueprint> InBlueprintClassType, TSubclassOf<UBlueprintGeneratedClass> InBlueprintGeneratedCClassType);
+
+	/** Opens the Blueprint Editor for a given custom Behaviour blueprint*/
+	bool OpenBlueprintEditor(UBlueprint* InBlueprint);
+}
\ No newline at end of file
diff --git a/Engine/Plugins/VirtualProduction/RemoteControl/Source/RemoteControlUI/Private/UI/BaseLogicUI/RCLogicModeBase.cpp b/Engine/Plugins/VirtualProduction/RemoteControl/Source/RemoteControlUI/Private/UI/BaseLogicUI/RCLogicModeBase.cpp
new file mode 100644
--- /dev/null
+++ b/Engine/Plugins/VirtualProduction/RemoteControl/Source/RemoteControlUI/Private/UI/BaseLogicUI/RCLogicModeBase.cpp
@@ -0,0 +1,27 @@
+﻿// Copyright Epic Games, Inc. All Rights Reserved.
+
+#include "RCLogicModeBase.h"
+#include "UI/SRemoteControlPanel.h"
+
+#define LOCTEXT_NAMESPACE "FRCLogicModeBase"
+
+FRCLogicModeBase::FRCLogicModeBase(const TSharedPtr<SRemoteControlPanel> InRemoteControlPanel) : PanelWeakPtr(InRemoteControlPanel)
+{
+}
+
+TSharedRef<SWidget> FRCLogicModeBase::GetWidget() const
+{
+	return SNew(STextBlock).Text(LOCTEXT("NotImplemented", "Not Implemented"));
+}
+
+URemoteControlPreset* FRCLogicModeBase::GetPreset() const
+{
+	if (!PanelWeakPtr.IsValid())
+	{
+		return nullptr;
+	}
+
+	return PanelWeakPtr.Pin()->GetPreset();
+}
+
+#undef LOCTEXT_NAMESPACE
\ No newline at end of file
diff --git a/Engine/Plugins/VirtualProduction/RemoteControl/Source/RemoteControlUI/Private/UI/BaseLogicUI/RCLogicModeBase.h b/Engine/Plugins/VirtualProduction/RemoteControl/Source/RemoteControlUI/Private/UI/BaseLogicUI/RCLogicModeBase.h
new file mode 100644
--- /dev/null
+++ b/Engine/Plugins/VirtualProduction/RemoteControl/Source/RemoteControlUI/Private/UI/BaseLogicUI/RCLogicModeBase.h
@@ -0,0 +1,37 @@
+﻿// Copyright Epic Games, Inc. All Rights Reserved.
+
+#pragma once
+
+#include "CoreMinimal.h"
+#include "UObject/WeakObjectPtrTemplates.h"
+
+class SRemoteControlPanel;
+class SWidget;
+class URemoteControlPreset;
+
+/*
+* ~ FRCLogicModeBase ~
+*
+* Base UI model for representing a Logic item (Controllers / Behaviours / Actions)
+* Provides a row widget that can be used in list views.
+*/
+class FRCLogicModeBase : public TSharedFromThis<FRCLogicModeBase>
+{
+public:
+	FRCLogicModeBase(){}
+	FRCLogicModeBase(const TSharedPtr<SRemoteControlPanel> InRemoteControlPanel);
+
+	virtual ~FRCLogicModeBase() = default;
+
+	/** Represents the widget to be rendered for associated UI Model.
+	* Implemented by child classes to supply the row widget for various panel lists
+	*/
+	virtual TSharedRef<SWidget> GetWidget() const;
+
+	/** Returns the Remote Control Preset object associated with us*/
+	URemoteControlPreset* GetPreset() const;
+
+protected:
+	/** The parent Remote Control Panel widget*/
+	TWeakPtr<SRemoteControlPanel> PanelWeakPtr;
+};
diff --git a/Engine/Plugins/VirtualProduction/RemoteControl/Source/RemoteControlUI/Private/UI/BaseLogicUI/SRCLogicPanelBase.cpp b/Engine/Plugins/VirtualProduction/RemoteControl/Source/RemoteControlUI/Private/UI/BaseLogicUI/SRCLogicPanelBase.cpp
new file mode 100644
--- /dev/null
+++ b/Engine/Plugins/VirtualProduction/RemoteControl/Source/RemoteControlUI/Private/UI/BaseLogicUI/SRCLogicPanelBase.cpp
@@ -0,0 +1,30 @@
+﻿// Copyright Epic Games, Inc. All Rights Reserved.
+
+#include "SRCLogicPanelBase.h"
+
+#include "SlateOptMacros.h"
+#include "UI/SRemoteControlPanel.h"
+
+BEGIN_SLATE_FUNCTION_BUILD_OPTIMIZATION
+
+void SRCLogicPanelBase::Construct(const FArguments& InArgs, const TSharedRef<SRemoteControlPanel>& InPanel)
+{
+	PanelWeakPtr = InPanel;
+}
+
+END_SLATE_FUNCTION_BUILD_OPTIMIZATION
+
+URemoteControlPreset* SRCLogicPanelBase::GetPreset() const
+{
+	if (const TSharedPtr<SRemoteControlPanel> RemoteControlPanel = GetRemoteControlPanel())
+	{
+		return RemoteControlPanel->GetPreset();
+	}
+
+	return nullptr;
+}
+
+TSharedPtr<SRemoteControlPanel> SRCLogicPanelBase::GetRemoteControlPanel() const
+{
+	return PanelWeakPtr.Pin();
+}
\ No newline at end of file
diff --git a/Engine/Plugins/VirtualProduction/RemoteControl/Source/RemoteControlUI/Private/UI/BaseLogicUI/SRCLogicPanelBase.h b/Engine/Plugins/VirtualProduction/RemoteControl/Source/RemoteControlUI/Private/UI/BaseLogicUI/SRCLogicPanelBase.h
new file mode 100644
--- /dev/null
+++ b/Engine/Plugins/VirtualProduction/RemoteControl/Source/RemoteControlUI/Private/UI/BaseLogicUI/SRCLogicPanelBase.h
@@ -0,0 +1,43 @@
+﻿// Copyright Epic Games, Inc. All Rights Reserved.
+
+#pragma once
+
+#include "CoreMinimal.h"
+#include "Widgets/SCompoundWidget.h"
+
+class SRemoteControlPanel;
+class URemoteControlPreset;
+
+/*
+* 
+* ~ SRCLogicPanelBase ~
+*
+* Base UI Widget for Logic Panels (Controllers / Behaviours / Actions).
+* Provides access to parent Remote Control Preset object and common Delete Panel Item functionality
+*/
+class REMOTECONTROLUI_API SRCLogicPanelBase : public SCompoundWidget
+{
+public:
+	SLATE_BEGIN_ARGS(SRCLogicPanelBase)
+		{
+		}
+
+	SLATE_END_ARGS()
+
+	/** Constructs this widget with InArgs */
+	void Construct(const FArguments& InArgs, const TSharedRef<SRemoteControlPanel>& InPanel);
+
+	/** Returns the Remote Control Preset object associated with us*/
+	virtual URemoteControlPreset* GetPreset() const;
+
+	/** Returns the parent Remote Control Panel widget,
+	* This is the main container holding all the Remote Control panels including Logic and Exposed Properties*/
+	virtual TSharedPtr<SRemoteControlPanel> GetRemoteControlPanel() const;
+
+	/** "Delete Item" UI command implementation for panels*/
+	virtual void DeleteSelectedPanelItem() = 0;
+
+protected:
+	/** The parent Remote Control Panel widget*/
+	TWeakPtr<SRemoteControlPanel> PanelWeakPtr;
+};
diff --git a/Engine/Plugins/VirtualProduction/RemoteControl/Source/RemoteControlUI/Private/UI/BaseLogicUI/SRCLogicPanelHeaderBase.cpp b/Engine/Plugins/VirtualProduction/RemoteControl/Source/RemoteControlUI/Private/UI/BaseLogicUI/SRCLogicPanelHeaderBase.cpp
new file mode 100644
--- /dev/null
+++ b/Engine/Plugins/VirtualProduction/RemoteControl/Source/RemoteControlUI/Private/UI/BaseLogicUI/SRCLogicPanelHeaderBase.cpp
@@ -0,0 +1,13 @@
+﻿// Copyright Epic Games, Inc. All Rights Reserved.
+
+#include "SRCLogicPanelHeaderBase.h"
+
+#include "SlateOptMacros.h"
+
+BEGIN_SLATE_FUNCTION_BUILD_OPTIMIZATION
+
+void SRCLogicPanelHeaderBase::Construct(const FArguments& InArgs)
+{
+}
+
+END_SLATE_FUNCTION_BUILD_OPTIMIZATION
diff --git a/Engine/Plugins/VirtualProduction/RemoteControl/Source/RemoteControlUI/Private/UI/BaseLogicUI/SRCLogicPanelHeaderBase.h b/Engine/Plugins/VirtualProduction/RemoteControl/Source/RemoteControlUI/Private/UI/BaseLogicUI/SRCLogicPanelHeaderBase.h
new file mode 100644
--- /dev/null
+++ b/Engine/Plugins/VirtualProduction/RemoteControl/Source/RemoteControlUI/Private/UI/BaseLogicUI/SRCLogicPanelHeaderBase.h
@@ -0,0 +1,24 @@
+﻿// Copyright Epic Games, Inc. All Rights Reserved.
+
+#pragma once
+
+#include "CoreMinimal.h"
+#include "Widgets/SCompoundWidget.h"
+
+/**
+ *  Base UI widget for Logic header panels.
+ * Typically contain metdata, Add/Remove buttons, etc; 
+ * Currently customized across Controllers / Behaviours / Actions
+ */
+class REMOTECONTROLUI_API SRCLogicPanelHeaderBase : public SCompoundWidget
+{
+public:
+	SLATE_BEGIN_ARGS(SRCLogicPanelHeaderBase)
+		{
+		}
+
+	SLATE_END_ARGS()
+
+	/** Constructs this widget with InArgs */
+	void Construct(const FArguments& InArgs);
+};
diff --git a/Engine/Plugins/VirtualProduction/RemoteControl/Source/RemoteControlUI/Private/UI/BaseLogicUI/SRCLogicPanelListBase.cpp b/Engine/Plugins/VirtualProduction/RemoteControl/Source/RemoteControlUI/Private/UI/BaseLogicUI/SRCLogicPanelListBase.cpp
new file mode 100644
--- /dev/null
+++ b/Engine/Plugins/VirtualProduction/RemoteControl/Source/RemoteControlUI/Private/UI/BaseLogicUI/SRCLogicPanelListBase.cpp
@@ -0,0 +1,56 @@
+﻿// Copyright Epic Games, Inc. All Rights Reserved.
+
+#include "SRCLogicPanelListBase.h"
+
+#include "RCLogicModeBase.h"
+#include "SlateOptMacros.h"
+#include "UI/Action/RCActionModel.h"
+#include "UI/Behaviour/RCBehaviourModel.h"
+#include "UI/Controller/RCControllerModel.h"
+
+BEGIN_SLATE_FUNCTION_BUILD_OPTIMIZATION
+
+void SRCLogicPanelListBase::Construct(const FArguments& InArgs)
+{
+}
+
+END_SLATE_FUNCTION_BUILD_OPTIMIZATION
+
+template<class T>
+void SRCLogicPanelListBase::DeleteItemFromLogicPanel(TArray<TSharedPtr<T>>& ItemsSource, const TArray<TSharedPtr<T>>& SelectedItems)
+{
+	bool bIsDeleted = false;
+	for (const TSharedPtr<T> SelectedItem : SelectedItems)
+	{
+		if (SelectedItem.IsValid())
+		{
+			// Remove Model from Data Container
+			const int32 RemoveCount = RemoveModel(SelectedItem);
+			if(ensure(RemoveCount > 0))
+			{
+				// Remove View Model from UI List
+				const int32 RemoveModelItemIndex = ItemsSource.IndexOfByPredicate([SelectedItem](TSharedPtr<T> InModel)
+					{
+						return SelectedItem == InModel;
+					});
+
+				if (RemoveModelItemIndex > INDEX_NONE)
+				{
+					ItemsSource.RemoveAt(RemoveModelItemIndex);
+
+					bIsDeleted = true;
+				}
+			}
+		}
+	}
+
+	if (bIsDeleted)
+	{
+		BroadcastOnItemRemoved();
+		Reset();
+	}
+}
+
+template void SRCLogicPanelListBase::DeleteItemFromLogicPanel(TArray<TSharedPtr<FRCControllerModel>>& ItemsSource, const TArray<TSharedPtr<FRCControllerModel>>& SelectedItems);
+template void SRCLogicPanelListBase::DeleteItemFromLogicPanel(TArray<TSharedPtr<FRCBehaviourModel>>& ItemsSource, const TArray<TSharedPtr<FRCBehaviourModel>>& SelectedItems);
+template void SRCLogicPanelListBase::DeleteItemFromLogicPanel(TArray<TSharedPtr<FRCActionModel>>& ItemsSource, const TArray<TSharedPtr<FRCActionModel>>& SelectedItems);
\ No newline at end of file
diff --git a/Engine/Plugins/VirtualProduction/RemoteControl/Source/RemoteControlUI/Private/UI/BaseLogicUI/SRCLogicPanelListBase.h b/Engine/Plugins/VirtualProduction/RemoteControl/Source/RemoteControlUI/Private/UI/BaseLogicUI/SRCLogicPanelListBase.h
new file mode 100644
--- /dev/null
+++ b/Engine/Plugins/VirtualProduction/RemoteControl/Source/RemoteControlUI/Private/UI/BaseLogicUI/SRCLogicPanelListBase.h
@@ -0,0 +1,48 @@
+﻿// Copyright Epic Games, Inc. All Rights Reserved.
+
+#pragma once
+
+#include "CoreMinimal.h"
+#include "Widgets/SCompoundWidget.h"
+
+class FRCLogicModeBase;
+class URemoteControlPreset;
+
+/*
+* ~ SRCLogicPanelListBase ~
+*
+* Base UI Widget for Lists in Logic Panels.
+* Can represent Controllers / Behaviours / Actions, etc.
+*/
+class REMOTECONTROLUI_API SRCLogicPanelListBase : public SCompoundWidget
+{
+public:
+	SLATE_BEGIN_ARGS(SRCLogicPanelListBase)
+		{
+		}
+
+	SLATE_END_ARGS()
+
+	/** Constructs this widget with InArgs */
+	void Construct(const FArguments& InArgs);
+
+private:
+	/** Refreshes the list from the latest state of the model*/
+	virtual void Reset()  = 0;
+
+	/** Handles broadcasting of a successful remove item operation.
+	* This is handled uniquely by each type of child list widget*/
+	virtual void BroadcastOnItemRemoved() = 0;
+
+	/** Fetches the Remote Control preset associated with the parent panel */
+	virtual URemoteControlPreset* GetPreset() = 0;
+
+	/** Removes the given UI model item from the list of UI models for this panel list*/
+	virtual int32 RemoveModel(const TSharedPtr<FRCLogicModeBase> InModel) = 0;
+
+protected:
+	/** Helper function for handling common Delete Item functionality across all child panels (Actions/Behaviours/Controllers)
+	* Currently invoked from each Panel List child class with appropriate template class*/
+	template<class T>
+	void DeleteItemFromLogicPanel(TArray<TSharedPtr<T>>& ItemsSource, const TArray<TSharedPtr<T>>& SelectedItems);
+};
\ No newline at end of file
diff --git a/Engine/Plugins/VirtualProduction/RemoteControl/Source/RemoteControlUI/Private/UI/Behaviour/Builtin/RCBehaviourIsEqualModel.cpp b/Engine/Plugins/VirtualProduction/RemoteControl/Source/RemoteControlUI/Private/UI/Behaviour/Builtin/RCBehaviourIsEqualModel.cpp
new file mode 100644
--- /dev/null
+++ b/Engine/Plugins/VirtualProduction/RemoteControl/Source/RemoteControlUI/Private/UI/Behaviour/Builtin/RCBehaviourIsEqualModel.cpp
@@ -0,0 +1,86 @@
+// Copyright Epic Games, Inc. All Rights Reserved.
+
+#include "RCBehaviourIsEqualModel.h"
+
+#include "Modules/ModuleManager.h"
+#include "IDetailTreeNode.h"
+#include "IPropertyRowGenerator.h"
+#include "PropertyEditorModule.h"
+#include "RCVirtualProperty.h"
+#include "SRCBehaviourIsEqual.h"
+
+#define LOCTEXT_NAMESPACE "FRCIsEqualBehaviourModel"
+
+FRCIsEqualBehaviourModel::FRCIsEqualBehaviourModel(URCIsEqualBehaviour* IsEqualBehaviour)
+	: FRCBehaviourModel(IsEqualBehaviour)
+	, EqualBehaviourWeakPtr(IsEqualBehaviour)
+{
+	FPropertyRowGeneratorArgs Args;
+	Args.bShouldShowHiddenProperties = true;
+	PropertyRowGenerator = FModuleManager::GetModuleChecked<FPropertyEditorModule>("PropertyEditor").CreatePropertyRowGenerator(Args);
+
+	if (IsEqualBehaviour)
+	{
+		PropertyRowGenerator->SetStructure(IsEqualBehaviour->PropertySelfContainer->CreateStructOnScope());
+
+		for (const TSharedRef<IDetailTreeNode>& CategoryNode : PropertyRowGenerator->GetRootTreeNodes())
+		{
+			TArray<TSharedRef<IDetailTreeNode>> Children;
+			CategoryNode->GetChildren(Children);
+			for (const TSharedRef<IDetailTreeNode>& Child : Children)
+			{
+				DetailTreeNodeWeakPtr = Child;
+				break;
+			}
+		}
+	}
+}
+
+TSharedRef<SWidget> FRCIsEqualBehaviourModel::GetBehaviourDetailsWidget() const
+{
+	return SNew(SRCBehaviourIsEqual, SharedThis(this));
+}
+
+TSharedRef<SWidget> FRCIsEqualBehaviourModel::GetPropertyWidget() const
+{
+	if (const TSharedPtr<IDetailTreeNode> DetailTreeNode = DetailTreeNodeWeakPtr.Pin())
+	{
+		const FNodeWidgets NodeWidgets = DetailTreeNode->CreateNodeWidgets();
+
+		const TSharedRef<SHorizontalBox> FieldWidget = SNew(SHorizontalBox);
+
+		if (NodeWidgets.ValueWidget)
+		{
+			FieldWidget->AddSlot()
+				.Padding(FMargin(3.0f, 2.0f))
+				.HAlign(HAlign_Right)
+				.AutoWidth()
+				[
+					NodeWidgets.ValueWidget.ToSharedRef()
+				];
+		}
+		else if (NodeWidgets.WholeRowWidget)
+		{
+			FieldWidget->AddSlot()
+				.Padding(FMargin(3.0f, 2.0f))
+				.AutoWidth()
+				[
+					NodeWidgets.WholeRowWidget.ToSharedRef()
+				];
+		}
+
+		return 	SNew(SHorizontalBox)
+			.Clipping(EWidgetClipping::OnDemand)
+			// Property Value Widget
+			+ SHorizontalBox::Slot()
+			.VAlign(VAlign_Center)
+			.AutoWidth()
+			[
+				FieldWidget
+			];
+	}
+
+	return SNullWidget::NullWidget;
+}
+
+#undef LOCTEXT_NAMESPACE
\ No newline at end of file
diff --git a/Engine/Plugins/VirtualProduction/RemoteControl/Source/RemoteControlUI/Private/UI/Behaviour/Builtin/RCBehaviourIsEqualModel.h b/Engine/Plugins/VirtualProduction/RemoteControl/Source/RemoteControlUI/Private/UI/Behaviour/Builtin/RCBehaviourIsEqualModel.h
new file mode 100644
--- /dev/null
+++ b/Engine/Plugins/VirtualProduction/RemoteControl/Source/RemoteControlUI/Private/UI/Behaviour/Builtin/RCBehaviourIsEqualModel.h
@@ -0,0 +1,40 @@
+// Copyright Epic Games, Inc. All Rights Reserved.
+
+#pragma once
+#include "UI/Behaviour/RCBehaviourModel.h"
+
+class IDetailTreeNode;
+class IPropertyRowGenerator;
+class SWidget;
+
+/*
+* ~ FRCIsEqualBehaviourModel ~
+*
+* Child Behaviour class representing the "Is Equal" Behaviour's UI model
+* 
+* Generates a Property Widget where users can enter the comparison value.
+* The value of this is compared with the associated Controller for evaluating Is Equal behaviour.
+*/
+class FRCIsEqualBehaviourModel : public FRCBehaviourModel
+{
+public:
+	FRCIsEqualBehaviourModel(URCIsEqualBehaviour* IsEqualBehaviour);
+
+	/** Builds a Behaviour specific widget as required for Is Equal behaviour*/
+	virtual TSharedRef<SWidget> GetBehaviourDetailsWidget() const override;
+
+	/** Builds a generic Value Widget of the Controller's type
+	* Use to store user input for performing the Is Equal Behaviour comparison*/
+	TSharedRef<SWidget> GetPropertyWidget() const;
+
+private:
+
+	/** The Is Equal Behaviour (Data model) associated with us*/
+	TWeakObjectPtr<URCIsEqualBehaviour> EqualBehaviourWeakPtr;
+
+	/** The row generator used to build a generic Value Widget*/
+	TSharedPtr<IPropertyRowGenerator> PropertyRowGenerator;
+
+	/** Used to create a generic Value Widget for comparison based on the active Controller's type*/
+	TWeakPtr<IDetailTreeNode> DetailTreeNodeWeakPtr;
+};
\ No newline at end of file
diff --git a/Engine/Plugins/VirtualProduction/RemoteControl/Source/RemoteControlUI/Private/UI/Behaviour/Builtin/SRCBehaviourIsEqual.cpp b/Engine/Plugins/VirtualProduction/RemoteControl/Source/RemoteControlUI/Private/UI/Behaviour/Builtin/SRCBehaviourIsEqual.cpp
new file mode 100644
--- /dev/null
+++ b/Engine/Plugins/VirtualProduction/RemoteControl/Source/RemoteControlUI/Private/UI/Behaviour/Builtin/SRCBehaviourIsEqual.cpp
@@ -0,0 +1,85 @@
+// Copyright Epic Games, Inc. All Rights Reserved.
+
+#include "SRCBehaviourIsEqual.h"
+
+#include "SlateOptMacros.h"
+#include "SPositiveActionButton.h"
+#include "Styling/CoreStyle.h"
+#include "UI/Behaviour/Builtin/RCBehaviourIsEqualModel.h"
+#include "UI/RemoteControlPanelStyle.h"
+#include "Widgets/Images/SImage.h"
+#include "Widgets/Input/SButton.h"
+#include "Widgets/Layout/SSeparator.h"
+
+#define LOCTEXT_NAMESPACE "SRCBehaviourIsEqual"
+
+class SPositiveActionButton;
+
+BEGIN_SLATE_FUNCTION_BUILD_OPTIMIZATION
+
+void SRCBehaviourIsEqual::Construct(const FArguments& InArgs, TSharedRef<const FRCIsEqualBehaviourModel> InBehaviourItem)
+{
+	IsEqualBehaviourItemWeakPtr = InBehaviourItem;
+
+	ChildSlot
+	[
+		SNew(SHorizontalBox)
+		// Header
+		+ SHorizontalBox::Slot()
+		.HAlign(HAlign_Fill)
+		.VAlign(VAlign_Center)
+		.AutoWidth()
+		[
+			SNew(STextBlock)
+			.Text(LOCTEXT("BehaviorLogicSummary", "If Value Is"))
+			.Font(FRemoteControlPanelStyle::Get()->GetFontStyle("RemoteControlPanel.Actions.ValuePanelHeader"))
+		]
+
+		+ SHorizontalBox::Slot()
+		.HAlign(HAlign_Right)
+		.VAlign(VAlign_Center)
+		.Padding(FMargin(5.f))
+		[
+			SNew(SHorizontalBox)
+			// Value Widget
+			+ SHorizontalBox::Slot()
+			.VAlign(VAlign_Center)
+			.AutoWidth()
+			[
+				InBehaviourItem->GetPropertyWidget()
+			]
+
+			// Bind button  (currently purely cosmetic)
+			+ SHorizontalBox::Slot()
+			.VAlign(VAlign_Center)
+			[
+				SNew(SButton)
+				.ButtonStyle(FAppStyle::Get(), "FlatButton.DarkGrey")
+				.ContentPadding(FMargin(0.f, 3.f))
+				[
+					SNew(SHorizontalBox)
+					+ SHorizontalBox::Slot()
+					.HAlign(HAlign_Center)
+					.VAlign(VAlign_Center)
+					.Padding(FMargin(3.f))
+					[
+						SNew(SImage)
+						.Image(FAppStyle::Get().GetBrush("EditableComboBox.Add"))
+					]
+
+					+ SHorizontalBox::Slot()
+					.AutoWidth()
+					[
+						SNew(STextBlock)
+						.Text(LOCTEXT("BindButton", "Bind"))
+						.Font(FRemoteControlPanelStyle::Get()->GetFontStyle("RemoteControlPanel.Behaviours.BehaviourDescription"))
+					]
+				]
+			]
+		]
+	];
+}
+
+END_SLATE_FUNCTION_BUILD_OPTIMIZATION
+
+#undef LOCTEXT_NAMESPACE
\ No newline at end of file
diff --git a/Engine/Plugins/VirtualProduction/RemoteControl/Source/RemoteControlUI/Private/UI/Behaviour/Builtin/SRCBehaviourIsEqual.h b/Engine/Plugins/VirtualProduction/RemoteControl/Source/RemoteControlUI/Private/UI/Behaviour/Builtin/SRCBehaviourIsEqual.h
new file mode 100644
--- /dev/null
+++ b/Engine/Plugins/VirtualProduction/RemoteControl/Source/RemoteControlUI/Private/UI/Behaviour/Builtin/SRCBehaviourIsEqual.h
@@ -0,0 +1,32 @@
+// Copyright Epic Games, Inc. All Rights Reserved.
+
+#pragma once
+
+#include "CoreMinimal.h"
+#include "Widgets/SCompoundWidget.h"
+#include "RemoteControlField.h"
+
+class FRCIsEqualBehaviourModel;
+
+/*
+* ~ SRCBehaviourIsEqual ~
+*
+* Behaviour specific details panel for Is Equal Behaviour
+* Contains a Value panel containing an input field of the active Controller's type and related functionality
+*/
+class REMOTECONTROLUI_API SRCBehaviourIsEqual : public SCompoundWidget
+{
+public:
+	SLATE_BEGIN_ARGS(SRCBehaviourIsEqual)
+		{
+		}
+
+	SLATE_END_ARGS()
+
+	/** Constructs this widget with InArgs */
+	void Construct(const FArguments& InArgs, TSharedRef<const FRCIsEqualBehaviourModel> InBehaviourItem);
+
+private:
+	/** The Behaviour (UI model) associated with us*/
+	TWeakPtr<const FRCIsEqualBehaviourModel> IsEqualBehaviourItemWeakPtr;
+};
diff --git a/Engine/Plugins/VirtualProduction/RemoteControl/Source/RemoteControlUI/Private/UI/Behaviour/RCBehaviourModel.cpp b/Engine/Plugins/VirtualProduction/RemoteControl/Source/RemoteControlUI/Private/UI/Behaviour/RCBehaviourModel.cpp
new file mode 100644
--- /dev/null
+++ b/Engine/Plugins/VirtualProduction/RemoteControl/Source/RemoteControlUI/Private/UI/Behaviour/RCBehaviourModel.cpp
@@ -0,0 +1,66 @@
+﻿// Copyright Epic Games, Inc. All Rights Reserved.
+
+#include "RCBehaviourModel.h"
+
+#include "Behaviour/RCBehaviourNode.h"
+#include "Engine/Blueprint.h"
+#include "UI/BaseLogicUI/RCLogicHelpers.h"
+#include "UI/RemoteControlPanelStyle.h"
+#include "UI/SRemoteControlPanel.h"
+#include "Widgets/SNullWidget.h"
+
+#define LOCTEXT_NAMESPACE "FRCBehaviourModel"
+
+FRCBehaviourModel::FRCBehaviourModel(URCBehaviour* InBehaviour)
+	: BehaviourWeakPtr(InBehaviour)
+{
+}
+
+TSharedRef<SWidget> FRCBehaviourModel::GetWidget() const
+{
+	if (!ensure(BehaviourWeakPtr.IsValid()))
+	{
+		return SNullWidget::NullWidget;
+	}
+
+	const FText BehaviorDisplayName = BehaviourWeakPtr->GetDisplayName().ToUpper();
+
+	return SNew(SHorizontalBox)
+		.Clipping(EWidgetClipping::OnDemand)
+		// Behaviour name
+		+ SHorizontalBox::Slot()
+		.VAlign(VAlign_Center)
+		.AutoWidth()
+		.Padding(FMargin(8.f))
+		[
+			SNew(STextBlock).Text(BehaviorDisplayName)
+			.Font(FRemoteControlPanelStyle::Get()->GetFontStyle("RemoteControlPanel.Behaviours.Title"))
+		];
+}
+
+TSharedRef<SWidget> FRCBehaviourModel::GetBehaviourDetailsWidget() const
+{
+	return SNullWidget::NullWidget;
+}
+
+void FRCBehaviourModel::OnOverrideBlueprint() const
+{
+	if (URCBehaviour* Behaviour = BehaviourWeakPtr.Get())
+	{
+		UBlueprint* Blueprint = Behaviour->GetBlueprint();
+		if (!Blueprint)
+		{
+			Blueprint = UE::RCLogicHelpers::CreateBlueprintWithDialog(Behaviour->BehaviourNodeClass, Behaviour->GetPackage(), UBlueprint::StaticClass(), UBlueprintGeneratedClass::StaticClass());
+			Behaviour->SetOverrideBehaviourBlueprintClass(Blueprint);
+		}
+
+		UE::RCLogicHelpers::OpenBlueprintEditor(Blueprint);
+	}
+}
+
+URCBehaviour* FRCBehaviourModel::GetBehaviour() const
+{
+	return BehaviourWeakPtr.Get();
+}
+
+#undef LOCTEXT_NAMESPACE
\ No newline at end of file
diff --git a/Engine/Plugins/VirtualProduction/RemoteControl/Source/RemoteControlUI/Private/UI/Behaviour/RCBehaviourModel.h b/Engine/Plugins/VirtualProduction/RemoteControl/Source/RemoteControlUI/Private/UI/Behaviour/RCBehaviourModel.h
new file mode 100644
--- /dev/null
+++ b/Engine/Plugins/VirtualProduction/RemoteControl/Source/RemoteControlUI/Private/UI/Behaviour/RCBehaviourModel.h
@@ -0,0 +1,39 @@
+﻿// Copyright Epic Games, Inc. All Rights Reserved.
+
+#pragma once
+#include "Behaviour/RCIsEqualBehaviour.h"
+#include "UI/BaseLogicUI/RCLogicModeBase.h"
+
+class IPropertyRowGenerator;
+class SWidget;
+class URCBehaviour;
+
+/*
+* ~ FRCBehaviourModel ~
+*
+* UI model for representing a Behaviour
+* Generates a row widget with Behaviour related metadata
+*/
+class FRCBehaviourModel : public FRCLogicModeBase
+{
+public:
+	FRCBehaviourModel(URCBehaviour* InBehaviour);
+
+	/** The widget to be rendered for this Property in the Behaviour panel
+	* Currently displays Behaviour Name metadata in the Behaviours Panel List
+	*/
+	virtual TSharedRef<SWidget> GetWidget() const override;
+
+	/** Builds a Behaviour specific widget that child Behaviour classes can implement as required*/
+	virtual TSharedRef<SWidget> GetBehaviourDetailsWidget() const;
+
+	/** Returns the Behaviour (Data model) associated with us*/
+	URCBehaviour* GetBehaviour() const;
+
+	/** Handling for user action to override this behaviour via new Blueprint class */
+	void OnOverrideBlueprint() const;	
+
+private:
+	/** The Behaviour (Data model) associated with us*/
+	TWeakObjectPtr<URCBehaviour> BehaviourWeakPtr;
+};
\ No newline at end of file
diff --git a/Engine/Plugins/VirtualProduction/RemoteControl/Source/RemoteControlUI/Private/UI/Behaviour/SRCBehaviourPanel.cpp b/Engine/Plugins/VirtualProduction/RemoteControl/Source/RemoteControlUI/Private/UI/Behaviour/SRCBehaviourPanel.cpp
new file mode 100644
--- /dev/null
+++ b/Engine/Plugins/VirtualProduction/RemoteControl/Source/RemoteControlUI/Private/UI/Behaviour/SRCBehaviourPanel.cpp
@@ -0,0 +1,76 @@
+﻿// Copyright Epic Games, Inc. All Rights Reserved.
+
+#include "SRCBehaviourPanel.h"
+#include "SRCBehaviourPanelHeader.h"
+#include "SRCBehaviourPanelList.h"
+
+#include "SlateOptMacros.h"
+#include "UI/SRemoteControlPanel.h"
+
+#define LOCTEXT_NAMESPACE "SRCBehaviourPanel"
+
+BEGIN_SLATE_FUNCTION_BUILD_OPTIMIZATION
+
+void SRCBehaviourPanel::Construct(const FArguments& InArgs, const TSharedRef<SRemoteControlPanel>& InPanel)
+{
+	SRCLogicPanelBase::Construct(SRCLogicPanelBase::FArguments(), InPanel);
+
+	WrappedBoxWidget = SNew(SBox);
+	UpdateWrappedWidget();
+	
+	ChildSlot
+	[
+		WrappedBoxWidget.ToSharedRef()
+	];
+
+	// Register delegates
+	InPanel->OnControllerSelectionChanged.AddSP(this, &SRCBehaviourPanel::OnControllerSelectionChanged);
+}
+
+END_SLATE_FUNCTION_BUILD_OPTIMIZATION
+
+void SRCBehaviourPanel::OnControllerSelectionChanged(TSharedPtr<FRCControllerModel> InControllerItem)
+{
+	UpdateWrappedWidget(InControllerItem);
+	SelectedControllerItemWeakPtr = InControllerItem;
+}
+
+void SRCBehaviourPanel::UpdateWrappedWidget(TSharedPtr<FRCControllerModel> InControllerItem)
+{
+	if (InControllerItem.IsValid())
+	{
+		WrappedBoxWidget->SetContent(
+			SNew(SVerticalBox)
+			// Header
+			+ SVerticalBox::Slot()
+			.Padding(5.0f, 5.0f)
+			.AutoHeight()
+			[
+				SNew(SRCBehaviourPanelHeader, SharedThis(this), InControllerItem)
+			]
+			// List
+			+ SVerticalBox::Slot()
+			[
+				SAssignNew(BehaviourPanelList, SRCBehaviourPanelList, SharedThis(this), InControllerItem)
+			]
+		);
+	}
+	else
+	{
+		WrappedBoxWidget->SetContent(
+			SNew(STextBlock).Text(LOCTEXT("Unselected", "No Controller Selected"))
+		);
+	}
+}
+
+bool SRCBehaviourPanel::IsListFocused() const
+{
+	return BehaviourPanelList.IsValid() && BehaviourPanelList->IsListFocused();
+}
+
+void SRCBehaviourPanel::DeleteSelectedPanelItem()
+{
+	BehaviourPanelList->DeleteSelectedPanelItem();
+}
+
+#undef LOCTEXT_NAMESPACE
\ No newline at end of file
diff --git a/Engine/Plugins/VirtualProduction/RemoteControl/Source/RemoteControlUI/Private/UI/Behaviour/SRCBehaviourPanel.h b/Engine/Plugins/VirtualProduction/RemoteControl/Source/RemoteControlUI/Private/UI/Behaviour/SRCBehaviourPanel.h
new file mode 100644
--- /dev/null
+++ b/Engine/Plugins/VirtualProduction/RemoteControl/Source/RemoteControlUI/Private/UI/Behaviour/SRCBehaviourPanel.h
@@ -0,0 +1,55 @@
+﻿// Copyright Epic Games, Inc. All Rights Reserved.
+
+#pragma once
+
+#include "UI/BaseLogicUI/SRCLogicPanelBase.h"
+
+class FRCControllerModel;
+class SBox;
+
+/*
+* ~ SRCBehaviourPanel ~
+*
+* UI Widget for Behaviour Panel.
+* Contains a header (Add/Remove/Empty) and List of Behaviours
+*/
+class REMOTECONTROLUI_API SRCBehaviourPanel : public SRCLogicPanelBase
+{
+public:
+	SLATE_BEGIN_ARGS(SRCBehaviourPanel)
+		{
+		}
+
+	SLATE_END_ARGS()
+	
+	/** Constructs this widget with InArgs */
+	void Construct(const FArguments& InArgs, const TSharedRef<SRemoteControlPanel>& InPanel);
+
+private:
+	/**
+	* Controller list selection change listener.
+	* Updates the list of behaviours from newly selected Controller
+	*/
+	void OnControllerSelectionChanged(TSharedPtr<FRCControllerModel> InControllerItem);
+
+	/* Rebuilds the Behaviour Panel for a newly selected Controller*/
+	void UpdateWrappedWidget(TSharedPtr<FRCControllerModel> InControllerItem = nullptr);
+
+private:
+	/** The parent Controller that this Behaviour panel is associated with */
+	TWeakPtr<FRCControllerModel> SelectedControllerItemWeakPtr = nullptr;
+	
+	/** Parent Box for the entire widget */
+	TSharedPtr<SBox> WrappedBoxWidget;
+
+	/** Widget representing List of Behaviours */
+	TSharedPtr<class SRCBehaviourPanelList> BehaviourPanelList;
+
+public:
+	/** Whether the Behaviour list widget currently has focus. Used for Delete Item UI command */
+	bool IsListFocused() const;
+
+	/** Delete Item UI command implementation for this panel */
+	virtual void DeleteSelectedPanelItem() override;
+
+};
diff --git a/Engine/Plugins/VirtualProduction/RemoteControl/Source/RemoteControlUI/Private/UI/Behaviour/SRCBehaviourPanelHeader.cpp b/Engine/Plugins/VirtualProduction/RemoteControl/Source/RemoteControlUI/Private/UI/Behaviour/SRCBehaviourPanelHeader.cpp
new file mode 100644
--- /dev/null
+++ b/Engine/Plugins/VirtualProduction/RemoteControl/Source/RemoteControlUI/Private/UI/Behaviour/SRCBehaviourPanelHeader.cpp
@@ -0,0 +1,149 @@
+﻿// Copyright Epic Games, Inc. All Rights Reserved.
+
+#include "SRCBehaviourPanelHeader.h"
+
+#include "Behaviour/RCBehaviourBlueprintNode.h"
+#include "Behaviour/RCBehaviourNode.h"
+#include "Controller/RCController.h"
+#include "Framework/MultiBox/MultiBoxBuilder.h"
+#include "Kismet2/KismetEditorUtilities.h"
+#include "SlateOptMacros.h"
+#include "SPositiveActionButton.h"
+#include "SRCBehaviourPanel.h"
+#include "UI/SRemoteControlPanel.h"
+#include "UI/Controller/RCControllerModel.h"
+#include "Widgets/Input/SButton.h"
+
+#define LOCTEXT_NAMESPACE "SRCControllerPanelHeader"
+
+BEGIN_SLATE_FUNCTION_BUILD_OPTIMIZATION
+
+void SRCBehaviourPanelHeader::Construct(const FArguments& InArgs, TSharedRef<SRCBehaviourPanel> InBehaviourPanel, TSharedPtr<FRCControllerModel> InControllerItem)
+{
+	SRCLogicPanelHeaderBase::Construct(SRCLogicPanelHeaderBase::FArguments());
+	
+	BehaviourPanelWeakPtr = InBehaviourPanel;
+	ControllerItemWeakPtr = InControllerItem;
+
+	// Add New Button
+	const TSharedRef<SPositiveActionButton> AddNewMenu = SNew(SPositiveActionButton)
+		.AddMetaData<FTagMetaData>(FTagMetaData(TEXT("Add Behaviour")))
+		.Icon(FAppStyle::Get().GetBrush("Icons.Plus"))
+		.Text(LOCTEXT("AddBehaviourLabel", "Add Behaviour"))
+		.OnGetMenuContent(this, &SRCBehaviourPanelHeader::GetBehaviourMenuContentWidget);
+	
+	ChildSlot
+	[
+		SNew(SHorizontalBox)
+		// [Add New]
+		+ SHorizontalBox::Slot()
+		.VAlign(VAlign_Center)
+		.AutoWidth()
+		[
+			AddNewMenu
+		]
+		// [Empty]
+		+ SHorizontalBox::Slot()
+		.VAlign(VAlign_Center)
+		.AutoWidth()
+		[
+			SNew(SButton)
+			.ContentPadding(2.0f)
+			.ButtonStyle(FAppStyle::Get(), "FlatButton")
+			.OnClicked(this, &SRCBehaviourPanelHeader::OnClickEmptyButton)
+			[
+				SNew(STextBlock)
+				.ColorAndOpacity(FColor::White)
+				.Text(LOCTEXT("EmptyBehavioursButtonLabel", "Empty"))
+			]
+		]
+	];
+}
+
+END_SLATE_FUNCTION_BUILD_OPTIMIZATION
+
+TSharedRef<SWidget> SRCBehaviourPanelHeader::GetBehaviourMenuContentWidget()
+{
+	constexpr bool bShouldCloseWindowAfterMenuSelection = true;
+	FMenuBuilder MenuBuilder(bShouldCloseWindowAfterMenuSelection, nullptr);
+
+	if (const TSharedPtr<FRCControllerModel> ControllerItem = ControllerItemWeakPtr.Pin())
+	{
+		if (URCController* Controller = Cast<URCController>(ControllerItem->GetVirtualProperty()))
+		{
+			for( TObjectIterator<UClass> It ; It ; ++It )
+			{
+				UClass* Class = *It;
+				
+				if( Class->IsChildOf(URCBehaviourNode::StaticClass()) )
+				{
+					const bool bIsClassInstantiatable = Class->HasAnyClassFlags(CLASS_Deprecated | CLASS_NewerVersionExists | CLASS_Abstract) || FKismetEditorUtilities::IsClassABlueprintSkeleton(Class);
+					if (bIsClassInstantiatable)
+					{
+						continue;
+					}
+					
+					if (Class->IsInBlueprint())
+					{
+						if (Class->GetSuperClass() != URCBehaviourBlueprintNode::StaticClass())
+						{
+							continue;
+						}
+					}
+					
+					URCBehaviour* Behaviour = Controller->CreateBehaviour(Class);
+					if (!Behaviour)
+					{
+						continue;
+					}
+			
+					FUIAction Action(FExecuteAction::CreateRaw(this, &SRCBehaviourPanelHeader::OnAddBehaviourClicked, Class));
+					MenuBuilder.AddMenuEntry(
+						FText::Format( LOCTEXT("AddBehaviourNode", "{0}"), FText::FromName(Class->GetFName())),
+						FText::Format( LOCTEXT("AddBehaviourNodeTooltip", "{0}"), FText::FromName(Class->GetFName())),
+						FSlateIcon(),
+						Action);
+				}
+			}		
+		}
+	}
+	
+	return MenuBuilder.MakeWidget();
+}
+
+void SRCBehaviourPanelHeader::OnAddBehaviourClicked(UClass* InClass)
+{
+	if (const TSharedPtr<FRCControllerModel> ControllerItem = ControllerItemWeakPtr.Pin())
+	{
+		if (URCController* Controller = Cast<URCController>(ControllerItem->GetVirtualProperty()))
+		{
+			URCBehaviour* NewBehaviour = Controller->AddBehaviour(InClass);
+
+			if (const TSharedPtr<SRCBehaviourPanel> BehaviourPanel = BehaviourPanelWeakPtr.Pin())
+			{
+				if (const TSharedPtr<SRemoteControlPanel> RemoteControlPanel = BehaviourPanel->GetRemoteControlPanel())
+				{
+					RemoteControlPanel->OnBehaviourAdded.Broadcast(NewBehaviour);
+				}
+			}
+		}
+	}
+}
+
+FReply SRCBehaviourPanelHeader::OnClickEmptyButton()
+{
+	if (const TSharedPtr<FRCControllerModel> ControllerItem = ControllerItemWeakPtr.Pin())
+	{
+		if (URCController* Controller = Cast<URCController>(ControllerItem->GetVirtualProperty()))
+		{
+			Controller->EmptyBehaviours();
+		}
+	}
+
+	const TSharedPtr<SRemoteControlPanel> RemoteControlPanel = BehaviourPanelWeakPtr.Pin()->GetRemoteControlPanel();
+	RemoteControlPanel->OnEmptyBehaviours.Broadcast();
+	
+	return FReply::Handled();
+}
+
+#undef LOCTEXT_NAMESPACE
\ No newline at end of file
diff --git a/Engine/Plugins/VirtualProduction/RemoteControl/Source/RemoteControlUI/Private/UI/Behaviour/SRCBehaviourPanelHeader.h b/Engine/Plugins/VirtualProduction/RemoteControl/Source/RemoteControlUI/Private/UI/Behaviour/SRCBehaviourPanelHeader.h
new file mode 100644
--- /dev/null
+++ b/Engine/Plugins/VirtualProduction/RemoteControl/Source/RemoteControlUI/Private/UI/Behaviour/SRCBehaviourPanelHeader.h
@@ -0,0 +1,44 @@
+﻿// Copyright Epic Games, Inc. All Rights Reserved.
+
+#pragma once
+
+#include "Behaviour/RCBehaviourNode.h"
+#include "UI/BaseLogicUI/SRCLogicPanelHeaderBase.h"
+
+class FRCControllerModel;
+class SRCBehaviourPanel;
+
+/**
+ * Behaviour panel UI Header widget.
+ *
+ * Contains buttons for Add / Remove / Empty
+ */
+class REMOTECONTROLUI_API SRCBehaviourPanelHeader : public SRCLogicPanelHeaderBase
+{
+public:
+	SLATE_BEGIN_ARGS(SRCBehaviourPanelHeader)
+		{
+		}
+
+	SLATE_END_ARGS()
+
+	/** Constructs this widget with InArgs */
+	void Construct(const FArguments& InArgs, TSharedRef<SRCBehaviourPanel> InBehaviourPanel, TSharedPtr<FRCControllerModel> InControllerItem);
+
+private:
+	/** Builds a menu containing the list of all possible Behaviours */
+	TSharedRef<SWidget> GetBehaviourMenuContentWidget();
+
+	/** Handles click event for Add Behaviour button*/
+	void OnAddBehaviourClicked(UClass* InClass);
+
+	/** Handles click event for "Empty" button; clears all Behaviours from the panel*/
+	FReply OnClickEmptyButton();
+	
+private:
+	/** The parent Behaviour Panel UI widget holding this header*/
+	TWeakPtr<SRCBehaviourPanel> BehaviourPanelWeakPtr;
+	
+	/** The Controller (UI model) associated with our Behaviour panel*/
+	TWeakPtr<FRCControllerModel> ControllerItemWeakPtr;
+};
diff --git a/Engine/Plugins/VirtualProduction/RemoteControl/Source/RemoteControlUI/Private/UI/Behaviour/SRCBehaviourPanelList.cpp b/Engine/Plugins/VirtualProduction/RemoteControl/Source/RemoteControlUI/Private/UI/Behaviour/SRCBehaviourPanelList.cpp
new file mode 100644
--- /dev/null
+++ b/Engine/Plugins/VirtualProduction/RemoteControl/Source/RemoteControlUI/Private/UI/Behaviour/SRCBehaviourPanelList.cpp
@@ -0,0 +1,159 @@
+﻿// Copyright Epic Games, Inc. All Rights Reserved.
+
+#include "SRCBehaviourPanelList.h"
+
+#include "Controller/RCController.h"
+#include "RCBehaviourModel.h"
+#include "RemoteControlPreset.h"
+#include "SlateOptMacros.h"
+#include "SRCBehaviourPanel.h"
+#include "UI/Behaviour/Builtin/RCBehaviourIsEqualModel.h"
+#include "UI/SRemoteControlPanel.h"
+#include "UI/Controller/RCControllerModel.h"
+
+BEGIN_SLATE_FUNCTION_BUILD_OPTIMIZATION
+
+void SRCBehaviourPanelList::Construct(const FArguments& InArgs, TSharedRef<SRCBehaviourPanel> InBehaviourPanel, TSharedPtr<FRCControllerModel> InControllerItem)
+{
+	SRCLogicPanelListBase::Construct(SRCLogicPanelListBase::FArguments());
+	
+	BehaviourPanelWeakPtr = InBehaviourPanel;
+	ControllerItemWeakPtr = InControllerItem;
+	
+	ListView = SNew(SListView<TSharedPtr<FRCBehaviourModel>>)
+		.ListItemsSource( &BehaviourItems )
+		.OnSelectionChanged(this, &SRCBehaviourPanelList::OnTreeSelectionChanged)
+		.OnGenerateRow(this, &SRCBehaviourPanelList::OnGenerateWidgetForList );
+	
+	ChildSlot
+	[
+		ListView.ToSharedRef()
+	];
+
+	// Add delegates
+	const TSharedPtr<SRemoteControlPanel> RemoteControlPanel = InBehaviourPanel->GetRemoteControlPanel();
+	check(RemoteControlPanel)
+	RemoteControlPanel->OnBehaviourAdded.AddSP(this, &SRCBehaviourPanelList::OnBehaviourAdded);
+	RemoteControlPanel->OnEmptyBehaviours.AddSP(this, &SRCBehaviourPanelList::OnEmptyBehaviours);
+
+	Reset();
+}
+
+END_SLATE_FUNCTION_BUILD_OPTIMIZATION
+
+void SRCBehaviourPanelList::Reset()
+{
+	BehaviourItems.Empty();
+
+	if (TSharedPtr<FRCControllerModel> ControllerItem = ControllerItemWeakPtr.Pin())
+	{
+		if (URCController* Controller = Cast<URCController>(ControllerItem->GetVirtualProperty()))
+		{
+			for (URCBehaviour* Behaviour : Controller->Behaviours)
+			{
+				if (URCIsEqualBehaviour* IsEqualBehaviour = Cast<URCIsEqualBehaviour>(Behaviour))
+				{
+					BehaviourItems.Add(MakeShared<FRCIsEqualBehaviourModel>(IsEqualBehaviour));
+				}
+				else
+				{
+					BehaviourItems.Add(MakeShared<FRCBehaviourModel>(Behaviour));
+				}
+			}
+		}
+	}
+
+	ListView->RebuildList();
+}
+
+TSharedRef<ITableRow> SRCBehaviourPanelList::OnGenerateWidgetForList(TSharedPtr<FRCBehaviourModel> InItem,
+	const TSharedRef<STableViewBase>& OwnerTable)
+{
+	return SNew(STableRow<TSharedPtr<FString>>, OwnerTable)
+	[
+		InItem->GetWidget()
+	];
+}
+
+void SRCBehaviourPanelList::OnTreeSelectionChanged(TSharedPtr<FRCBehaviourModel> InItem, ESelectInfo::Type)
+{
+	if (const TSharedPtr<SRemoteControlPanel> RemoteControlPanel = BehaviourPanelWeakPtr.Pin()->GetRemoteControlPanel())
+	{
+		if (InItem.IsValid())
+		{
+			if (InItem != SelectedBehaviourItemWeakPtr.Pin())
+			{
+				RemoteControlPanel->OnBehaviourSelectionChanged.Broadcast(InItem);
+			}
+		}
+
+		SelectedBehaviourItemWeakPtr = InItem;
+
+		if (const TSharedPtr<FRCControllerModel> ControllerItem = ControllerItemWeakPtr.Pin())
+		{
+			ControllerItem->UpdateSelectedBehaviourModel(InItem);
+		}
+	}
+}
+
+void SRCBehaviourPanelList::OnBehaviourAdded(URCBehaviour* InBehaviour)
+{
+	Reset();
+}
+
+void SRCBehaviourPanelList::OnEmptyBehaviours()
+{
+	if (const TSharedPtr<SRemoteControlPanel> RemoteControlPanel = BehaviourPanelWeakPtr.Pin()->GetRemoteControlPanel())
+	{
+		RemoteControlPanel->OnBehaviourSelectionChanged.Broadcast(nullptr);	
+	}
+	
+	Reset();
+}
+
+void SRCBehaviourPanelList::BroadcastOnItemRemoved()
+{
+	if (const TSharedPtr<SRemoteControlPanel> RemoteControlPanel = BehaviourPanelWeakPtr.Pin()->GetRemoteControlPanel())
+	{
+		RemoteControlPanel->OnBehaviourSelectionChanged.Broadcast(nullptr);
+	}
+}
+
+URemoteControlPreset* SRCBehaviourPanelList::GetPreset()
+{
+	if (BehaviourPanelWeakPtr.IsValid())
+	{
+		return BehaviourPanelWeakPtr.Pin()->GetPreset();
+	}
+
+	return nullptr;
+}
+
+int32 SRCBehaviourPanelList::RemoveModel(const TSharedPtr<FRCLogicModeBase> InModel)
+{
+	if (const TSharedPtr<FRCControllerModel> ControllerModel = ControllerItemWeakPtr.Pin())
+	{
+		if (URCController* Controller = Cast<URCController>(ControllerModel->GetVirtualProperty()))
+		{
+			if(const TSharedPtr<FRCBehaviourModel> SelectedBehaviour = StaticCastSharedPtr<FRCBehaviourModel>(InModel))
+			{
+				// Remove Model from Data Container
+				const int32 RemoveCount = Controller->RemoveBehaviour(SelectedBehaviour->GetBehaviour());
+
+				return RemoveCount;
+			}
+		}
+	}
+
+	return 0;
+}
+
+bool SRCBehaviourPanelList::IsListFocused() const
+{
+	return ListView->HasAnyUserFocus().IsSet();
+}
+
+void SRCBehaviourPanelList::DeleteSelectedPanelItem()
+{
+	DeleteItemFromLogicPanel<FRCBehaviourModel>(BehaviourItems, ListView->GetSelectedItems());
+}
diff --git a/Engine/Plugins/VirtualProduction/RemoteControl/Source/RemoteControlUI/Private/UI/Behaviour/SRCBehaviourPanelList.h b/Engine/Plugins/VirtualProduction/RemoteControl/Source/RemoteControlUI/Private/UI/Behaviour/SRCBehaviourPanelList.h
new file mode 100644
--- /dev/null
+++ b/Engine/Plugins/VirtualProduction/RemoteControl/Source/RemoteControlUI/Private/UI/Behaviour/SRCBehaviourPanelList.h
@@ -0,0 +1,82 @@
+﻿// Copyright Epic Games, Inc. All Rights Reserved.
+
+#pragma once
+
+#include "UI/BaseLogicUI/SRCLogicPanelListBase.h"
+
+class URCBehaviour;
+class FRCControllerModel;
+class FRCBehaviourModel;
+class ITableRow;
+class ITableBase;
+class SRCBehaviourPanel;
+class STableViewBase;
+class URemoteControlPreset;
+template <typename ItemType> class SListView;
+
+/*
+* ~ SRCBehaviourPanelList ~
+*
+* UI Widget for Behaviours List
+* Used as part of the RC Logic Behaviour Panel.
+*/
+class REMOTECONTROLUI_API SRCBehaviourPanelList : public SRCLogicPanelListBase
+{
+public:
+	SLATE_BEGIN_ARGS(SRCBehaviourPanelList)
+	{
+	}
+
+	SLATE_END_ARGS()
+
+	/** Constructs this widget with InArgs */
+	void Construct(const FArguments& InArgs, const TSharedRef<SRCBehaviourPanel> InBehaviourPanel, TSharedPtr<FRCControllerModel> InControllerItem);
+
+private:
+	/** OnGenerateRow delegate for the Behaviours List View */
+	TSharedRef<ITableRow> OnGenerateWidgetForList( TSharedPtr<FRCBehaviourModel> InItem, const TSharedRef<STableViewBase>& OwnerTable );
+	
+	/** OnSelectionChanged delegate for Behaviours List View */
+	void OnTreeSelectionChanged(TSharedPtr<FRCBehaviourModel> InItem , ESelectInfo::Type);
+
+	/** Responds to the selection of a newly created Behaviour. Resets UI state */
+	void OnBehaviourAdded(URCBehaviour* InBehaviour);
+
+	/** Responds to the removal of all Behaviours. Rests UI state */
+	void OnEmptyBehaviours();
+
+	/** The parent Behaviour Panel widget */
+	TWeakPtr<SRCBehaviourPanel> BehaviourPanelWeakPtr;
+	
+	/** The Controller (UI model) associated with us */
+	TWeakPtr<FRCControllerModel> ControllerItemWeakPtr;
+	
+	/** The currently selected Behaviour item (UI model) */
+	TWeakPtr<FRCBehaviourModel> SelectedBehaviourItemWeakPtr;
+
+	/** List of Behaviours (UI models) active in this widget */
+	TArray<TSharedPtr<FRCBehaviourModel>> BehaviourItems;
+	
+	/** List View widget for representing our Behaivours List*/
+	TSharedPtr<SListView<TSharedPtr<FRCBehaviourModel>>> ListView;
+
+	/** Refreshes the list from the latest state of the model*/
+	virtual void Reset() override;
+
+	/** Handles broadcasting of a successful remove item operation.*/
+	virtual void BroadcastOnItemRemoved() override;
+
+	/** Fetches the Remote Control preset associated with the parent panel */
+	virtual URemoteControlPreset* GetPreset() override;
+
+	/** Removes the given Behaviour UI model item from the list of UI models*/
+	virtual int32 RemoveModel(const TSharedPtr<FRCLogicModeBase> InModel) override;
+
+public:
+	/** Whether the Behaviours List View currently has focus.*/
+	bool IsListFocused() const;
+
+	/** Deletes currently selected items from the list view*/
+	void DeleteSelectedPanelItem();
+};
+
diff --git a/Engine/Plugins/VirtualProduction/RemoteControl/Source/RemoteControlUI/Private/UI/Controller/RCControllerModel.cpp b/Engine/Plugins/VirtualProduction/RemoteControl/Source/RemoteControlUI/Private/UI/Controller/RCControllerModel.cpp
new file mode 100644
--- /dev/null
+++ b/Engine/Plugins/VirtualProduction/RemoteControl/Source/RemoteControlUI/Private/UI/Controller/RCControllerModel.cpp
@@ -0,0 +1,85 @@
+﻿// Copyright Epic Games, Inc. All Rights Reserved.
+
+#include "RCControllerModel.h"
+
+#include "IDetailTreeNode.h"
+#include "RemoteControlPreset.h"
+#include "RCVirtualProperty.h"
+#include "UI/SRemoteControlPanel.h"
+
+FRCControllerModel::FRCControllerModel(const FName& InPropertyName, const TSharedRef<IDetailTreeNode>& InTreeNode, const TSharedPtr<SRemoteControlPanel> InRemoteControlPanel)
+	: FRCLogicModeBase(InRemoteControlPanel)
+	, DetailTreeNodeWeakPtr(InTreeNode)
+	, PropertyName(InPropertyName)
+{
+}
+
+TSharedRef<SWidget> FRCControllerModel::GetWidget() const
+{
+	const FNodeWidgets NodeWidgets = DetailTreeNodeWeakPtr.Pin()->CreateNodeWidgets();
+
+	const TSharedRef<SHorizontalBox> FieldWidget = SNew(SHorizontalBox);
+	
+	if (NodeWidgets.ValueWidget)
+	{
+		FieldWidget->AddSlot()
+			.Padding(FMargin(3.0f, 2.0f))
+			.HAlign(HAlign_Right)
+			.AutoWidth()
+			[
+				NodeWidgets.ValueWidget.ToSharedRef()
+			];
+	}
+	else if (NodeWidgets.WholeRowWidget)
+	{
+		FieldWidget->AddSlot()
+			.Padding(FMargin(3.0f, 2.0f))
+			.AutoWidth()
+			[
+				NodeWidgets.WholeRowWidget.ToSharedRef()
+			];
+	}
+	
+	return 	SNew(SHorizontalBox)
+		.Clipping(EWidgetClipping::OnDemand)
+		// Field name
+		+ SHorizontalBox::Slot()
+		.VAlign(VAlign_Center)
+		.AutoWidth()		
+		[
+			NodeWidgets.NameWidget.ToSharedRef()
+		]
+		// Value Widget
+		+ SHorizontalBox::Slot()
+		.VAlign(VAlign_Center)
+		.AutoWidth()
+		.Padding(FMargin(3.f))
+		[
+			FieldWidget
+		];
+}
+
+URCVirtualPropertyBase* FRCControllerModel::GetVirtualProperty() const
+{
+	if (const URemoteControlPreset* Preset = GetPreset())
+	{
+		return Preset->GetVirtualProperty(PropertyName);
+	}
+
+	return nullptr;
+}
+
+const FName& FRCControllerModel::GetPropertyName() const
+{
+	return PropertyName;
+}
+
+TSharedPtr<FRCBehaviourModel> FRCControllerModel::GetSelectedBehaviourModel() const
+{
+	return SelectedBehaviourModelWeakPtr.Pin();
+}
+
+void FRCControllerModel::UpdateSelectedBehaviourModel(TSharedPtr<FRCBehaviourModel> InModel)
+{
+	SelectedBehaviourModelWeakPtr = InModel;
+}
\ No newline at end of file
diff --git a/Engine/Plugins/VirtualProduction/RemoteControl/Source/RemoteControlUI/Private/UI/Controller/RCControllerModel.h b/Engine/Plugins/VirtualProduction/RemoteControl/Source/RemoteControlUI/Private/UI/Controller/RCControllerModel.h
new file mode 100644
--- /dev/null
+++ b/Engine/Plugins/VirtualProduction/RemoteControl/Source/RemoteControlUI/Private/UI/Controller/RCControllerModel.h
@@ -0,0 +1,57 @@
+﻿// Copyright Epic Games, Inc. All Rights Reserved.
+
+#pragma once
+
+#include "UI/BaseLogicUI/RCLogicModeBase.h"
+#include "UObject/WeakFieldPtr.h"
+
+class FRCBehaviourModel;
+class IDetailTreeNode;
+class SRemoteControlPanel;
+class URCBehaviour;
+class URCControllerContainer;
+class URCUserDefinedStruct;
+class URCVirtualPropertyContainerBase;
+class URCVirtualPropertyBase;
+
+/*
+* ~ FRCControllerModel ~
+*
+* UI model for representing a Controller
+* Contains a row widget with Controller Name and a Value Widget
+*/
+class FRCControllerModel : public FRCLogicModeBase
+{
+public:
+	FRCControllerModel(const FName& PropertyName, const TSharedRef<IDetailTreeNode>& InTreeNode, const TSharedPtr<SRemoteControlPanel> InRemoteControlPanel);
+	
+	/** The widget to be rendered for this Controller
+	* Used to represent a single row when added to the Controllers Panel List
+	*/
+	virtual TSharedRef<SWidget> GetWidget() const override;
+
+	/** Returns the Controller (Data Model) associated with us*/
+	virtual URCVirtualPropertyBase* GetVirtualProperty() const;
+	
+	/** Returns the Controller Name (from Virtual Property) represented by us*/
+	virtual const FName& GetPropertyName() const;
+
+	/** Returns the currently selected Behaviour (UI model) */
+	TSharedPtr<FRCBehaviourModel> GetSelectedBehaviourModel() const;
+
+	/** Updates our internal record of the currently selected Behaviour (UI Model) */
+	void UpdateSelectedBehaviourModel(TSharedPtr<FRCBehaviourModel> InModel);
+
+private:
+	/** The Parent Container holding our Controller (Virtual Property) Data Model */
+	TWeakObjectPtr<URCVirtualPropertyContainerBase> PropertyContainerWeakPtr;
+	
+	/** The row generator used to build a generic Value Widget*/
+	TWeakPtr<IDetailTreeNode> DetailTreeNodeWeakPtr;
+
+	/** The currently selected Behaviour (UI model) */
+	TWeakPtr<FRCBehaviourModel>  SelectedBehaviourModelWeakPtr;
+
+	/** The Controller (Virtual Property Name) represented by us*/
+	FName PropertyName;
+};
diff --git a/Engine/Plugins/VirtualProduction/RemoteControl/Source/RemoteControlUI/Private/UI/Controller/SRCControllerPanel.cpp b/Engine/Plugins/VirtualProduction/RemoteControl/Source/RemoteControlUI/Private/UI/Controller/SRCControllerPanel.cpp
new file mode 100644
--- /dev/null
+++ b/Engine/Plugins/VirtualProduction/RemoteControl/Source/RemoteControlUI/Private/UI/Controller/SRCControllerPanel.cpp
@@ -0,0 +1,49 @@
+﻿// Copyright Epic Games, Inc. All Rights Reserved.
+
+#include "SRCControllerPanel.h"
+
+#include "RCControllerModel.h"
+#include "SlateOptMacros.h"
+#include "SRCControllerPanelHeader.h"
+#include "SRCControllerPanelList.h"
+#include "UI/BaseLogicUI/RCLogicModeBase.h"
+#include "UI/SRemoteControlPanel.h"
+#include "Widgets/SBoxPanel.h"
+#include "Widgets/Views/SListView.h"
+
+BEGIN_SLATE_FUNCTION_BUILD_OPTIMIZATION
+
+void SRCControllerPanel::Construct(const FArguments& InArgs, const TSharedRef<SRemoteControlPanel>& InPanel)
+{
+	SRCLogicPanelBase::Construct(SRCLogicPanelBase::FArguments(), InPanel);
+	
+	ChildSlot
+	[
+		SNew(SVerticalBox)
+		// Header
+		+ SVerticalBox::Slot()
+		.Padding(FMargin(5.f))
+		.AutoHeight()
+		[
+			SNew(SRCControllerPanelHeader, SharedThis(this))
+		]
+		// List
+		+ SVerticalBox::Slot()
+		.Padding(FMargin(5.f))
+		[
+			SAssignNew(ControllerPanelList, SRCControllerPanelList, SharedThis(this))
+		]
+	];
+}
+
+END_SLATE_FUNCTION_BUILD_OPTIMIZATION
+
+bool SRCControllerPanel::IsListFocused() const
+{
+	return ControllerPanelList->IsListFocused();
+}
+
+void SRCControllerPanel::DeleteSelectedPanelItem()
+{
+	ControllerPanelList->DeleteSelectedPanelItem();
+}
\ No newline at end of file
diff --git a/Engine/Plugins/VirtualProduction/RemoteControl/Source/RemoteControlUI/Private/UI/Controller/SRCControllerPanel.h b/Engine/Plugins/VirtualProduction/RemoteControl/Source/RemoteControlUI/Private/UI/Controller/SRCControllerPanel.h
new file mode 100644
--- /dev/null
+++ b/Engine/Plugins/VirtualProduction/RemoteControl/Source/RemoteControlUI/Private/UI/Controller/SRCControllerPanel.h
@@ -0,0 +1,35 @@
+﻿// Copyright Epic Games, Inc. All Rights Reserved.
+
+#pragma once
+
+#include "UI/BaseLogicUI/SRCLogicPanelBase.h"
+
+/*
+* ~ SRCControllerPanel ~
+*
+* UI Widget for Controller Panel.
+* Contains a header (Add/Remove/Empty) and List of Controllers
+*/
+class REMOTECONTROLUI_API SRCControllerPanel : public SRCLogicPanelBase
+{
+public:
+	SLATE_BEGIN_ARGS(SRCControllerPanel)
+		{
+		}
+
+	SLATE_END_ARGS()
+
+	/** Constructs this widget with InArgs */
+	void Construct(const FArguments& InArgs, const TSharedRef<SRemoteControlPanel>& InPanel);
+
+private:
+	/** Widget representing List of Controllers */
+	TSharedPtr<class SRCControllerPanelList> ControllerPanelList;
+
+public:
+	/** Whether the Controller list widget currently has focus. Used for Delete Item UI command */
+	bool IsListFocused() const;
+
+	/** Delete Item UI command implementation for this panel */
+	virtual void DeleteSelectedPanelItem() override;
+};
diff --git a/Engine/Plugins/VirtualProduction/RemoteControl/Source/RemoteControlUI/Private/UI/Controller/SRCControllerPanelHeader.cpp b/Engine/Plugins/VirtualProduction/RemoteControl/Source/RemoteControlUI/Private/UI/Controller/SRCControllerPanelHeader.cpp
new file mode 100644
--- /dev/null
+++ b/Engine/Plugins/VirtualProduction/RemoteControl/Source/RemoteControlUI/Private/UI/Controller/SRCControllerPanelHeader.cpp
@@ -0,0 +1,162 @@
+﻿// Copyright Epic Games, Inc. All Rights Reserved.
+
+#include "SRCControllerPanelHeader.h"
+
+#include "Controller/RCController.h"
+#include "Framework/MultiBox/MultiBoxBuilder.h"
+#include "RCVirtualPropertyContainer.h"
+#include "RemoteControlLogicConfig.h"
+#include "RemoteControlPreset.h"
+#include "SlateOptMacros.h"
+#include "SPositiveActionButton.h"
+#include "SRCControllerPanel.h"
+#include "UI/SRemoteControlPanel.h"
+#include "Widgets/Input/SButton.h"
+
+#define LOCTEXT_NAMESPACE "SRCControllerPanelHeader"
+
+BEGIN_SLATE_FUNCTION_BUILD_OPTIMIZATION
+
+void SRCControllerPanelHeader::Construct(const FArguments& InArgs, const TSharedRef<SRCControllerPanel> InControllerPanel)
+{
+	SRCLogicPanelHeaderBase::Construct(SRCLogicPanelHeaderBase::FArguments());
+	
+	ControllerPanelWeakPtr = InControllerPanel;
+
+	const TSharedRef<SPositiveActionButton> AddNewMenu = SNew(SPositiveActionButton)
+		.AddMetaData<FTagMetaData>(FTagMetaData(TEXT("Add Controller")))
+		.Icon(FAppStyle::Get().GetBrush("Icons.Plus"))
+		.Text(LOCTEXT("AddControllerLabel", "Add Controller"))
+		.OnGetMenuContent(this, &SRCControllerPanelHeader::GetControllerMenuContentWidget);
+	
+	ChildSlot
+	[
+		SNew(SHorizontalBox)
+		// [Add]
+		+ SHorizontalBox::Slot()
+		.VAlign(VAlign_Center)
+		.AutoWidth()
+		[
+			AddNewMenu
+		]
+		// [Empty]
+		+ SHorizontalBox::Slot()
+		.VAlign(VAlign_Center)
+		.AutoWidth()
+		[
+			SNew(SButton)
+			.ContentPadding(2.0f)
+			.ButtonStyle(FAppStyle::Get(), "FlatButton")
+			.OnClicked(this, &SRCControllerPanelHeader::OnClickEmptyButton)
+			[
+				SNew(STextBlock)
+				.ColorAndOpacity(FColor::White)
+				.Text(LOCTEXT("EmptyAllControllersButtonLabel", "Empty"))
+			]
+		]
+	];
+}
+
+END_SLATE_FUNCTION_BUILD_OPTIMIZATION
+
+static UObject* GetBaseStructForType(FName StructType)
+{
+	if (StructType == "FVector")
+	{
+		return TBaseStructure<FVector>::Get();
+	}
+	else if (StructType == "FColor")
+	{
+		return TBaseStructure<FColor>::Get();
+	}
+	else if (StructType == "FRotator")
+	{
+		return TBaseStructure<FRotator>::Get();
+	}
+
+	ensureMsgf(false, TEXT("Found unsupported struct type %s in config."), *StructType.ToString());
+
+	return nullptr;
+}
+
+TSharedRef<SWidget> SRCControllerPanelHeader::GetControllerMenuContentWidget() const
+{
+	constexpr  bool bShouldCloseWindowAfterMenuSelection = true;
+	FMenuBuilder MenuBuilder(bShouldCloseWindowAfterMenuSelection, nullptr);
+
+	TArray<TPair<EPropertyBagPropertyType, UObject*>> VirtualPropertyFieldClassNames;
+
+	// See Config file: BaseRemoteControl.ini
+	const URemoteControlLogicConfig* RCLogicConfig = GetDefault<URemoteControlLogicConfig>();
+	for (const EPropertyBagPropertyType ControllerType : RCLogicConfig->SupportedControllerTypes)
+	{
+		if (ControllerType == EPropertyBagPropertyType::Struct)
+		{
+			for (const FName StructType : RCLogicConfig->SupportedControllerStructTypes)
+			{
+				UObject* ValueTypeObject = GetBaseStructForType(StructType);
+				if (ensure(ValueTypeObject))
+				{
+					VirtualPropertyFieldClassNames.Add({ ControllerType, ValueTypeObject });
+				}
+			}
+		}
+		else
+		{
+			VirtualPropertyFieldClassNames.Add({ ControllerType, nullptr });
+		}
+	}
+
+	// Generate a menu from the list of supported Controllers
+	for (const TPair<EPropertyBagPropertyType, UObject*>& Pair : VirtualPropertyFieldClassNames)
+	{
+		// Display Name
+		UEnum* Enum = FindObjectChecked<UEnum>(nullptr, TEXT("EPropertyBagPropertyType"));
+		FString PropertyTypeString = Enum->GetValueAsName(Pair.Key).ToString();
+		PropertyTypeString.RemoveFromStart("EPropertyBagPropertyType::");
+
+		const FName DefaultName = Pair.Value ? Pair.Value->GetFName() : *PropertyTypeString;
+		
+		// Menu Item
+		FUIAction Action(FExecuteAction::CreateRaw(this, &SRCControllerPanelHeader::OnAddControllerClicked, Pair.Key, Pair.Value));
+		MenuBuilder.AddMenuEntry(
+			FText::Format( LOCTEXT("AddController", "{0}"), FText::FromName(DefaultName)),
+			FText::Format( LOCTEXT("AddControllerTooltip", "{0}"), FText::FromName(DefaultName)),
+			FSlateIcon(),
+			Action);
+	}
+
+	return MenuBuilder.MakeWidget();
+}
+
+void SRCControllerPanelHeader::OnAddControllerClicked(const EPropertyBagPropertyType InValueType, UObject* InValueTypeObject) const
+{
+	// Add to the asset
+	if (URemoteControlPreset* Preset = ControllerPanelWeakPtr.Pin()->GetPreset())
+	{
+		URCVirtualPropertyInContainer* NewVirtualProperty = Preset->AddVirtualProperty(URCController::StaticClass(), InValueType, InValueTypeObject);
+		
+		// Refresh list
+		const TSharedPtr<SRemoteControlPanel> RemoteControlPanel = ControllerPanelWeakPtr.Pin()->GetRemoteControlPanel();
+		check(RemoteControlPanel)
+		RemoteControlPanel->OnControllerAdded.Broadcast(NewVirtualProperty->PropertyName);
+	}
+}
+
+FReply SRCControllerPanelHeader::OnClickEmptyButton()
+{
+	if (TSharedPtr<SRCControllerPanel> ControllerPanel = ControllerPanelWeakPtr.Pin())
+	{
+		if (URemoteControlPreset* Preset = ControllerPanel->GetPreset())
+		{
+			Preset->ResetVirtualProperties();
+		}
+
+		const TSharedPtr<SRemoteControlPanel> RemoteControlPanel = ControllerPanelWeakPtr.Pin()->GetRemoteControlPanel();
+		RemoteControlPanel->OnEmptyControllers.Broadcast();
+	}
+	
+	return FReply::Handled();
+}
+
+#undef LOCTEXT_NAMESPACE
\ No newline at end of file
diff --git a/Engine/Plugins/VirtualProduction/RemoteControl/Source/RemoteControlUI/Private/UI/Controller/SRCControllerPanelHeader.h b/Engine/Plugins/VirtualProduction/RemoteControl/Source/RemoteControlUI/Private/UI/Controller/SRCControllerPanelHeader.h
new file mode 100644
--- /dev/null
+++ b/Engine/Plugins/VirtualProduction/RemoteControl/Source/RemoteControlUI/Private/UI/Controller/SRCControllerPanelHeader.h
@@ -0,0 +1,41 @@
+﻿// Copyright Epic Games, Inc. All Rights Reserved.
+
+#pragma once
+
+#include "PropertyBag.h"
+#include "UI/BaseLogicUI/SRCLogicPanelHeaderBase.h"
+
+enum class EPropertyBagPropertyType : uint8;
+class SRCControllerPanel;
+
+/**
+ * Controller panel UI Header widget.
+ *
+ * Contains buttons for Add / Remove / Empty
+ */
+class REMOTECONTROLUI_API SRCControllerPanelHeader : public SRCLogicPanelHeaderBase
+{
+public:
+	SLATE_BEGIN_ARGS(SRCControllerPanelHeader)
+		{
+		}
+
+	SLATE_END_ARGS()
+
+	/** Constructs this widget with InArgs */
+	void Construct(const FArguments& InArgs, const TSharedRef<SRCControllerPanel> InControllerPanel);
+
+private:
+	/** Builds a menu containing the list of all possible Controllers*/
+	TSharedRef<SWidget> GetControllerMenuContentWidget() const;
+
+	/** Handles click event for Add Controller button*/
+	void OnAddControllerClicked(const EPropertyBagPropertyType InValueType, UObject* InValueTypeObject = nullptr) const;
+
+	/** Handles click event for Empty button; clears all controllers from the panel*/
+	FReply OnClickEmptyButton();
+
+	/** The parent Controller panel UI widget holding this header */
+	TWeakPtr<SRCControllerPanel> ControllerPanelWeakPtr;
+};
+
diff --git a/Engine/Plugins/VirtualProduction/RemoteControl/Source/RemoteControlUI/Private/UI/Controller/SRCControllerPanelList.cpp b/Engine/Plugins/VirtualProduction/RemoteControl/Source/RemoteControlUI/Private/UI/Controller/SRCControllerPanelList.cpp
new file mode 100644
--- /dev/null
+++ b/Engine/Plugins/VirtualProduction/RemoteControl/Source/RemoteControlUI/Private/UI/Controller/SRCControllerPanelList.cpp
@@ -0,0 +1,190 @@
+﻿// Copyright Epic Games, Inc. All Rights Reserved.
+
+#include "SRCControllerPanelList.h"
+
+#include "IDetailTreeNode.h"
+#include "IPropertyRowGenerator.h"
+#include "RCControllerModel.h"
+#include "RCVirtualProperty.h"
+#include "RCVirtualPropertyContainer.h"
+#include "RemoteControlPreset.h"
+#include "SlateOptMacros.h"
+#include "SRCControllerPanel.h"
+#include "UI/BaseLogicUI/RCLogicModeBase.h"
+#include "UI/SRemoteControlPanel.h"
+
+BEGIN_SLATE_FUNCTION_BUILD_OPTIMIZATION
+
+void SRCControllerPanelList::Construct(const FArguments& InArgs, const TSharedRef<SRCControllerPanel> InControllerPanel)
+{
+	SRCLogicPanelListBase::Construct(SRCLogicPanelListBase::FArguments());
+	
+	ControllerPanelWeakPtr = InControllerPanel;
+	
+	ListView = SNew(SListView<TSharedPtr<FRCControllerModel>>)
+		.ListItemsSource( &ControllerItems )
+		.OnSelectionChanged(this, &SRCControllerPanelList::OnTreeSelectionChanged)
+		.OnGenerateRow(this, &SRCControllerPanelList::OnGenerateWidgetForList );
+	
+	ChildSlot
+	[
+		ListView.ToSharedRef()
+	];
+
+	// Add delegates
+	if (const URemoteControlPreset* Preset = ControllerPanelWeakPtr.Pin()->GetPreset())
+	{
+		// Refresh list
+		const TSharedPtr<SRemoteControlPanel> RemoteControlPanel = ControllerPanelWeakPtr.Pin()->GetRemoteControlPanel();
+		check(RemoteControlPanel)
+		RemoteControlPanel->OnControllerAdded.AddSP(this, &SRCControllerPanelList::OnControllerAdded);
+		RemoteControlPanel->OnEmptyControllers.AddSP(this, &SRCControllerPanelList::OnEmptyControllers);
+	}
+
+	FPropertyRowGeneratorArgs Args;
+	Args.bShouldShowHiddenProperties = true;
+	PropertyRowGenerator = FModuleManager::GetModuleChecked<FPropertyEditorModule>("PropertyEditor").CreatePropertyRowGenerator(Args);
+	
+	Reset();
+}
+
+END_SLATE_FUNCTION_BUILD_OPTIMIZATION
+
+void SRCControllerPanelList::Reset()
+{
+	ControllerItems.Empty();
+
+	check(ControllerPanelWeakPtr.IsValid());
+
+	TSharedPtr<SRCControllerPanel> ControllerPanel = ControllerPanelWeakPtr.Pin();
+	URemoteControlPreset* Preset = ControllerPanel->GetPreset();
+	TSharedPtr<SRemoteControlPanel> RemoteControlPanel = ControllerPanel->GetRemoteControlPanel();
+
+	check(Preset);
+
+	PropertyRowGenerator->SetStructure(Preset->GetControllerContainerStructOnScope());
+	PropertyRowGenerator->OnFinishedChangingProperties().AddSP(this, &SRCControllerPanelList::OnFinishedChangingProperties);
+
+	// Generator should be moved to separate class
+	for (const TSharedRef<IDetailTreeNode>& CategoryNode : PropertyRowGenerator->GetRootTreeNodes())
+	{
+		TArray<TSharedRef<IDetailTreeNode>> Children;
+		CategoryNode->GetChildren(Children);
+		for (TSharedRef<IDetailTreeNode>& Child : Children)
+		{
+			FProperty* Property = Child->CreatePropertyHandle()->GetProperty();
+			check(Property);
+			ControllerItems.Add(MakeShared<FRCControllerModel>(Property->GetFName(), Child, RemoteControlPanel));
+		}
+	}
+
+	ListView->RebuildList();
+}
+
+TSharedRef<ITableRow> SRCControllerPanelList::OnGenerateWidgetForList(TSharedPtr<FRCControllerModel> InItem, const TSharedRef<STableViewBase>& OwnerTable)
+{
+	return SNew(STableRow<TSharedPtr<FString>>, OwnerTable)
+	[
+		InItem->GetWidget()
+	];
+}
+
+void SRCControllerPanelList::OnTreeSelectionChanged(TSharedPtr<FRCControllerModel> InItem, ESelectInfo::Type)
+{
+	if (TSharedPtr<SRCControllerPanel> ControllerPanel = ControllerPanelWeakPtr.Pin())
+	{
+		if (const TSharedPtr<SRemoteControlPanel> RemoteControlPanel = ControllerPanel->GetRemoteControlPanel())
+		{
+			if (InItem.IsValid())
+			{
+				if (InItem != SelectedControllerItemWeakPtr.Pin())
+				{
+					RemoteControlPanel->OnControllerSelectionChanged.Broadcast(InItem);
+					RemoteControlPanel->OnBehaviourSelectionChanged.Broadcast(InItem->GetSelectedBehaviourModel());
+				}
+			}
+
+			SelectedControllerItemWeakPtr = InItem;
+		}
+	}
+}
+
+void SRCControllerPanelList::OnControllerAdded(const FName& InNewPropertyName)
+{
+	Reset();
+}
+
+void SRCControllerPanelList::OnFinishedChangingProperties(const FPropertyChangedEvent& PropertyChangedEvent)
+{
+	if (TSharedPtr< SRCControllerPanel> ControllerPanel = ControllerPanelWeakPtr.Pin())
+	{
+		if (URemoteControlPreset* Preset = ControllerPanel->GetPreset())
+		{
+			Preset->OnModifyVirtualProperty(PropertyChangedEvent);
+		}
+	}
+}
+
+void SRCControllerPanelList::OnEmptyControllers()
+{
+	if (TSharedPtr< SRCControllerPanel> ControllerPanel = ControllerPanelWeakPtr.Pin())
+	{
+		if (TSharedPtr<SRemoteControlPanel> RemoteControlPanel = ControllerPanel->GetRemoteControlPanel())
+		{
+			RemoteControlPanel->OnControllerSelectionChanged.Broadcast(nullptr);
+			RemoteControlPanel->OnBehaviourSelectionChanged.Broadcast(nullptr);
+		}
+
+		Reset();
+	}
+}
+
+void SRCControllerPanelList::BroadcastOnItemRemoved()
+{
+	if (const TSharedPtr<SRemoteControlPanel> RemoteControlPanel = ControllerPanelWeakPtr.Pin()->GetRemoteControlPanel())
+	{
+		RemoteControlPanel->OnControllerSelectionChanged.Broadcast(nullptr);
+		RemoteControlPanel->OnBehaviourSelectionChanged.Broadcast(nullptr);
+	}
+}
+
+URemoteControlPreset* SRCControllerPanelList::GetPreset()
+{
+	if (ControllerPanelWeakPtr.IsValid())
+	{
+		return ControllerPanelWeakPtr.Pin()->GetPreset();
+	}
+
+	return nullptr;	
+}
+
+int32 SRCControllerPanelList::RemoveModel(const TSharedPtr<FRCLogicModeBase> InModel)
+{
+	if(ControllerPanelWeakPtr.IsValid())
+	{
+		if (URemoteControlPreset* Preset = ControllerPanelWeakPtr.Pin()->GetPreset())
+		{
+			if(const TSharedPtr<FRCControllerModel> SelectedController = StaticCastSharedPtr<FRCControllerModel>(InModel))
+			{
+				// Remove Model from Data Container
+				const bool bRemoved = Preset->RemoveVirtualProperty(SelectedController->GetPropertyName());
+				if (bRemoved)
+				{
+					return 1; // Remove Count
+				}
+			}
+		}
+	}
+
+	return 0;
+}
+
+bool SRCControllerPanelList::IsListFocused() const
+{
+	return ListView->HasAnyUserFocus().IsSet();
+}
+
+void SRCControllerPanelList::DeleteSelectedPanelItem()
+{
+	DeleteItemFromLogicPanel<FRCControllerModel>(ControllerItems, ListView->GetSelectedItems());
+}
diff --git a/Engine/Plugins/VirtualProduction/RemoteControl/Source/RemoteControlUI/Private/UI/Controller/SRCControllerPanelList.h b/Engine/Plugins/VirtualProduction/RemoteControl/Source/RemoteControlUI/Private/UI/Controller/SRCControllerPanelList.h
new file mode 100644
--- /dev/null
+++ b/Engine/Plugins/VirtualProduction/RemoteControl/Source/RemoteControlUI/Private/UI/Controller/SRCControllerPanelList.h
@@ -0,0 +1,87 @@
+﻿// Copyright Epic Games, Inc. All Rights Reserved.
+
+#pragma once
+#include "UI/BaseLogicUI/SRCLogicPanelListBase.h"
+
+class FRCControllerModel;
+class FRCLogicModeBase;
+class IPropertyRowGenerator;
+class ITableRow;
+class ITableBase;
+class SRCControllerPanel;
+class STableViewBase;
+class URCController;
+class URemoteControlPreset;
+template <typename ItemType> class SListView;
+
+/*
+* ~ SRCControllerPanelList ~
+*
+* UI Widget for Controllers List
+* Used as part of the RC Logic Actions Panel.
+*/
+class REMOTECONTROLUI_API SRCControllerPanelList : public SRCLogicPanelListBase
+{
+public:
+	SLATE_BEGIN_ARGS(SRCControllerPanelList)
+		{
+		}
+
+	SLATE_END_ARGS()
+
+	/** Constructs this widget with InArgs */
+	void Construct(const FArguments& InArgs, const TSharedRef<SRCControllerPanel> InControllerPanel);
+
+private:
+	/** OnGenerateRow delegate for the Actions List View */
+	TSharedRef<ITableRow> OnGenerateWidgetForList( TSharedPtr<FRCControllerModel> InItem, const TSharedRef<STableViewBase>& OwnerTable );
+	
+	/** OnSelectionChanged delegate for Actions List View */
+	void OnTreeSelectionChanged(TSharedPtr<FRCControllerModel> InItem , ESelectInfo::Type);
+
+	/** Responds to the selection of a newly created Controller. Resets UI state */
+	void OnControllerAdded(const FName& InNewPropertyName);
+	
+	/** Responds to the removal of all Controllers. Resets UI state */
+	void OnEmptyControllers();
+
+	/** Change listener for Controllers. Bound to the PropertyRowGenerator's delegate
+	* This is propagated to the corresponding Controller model (Virtual Property) for evaluating all associated Behaviours.
+	*/
+	void OnFinishedChangingProperties(const FPropertyChangedEvent& PropertyChangedEvent);
+
+	/** The row generator used to represent each Controller as a row, when used with SListView */
+	TSharedPtr<IPropertyRowGenerator> PropertyRowGenerator;
+
+	/** The currently selected Controller item (UI model) */
+	TWeakPtr<FRCControllerModel> SelectedControllerItemWeakPtr = nullptr;
+	
+	/** The parent Controller Panel widget */
+	TWeakPtr<SRCControllerPanel> ControllerPanelWeakPtr;
+
+	/** List of Controllers (UI model) active in this widget */
+	TArray<TSharedPtr<FRCControllerModel>> ControllerItems;
+
+	/** List View widget for representing our Controllers List*/
+	TSharedPtr<SListView<TSharedPtr<FRCControllerModel>>>  ListView;
+	
+	/** Refreshes the list from the latest state of the model*/
+	virtual void Reset() override;
+
+	/** Handles broadcasting of a successful remove item operation.*/
+	virtual void BroadcastOnItemRemoved() override;
+
+	/** Fetches the Remote Control preset associated with the parent panel */
+	virtual URemoteControlPreset* GetPreset() override;
+
+	/** Removes the given Controller UI model item from the list of UI models*/
+	virtual int32 RemoveModel(const TSharedPtr<FRCLogicModeBase> InModel) override;
+
+public:
+	/** Whether the Controllers List View currently has focus.*/
+	bool IsListFocused() const;
+
+	/** Deletes currently selected items from the list view*/
+	void DeleteSelectedPanelItem();
+};
+
diff --git a/Engine/Plugins/VirtualProduction/RemoteControl/Source/RemoteControlUI/Private/UI/RCUIHelpers.cpp b/Engine/Plugins/VirtualProduction/RemoteControl/Source/RemoteControlUI/Private/UI/RCUIHelpers.cpp
new file mode 100644
--- /dev/null
+++ b/Engine/Plugins/VirtualProduction/RemoteControl/Source/RemoteControlUI/Private/UI/RCUIHelpers.cpp
@@ -0,0 +1,47 @@
+// Copyright Epic Games, Inc. All Rights Reserved.
+
+#include "RCUIHelpers.h"
+
+#include "RCVirtualProperty.h"
+
+#include "EdGraphSchema_K2.h"
+#include "EdGraph/EdGraphPin.h"
+
+#include "GraphEditorSettings.h"
+#include "Controller/RCController.h"
+
+FLinearColor UE::RCUIHelpers::GetFieldClassTypeColor(const FProperty* InProperty)
+{
+	if (ensure(InProperty))
+	{
+		const UEdGraphSchema_K2* Schema = GetDefault<UEdGraphSchema_K2>();
+		FEdGraphPinType PinType;
+		if (Schema->ConvertPropertyToPinType(InProperty, PinType))
+		{
+			return Schema->GetPinTypeColor(PinType);
+		}
+	}
+
+	return FLinearColor::White;
+}
+
+FText UE::RCUIHelpers::GetFieldClassDisplayName(const FProperty* InProperty)
+{
+	if (ensure(InProperty))
+	{
+		const UEdGraphSchema_K2* Schema = GetDefault<UEdGraphSchema_K2>();
+		FEdGraphPinType PinType;
+
+		if (Schema->ConvertPropertyToPinType(InProperty, PinType))
+		{
+			if(UObject* SubCategoryObject = PinType.PinSubCategoryObject.Get())
+			{
+				return FText::FromName(*SubCategoryObject->GetName());
+			}
+			
+			return FText::FromName(PinType.PinCategory);
+		}
+	}
+
+	return FText::GetEmpty();
+}
\ No newline at end of file
diff --git a/Engine/Plugins/VirtualProduction/RemoteControl/Source/RemoteControlUI/Private/UI/RCUIHelpers.h b/Engine/Plugins/VirtualProduction/RemoteControl/Source/RemoteControlUI/Private/UI/RCUIHelpers.h
new file mode 100644
--- /dev/null
+++ b/Engine/Plugins/VirtualProduction/RemoteControl/Source/RemoteControlUI/Private/UI/RCUIHelpers.h
@@ -0,0 +1,21 @@
+// Copyright Epic Games, Inc. All Rights Reserved.
+
+#pragma once
+
+#include "CoreMinimal.h"
+
+class URCController;
+
+namespace UE::RCUIHelpers
+{
+	/** GetFieldClassTypeColor
+	* Fetches the editor color associated with a given Unreal Type (FProperty)
+	* Used to provide color coding in the Remote Control Logic Actions panel 
+	*/
+	FLinearColor GetFieldClassTypeColor(const FProperty* InProperty);
+
+	/** GetFieldClassDisplayName
+	* Fetches the display name associated with a given Unreal Type (FProperty)
+	*/
+	FText GetFieldClassDisplayName(const FProperty* InProperty);
+}
\ No newline at end of file
diff --git a/Engine/Plugins/VirtualProduction/RemoteControl/Source/RemoteControlUI/Private/UI/RemoteControlPanelStyle.cpp b/Engine/Plugins/VirtualProduction/RemoteControl/Source/RemoteControlUI/Private/UI/RemoteControlPanelStyle.cpp
--- a/Engine/Plugins/VirtualProduction/RemoteControl/Source/RemoteControlUI/Private/UI/RemoteControlPanelStyle.cpp
+++ b/Engine/Plugins/VirtualProduction/RemoteControl/Source/RemoteControlUI/Private/UI/RemoteControlPanelStyle.cpp
@@ -11,6 +11,7 @@
 
 #define IMAGE_PLUGIN_BRUSH( RelativePath, ... ) FSlateImageBrush(FRemoteControlPanelStyle::InContent(RelativePath, ".png" ), __VA_ARGS__)
 #define IMAGE_BRUSH(RelativePath, ...) FSlateImageBrush(StyleSet->RootToContentDir(RelativePath, TEXT(".png")), __VA_ARGS__)
+#define IMAGE_BRUSH_SVG( RelativePath, ... ) FSlateVectorImageBrush(StyleSet->RootToContentDir(RelativePath, TEXT(".svg")), __VA_ARGS__)
 #define BOX_BRUSH(RelativePath, ...) FSlateBoxBrush(StyleSet->RootToContentDir(RelativePath, TEXT(".png")), __VA_ARGS__)
 #define CORE_BOX_BRUSH( RelativePath, ... ) FSlateBoxBrush(StyleSet->RootToCoreContentDir(RelativePath, TEXT(".png") ), __VA_ARGS__)
 #define BOX_PLUGIN_BRUSH( RelativePath, ... ) FSlateBoxBrush(FRemoteControlPanelStyle::InContent( RelativePath, ".png" ), __VA_ARGS__)
@@ -28,6 +29,7 @@ void FRemoteControlPanelStyle::Initialize()
 	const FVector2D Icon8x8(8.0f, 8.0f);
 	const FVector2D Icon16x16(16.0f, 16.0f);
 	const FVector2D Icon20x20(20.0f, 20.0f);
+	const FVector2D Icon32x32(32.f);
 	const FVector2D Icon28x14(28.0f, 14.0f);
 	const FVector2D Icon64x64(64.0f, 64.0f);
 
@@ -121,7 +123,13 @@ void FRemoteControlPanelStyle::Initialize()
 			.SetPadding(FMargin(1, 1, 1, 1));
 
 		StyleSet->Set("RemoteControlPanel.Switch", SwitchStyle);
-	}
+	} 
+
+	// Remote Control Logic UI
+	StyleSet->Set("RemoteControlPanel.Behaviours.Title", FCoreStyle::GetDefaultFontStyle("Regular", 12));
+	StyleSet->Set("RemoteControlPanel.Behaviours.BehaviourDescription", FCoreStyle::GetDefaultFontStyle("Regular", 10));
+	StyleSet->Set("RemoteControlPanel.Behaviours.CustomBlueprint", new IMAGE_BRUSH_SVG("Starship/MainToolbar/blueprints", Icon20x20));
+	StyleSet->Set("RemoteControlPanel.Actions.ValuePanelHeader", FCoreStyle::GetDefaultFontStyle("Bold", 12));
 
 	FSlateStyleRegistry::RegisterSlateStyle(*StyleSet.Get());
 }
@@ -155,6 +163,7 @@ FString FRemoteControlPanelStyle::InContent(const FString& RelativePath, const A
 
 #undef IMAGE_PLUGIN_BRUSH
 #undef IMAGE_BRUSH
+#undef IMAGE_BRUSH_SVG
 #undef BOX_BRUSH
 #undef BOX_PLUGIN_BRUSH
 #undef DEFAULT_FONT
diff --git a/Engine/Plugins/VirtualProduction/RemoteControl/Source/RemoteControlUI/Private/UI/SRemoteControlPanel.cpp b/Engine/Plugins/VirtualProduction/RemoteControl/Source/RemoteControlUI/Private/UI/SRemoteControlPanel.cpp
--- a/Engine/Plugins/VirtualProduction/RemoteControl/Source/RemoteControlUI/Private/UI/SRemoteControlPanel.cpp
+++ b/Engine/Plugins/VirtualProduction/RemoteControl/Source/RemoteControlUI/Private/UI/SRemoteControlPanel.cpp
@@ -2,11 +2,14 @@
 
 #include "SRemoteControlPanel.h"
 
+#include "Action/SRCActionPanel.h"
 #include "ActorEditorUtils.h"
 #include "AssetRegistry/AssetRegistryModule.h"
+#include "Behaviour/SRCBehaviourPanel.h"
 #include "ClassViewerFilter.h"
 #include "ClassViewerModule.h"
 #include "Commands/RemoteControlCommands.h"
+#include "Controller/SRCControllerPanel.h"
 #include "Editor.h"
 #include "Editor/EditorPerformanceSettings.h"
 #include "EditorFontGlyphs.h"
@@ -52,6 +55,8 @@
 #include "Templates/UnrealTypeTraits.h"
 #include "Toolkits/IToolkitHost.h"
 #include "ToolMenus.h"
+#include "UI/BaseLogicUI/RCLogicModeBase.h"
+#include "UI/BaseLogicUI/SRCLogicPanelListBase.h"
 #include "Widgets/DeclarativeSyntaxSupport.h"
 #include "Widgets/Docking/SDockTab.h"
 #include "Widgets/Images/SLayeredImage.h"
@@ -71,6 +76,8 @@ const FName SRemoteControlPanel::DefaultRemoteControlPanelToolBarName("RemoteCon
 const FName SRemoteControlPanel::AuxiliaryRemoteControlPanelToolBarName("RemoteControlPanel.AuxiliaryToolBar");
 const float SRemoteControlPanel::MinimumPanelWidth = 640.f;
 
+static TAutoConsoleVariable<bool> CVarRemoteControlEnableLogicUI(TEXT("RemoteControl.EnableLogicUI"), false, TEXT("Toggles visibility of Remote Control Logic User Interface (part of the Remote Control Preset asset)"));
+
 namespace RemoteControlPanelUtils
 {
 	bool IsExposableActor(AActor* Actor)
@@ -323,8 +330,7 @@ void SRemoteControlPanel::Construct(const FArguments& InArgs, URemoteControlPres
 			return 1.0f - Settings->TreeBindingSplitRatio;
 		}));
 
-	ChildSlot
-	[
+	TSharedRef<SWidget> ExposedPropertyPanel =
 		SNew(SVerticalBox)
 		+ SVerticalBox::Slot()
 		.Padding(3.f)
@@ -656,9 +662,92 @@ void SRemoteControlPanel::Construct(const FArguments& InArgs, URemoteControlPres
 		.Padding(FMargin(5.f, 8.f, 5.f, 5.f))
 		[
 			CreateCPUThrottleWarning()
-		]
+		];
+
+	TabManager = FGlobalTabmanager::New();
+	const TSharedPtr<FTabManager::FStack> ControlPanelStackManager = FTabManager::NewStack();
+	const constexpr TCHAR* ExposedPropertyTabName = TEXT("Expose Property Tab");
+	const constexpr TCHAR* ControllerTabName = TEXT("Controller Tab");
+	const FOnSpawnTab OnExposedPropertySpawnTab = FOnSpawnTab::CreateLambda([ExposedPropertyPanel] (const FSpawnTabArgs&)
+	{
+		return SNew(SDockTab)
+		[
+			ExposedPropertyPanel
+		];
+	});
+	const FOnSpawnTab OnControllerSpawnTab = FOnSpawnTab::CreateLambda([this] (const FSpawnTabArgs&)
+	{
+		return SNew(SDockTab)
+		[
+			SAssignNew(ControllerPanel, SRCControllerPanel, SharedThis(this))
+		];
+	});
+	const FCanSpawnTab CanSpawnPanelTab = FCanSpawnTab::CreateLambda([this](const FSpawnTabArgs& InArgs)
+	{
+		return true;
+	});
+	if (!TabManager->HasTabSpawner(ExposedPropertyTabName))
+	{
+		TabManager->RegisterTabSpawner(ExposedPropertyTabName, OnExposedPropertySpawnTab, CanSpawnPanelTab);
+	}
+	if (!TabManager->HasTabSpawner(ControllerTabName))
+	{
+		TabManager->RegisterTabSpawner(ControllerTabName, OnControllerSpawnTab, CanSpawnPanelTab);
+	}
+	
+	const TSharedRef<FTabManager::FLayout> Layout = FTabManager::NewLayout("RemoteControlPanelLayout1.0") 
+	->AddArea
+	(
+		FTabManager::NewPrimaryArea()
+		->SetOrientation(Orient_Vertical)
+		->Split
+		(
+			FTabManager::NewStack()
+			->SetSizeCoefficient(0.5f)
+			->AddTab(ControllerTabName, ETabState::OpenedTab)
+		)
+		->Split
+		(
+			FTabManager::NewStack()
+			->SetSizeCoefficient(0.5f)
+			->AddTab(ExposedPropertyTabName, ETabState::OpenedTab)
+		)
+	);
+		
+	// Make 3 Columns with Exposed Properties + Controllers / Behaviours / Actions
+	TSharedRef<SSplitter> SplitterBox = SNew(SSplitter)
+	.Orientation(Orient_Horizontal)
+	+ SSplitter::Slot()
+	.Value(0.33)
+	[
+		TabManager->RestoreFrom(Layout, nullptr).ToSharedRef()
+	]
+	+ SSplitter::Slot()
+	.Value(0.33)
+	[
+		SAssignNew(BehaviourPanel, SRCBehaviourPanel, SharedThis(this))
+	]
+	+ SSplitter::Slot()
+	.Value(0.33)
+	[
+		SAssignNew(ActionPanel, SRCActionPanel, SharedThis(this))
 	];
 
+	if (CVarRemoteControlEnableLogicUI.GetValueOnAnyThread())
+	{
+		ChildSlot
+		[
+			SplitterBox
+		];
+	}
+	else
+	{
+		ChildSlot
+			[
+				ExposedPropertyPanel
+			];
+	}
+
 	RegisterEvents();
 	CacheLevelClasses();
 	Refresh();
@@ -1803,8 +1892,41 @@ void SRemoteControlPanel::OnToggleLogicEditor() const
 {
 }
 
-void SRemoteControlPanel::DeleteEntity_Execute() const
+TSharedPtr<class SRCLogicPanelBase> SRemoteControlPanel::GetActiveLogicPanel() const
+{
+	if (ControllerPanel->IsListFocused())
+	{
+		return ControllerPanel;
+	}
+	else if (BehaviourPanel->IsListFocused())
+	{
+		return BehaviourPanel;
+	}
+	else if (ActionPanel->IsListFocused())
+	{
+		return ActionPanel;
+	}
+	
+	return nullptr;
+}
+
+void SRemoteControlPanel::DeleteEntity_Execute()
 {
+	// Currently used  as common entry point of Delete UI command for both RC Entity and Logic Items.
+	// This could potentially be moved out if the Logic panels are moved to a separate tab.
+	
+	// ~ Delete Logic Item ~
+	// 
+	// If the user focus is currently active on a Logic panel then route the Delete command to it and return.
+	if (TSharedPtr<class SRCLogicPanelBase> ActivePanel = GetActiveLogicPanel())
+	{
+		ActivePanel->DeleteSelectedPanelItem();
+
+		return; // handled
+	}
+
+	// ~ Delete Entity ~
+
 	if (SelectedEntity->GetRCType() == SRCPanelTreeNode::FieldChild) // Field Child does not contain entity ID, that is why it should not be processed
 	{
 		return;
@@ -1828,6 +1950,11 @@ void SRemoteControlPanel::DeleteEntity_Execute() const
 
 bool SRemoteControlPanel::CanDeleteEntity() const
 {
+	if (TSharedPtr<class SRCLogicPanelBase> ActivePanel = GetActiveLogicPanel())
+	{
+		return true; // User has focus on a logic panel
+	}
+
 	if (SelectedEntity.IsValid() && Preset.IsValid() && bIsInEditMode)
 	{
 		// Do not allow default group to be deleted.
diff --git a/Engine/Plugins/VirtualProduction/RemoteControl/Source/RemoteControlUI/Private/UI/SRemoteControlPanel.h b/Engine/Plugins/VirtualProduction/RemoteControl/Source/RemoteControlUI/Private/UI/SRemoteControlPanel.h
--- a/Engine/Plugins/VirtualProduction/RemoteControl/Source/RemoteControlUI/Private/UI/SRemoteControlPanel.h
+++ b/Engine/Plugins/VirtualProduction/RemoteControl/Source/RemoteControlUI/Private/UI/SRemoteControlPanel.h
@@ -34,6 +34,12 @@ class SSearchBox;
 class STextBlock;
 class URemoteControlPreset;
 
+class URCController;
+class URCBehaviour;
+class URCAction;
+class FRCControllerModel;
+class FRCBehaviourModel;
+
 
 DECLARE_DELEGATE_TwoParams(FOnEditModeChange, TSharedPtr<SRemoteControlPanel> /* Panel */, bool /* bEditModeChange */);
 
@@ -49,6 +55,17 @@ class SRemoteControlPanel : public SCompoundWidget, public FSelfRegisteringEdito
 	SLATE_END_ARGS()
 
 public:
+	// Remote Control Logic Delegates
+	DECLARE_MULTICAST_DELEGATE_OneParam(FOnControllerAdded, const FName& /* InPropertyName */);
+	DECLARE_MULTICAST_DELEGATE_OneParam(FOnControllerSelectionChanged, TSharedPtr<FRCControllerModel> /* InControllerItem */);
+	DECLARE_MULTICAST_DELEGATE_OneParam(FOnBehaviourAdded, URCBehaviour* /* InBehaviour */);
+	DECLARE_MULTICAST_DELEGATE_OneParam(FOnBehaviourSelectionChanged, TSharedPtr<FRCBehaviourModel> /* InBehaviourItem */);
+	DECLARE_MULTICAST_DELEGATE_OneParam(FOnActionAdded, URCAction* /* InAction */);
+
+	DECLARE_MULTICAST_DELEGATE(FOnEmptyControllers);
+	DECLARE_MULTICAST_DELEGATE(FOnEmptyBehaviours);
+	DECLARE_MULTICAST_DELEGATE(FOnEmptyActions);
+	
 	void Construct(const FArguments& InArgs, URemoteControlPreset* InPreset, TSharedPtr<IToolkitHost> InToolkitHost);
 	~SRemoteControlPanel();
 
@@ -250,7 +267,7 @@ private:
 	void OnToggleLogicEditor() const;
 	
 	/** Called when user attempts to delete a group/exposed entity. */
-	void DeleteEntity_Execute() const;
+	void DeleteEntity_Execute();
 
 	/** Called to test if user is able to delete a group/exposed entity. */
 	bool CanDeleteEntity() const;
@@ -317,6 +334,35 @@ private:
 	/** Holds a shared pointer reference to the active entity that is selected. */
 	TSharedPtr<SRCPanelTreeNode> SelectedEntity;
 
+	// ~ Remote Control Logic Panels ~
+
+	/** Controller panel UI widget for Remote Control Logic*/
+	TSharedPtr<class SRCControllerPanel> ControllerPanel;
+
+	/** Behaviour panel UI widget for Remote Control Logic*/
+	TSharedPtr<class SRCBehaviourPanel> BehaviourPanel;
+
+	/** Action panel UI widget for Remote Control Logic*/
+	TSharedPtr<class SRCActionPanel> ActionPanel;
+
+
 public:
 	static const float MinimumPanelWidth;
-};
+
+	/** Tab Manager used to handle spawning of Controller and Exposed Property tabs*/
+	TSharedPtr<FGlobalTabmanager> TabManager;
+
+public:
+	// Global Delegates for Remote Control Logic
+	FOnControllerAdded OnControllerAdded;	
+	FOnBehaviourAdded OnBehaviourAdded;
+	FOnActionAdded OnActionAdded;
+	FOnControllerSelectionChanged OnControllerSelectionChanged;
+	FOnBehaviourSelectionChanged OnBehaviourSelectionChanged;
+	FOnEmptyControllers OnEmptyControllers;
+	FOnEmptyBehaviours OnEmptyBehaviours;
+	FOnEmptyActions OnEmptyActions;
+
+private:
+	TSharedPtr<class SRCLogicPanelBase> GetActiveLogicPanel() const;
+};
\ No newline at end of file
diff --git a/Engine/Plugins/VirtualProduction/RemoteControl/Source/RemoteControlUI/RemoteControlUI.Build.cs b/Engine/Plugins/VirtualProduction/RemoteControl/Source/RemoteControlUI/RemoteControlUI.Build.cs
--- a/Engine/Plugins/VirtualProduction/RemoteControl/Source/RemoteControlUI/RemoteControlUI.Build.cs
+++ b/Engine/Plugins/VirtualProduction/RemoteControl/Source/RemoteControlUI/RemoteControlUI.Build.cs
@@ -15,13 +15,16 @@ public class RemoteControlUI : ModuleRules
 				"ApplicationCore",
 				"AssetRegistry",
 				"AssetTools",
+				"BlueprintGraph",
 				"ClassViewer",
 				"Core",
 				"CoreUObject",
+				"ContentBrowser",
 				"DeveloperSettings",
 				"EditorWidgets",
 				"EditorStyle",
 				"Engine",
+				"GraphEditor",
 				"HotReload",
 				"InputCore",
 				"MainFrame",
@@ -30,10 +33,12 @@ public class RemoteControlUI : ModuleRules
 				"PropertyEditor",
 				"RemoteControl",
 				"RemoteControlCommon",
+				"RemoteControlLogic",				
 				"RemoteControlProtocolWidgets",
 				"SceneOutliner",
 				"Slate",
 				"SlateCore",
+				"StructUtils",
 				"ToolMenus",
 				"ToolWidgets",
 				"UnrealEd"
