[UE-155440] Fixing localization namespace collision

diff --git a/Engine/Plugins/VirtualProduction/RemoteControl/Source/RemoteControlUI/Private/Factories/RCBehaviourNodeFactory.cpp b/Engine/Plugins/VirtualProduction/RemoteControl/Source/RemoteControlUI/Private/Factories/RCBehaviourNodeFactory.cpp
--- a/Engine/Plugins/VirtualProduction/RemoteControl/Source/RemoteControlUI/Private/Factories/RCBehaviourNodeFactory.cpp
+++ b/Engine/Plugins/VirtualProduction/RemoteControl/Source/RemoteControlUI/Private/Factories/RCBehaviourNodeFactory.cpp
@@ -6,7 +6,7 @@
 #include "Behaviour/RCBehaviourBlueprintNode.h"
 #include "Kismet2/KismetEditorUtilities.h"
 
-#define LOCTEXT_NAMESPACE "LevelSnapshotFilterFactory"
+#define LOCTEXT_NAMESPACE "RCBehaviourNodeFactory"
 
 URCBehaviourNodeFactory::URCBehaviourNodeFactory()
 {
