Fix non-unity compile error.

diff --git a/Engine/Plugins/VirtualProduction/RemoteControl/Source/RemoteControlUI/Private/UI/SRCPanelDragHandle.h b/Engine/Plugins/VirtualProduction/RemoteControl/Source/RemoteControlUI/Private/UI/SRCPanelDragHandle.h
--- a/Engine/Plugins/VirtualProduction/RemoteControl/Source/RemoteControlUI/Private/UI/SRCPanelDragHandle.h
+++ b/Engine/Plugins/VirtualProduction/RemoteControl/Source/RemoteControlUI/Private/UI/SRCPanelDragHandle.h
@@ -7,6 +7,7 @@
 #include "Widgets/Images/SImage.h"
 #include "Styling/CoreStyle.h"
 #include "Input/DragAndDrop.h"
+#include "DragAndDrop/DecoratedDragDropOp.h"
 
 class SRCPanelGroup;
 class FExposedEntityDragDrop;
