Fix remote control deprecation warning

diff --git a/Engine/Plugins/VirtualProduction/RemoteControl/Source/RemoteControl/Private/RemoteControlPreset.cpp b/Engine/Plugins/VirtualProduction/RemoteControl/Source/RemoteControl/Private/RemoteControlPreset.cpp
--- a/Engine/Plugins/VirtualProduction/RemoteControl/Source/RemoteControl/Private/RemoteControlPreset.cpp
+++ b/Engine/Plugins/VirtualProduction/RemoteControl/Source/RemoteControl/Private/RemoteControlPreset.cpp
@@ -410,267 +410,6 @@ FRemoteControlPresetGroup& FRemoteControlPresetLayout::CreateGroupInternal(FName
 	PRAGMA_ENABLE_DEPRECATION_WARNINGS
 }
 
-TOptional<FRemoteControlProperty> FRemoteControlTarget::ExposeProperty(FRCFieldPathInfo FieldPathInfo, const FString& DesiredDisplayName, FGuid GroupId, bool bAppendAliasToLabel)
-{
-	FName FieldLabel = Owner->GenerateUniqueFieldLabel(Alias, DesiredDisplayName, bAppendAliasToLabel);
-	FRemoteControlProperty RCField(Owner.Get(), FieldLabel, MoveTemp(FieldPathInfo), {});
-	ExposedProperties.Add(RCField);
-
-	URemoteControlPreset::FExposeInfo ExposeInfo;
-	ExposeInfo.Alias = Alias;
-	ExposeInfo.FieldId = RCField.GetId();
-	ExposeInfo.LayoutGroupId = GroupId;
-	Owner->OnExpose(MoveTemp(ExposeInfo));
-
-	return RCField;
-}
-
-TOptional<FRemoteControlProperty> FRemoteControlTarget::ExposeProperty(FRCFieldPathInfo FieldPathInfo, TArray<FString> ComponentHierarchy, const FString& DesiredDisplayName, FGuid GroupId)
-{
-	return ExposeProperty(MoveTemp(FieldPathInfo), DesiredDisplayName, MoveTemp(GroupId));
-}
-
-TOptional<FRemoteControlFunction> FRemoteControlTarget::ExposeFunction(FString RelativeFieldPath, const FString& DesiredDisplayName, FGuid GroupId, bool bAppendAliasToLabel)
-{
-	TOptional<FRemoteControlFunction> RCFunction;
-
-	FRCFieldPathInfo Path{ MoveTemp(RelativeFieldPath) };
-	if (UFunction* Function = Class->FindFunctionByName(Path.GetFieldName()))
-	{
-		FName FieldLabel = Owner->GenerateUniqueFieldLabel(Alias, DesiredDisplayName, bAppendAliasToLabel);
-		RCFunction = FRemoteControlFunction{ Owner.Get(), FieldLabel, MoveTemp(Path), Function, {}};
-		ExposedFunctions.Add(*RCFunction);
-
-		URemoteControlPreset::FExposeInfo ExposeInfo;
-		ExposeInfo.Alias = Alias;
-		ExposeInfo.FieldId = RCFunction->GetId();
-		ExposeInfo.LayoutGroupId = GroupId;
-		Owner->OnExpose(MoveTemp(ExposeInfo));
-	}
-
-	return RCFunction;
-}
-
-void FRemoteControlTarget::Unexpose(FGuid FieldId)
-{
-	RemoveField<FRemoteControlProperty>(ExposedProperties, FieldId);
-	RemoveField<FRemoteControlFunction>(ExposedFunctions, FieldId);
-
-	Owner->OnUnexpose(FieldId);
-}
-
-FName FRemoteControlTarget::FindFieldLabel(FName FieldName) const
-{
-	FName Label = NAME_None;
-
-	for (const FRemoteControlProperty& RCProperty : ExposedProperties)
-	{
-		if (RCProperty.FieldName == FieldName)
-		{
-			Label = RCProperty.GetLabel();
-		}
-	}
-
-	if (Label == NAME_None)
-	{
-		for (const FRemoteControlFunction& RCFunction : ExposedFunctions)
-		{
-			if (RCFunction.FieldName == FieldName)
-			{
-				Label = RCFunction.GetLabel();
-			}
-		}
-	}
-
-	return Label;
-}
-
-FName FRemoteControlTarget::FindFieldLabel(const FRCFieldPathInfo& Path) const
-{
-	FName Label = NAME_None;
-
-	for (const FRemoteControlProperty& RCProperty : ExposedProperties)
-	{
-		if (RCProperty.FieldPathInfo.IsEqual(Path))
-		{
-			Label = RCProperty.GetLabel();
-		}
-	}
-
-	if (Label == NAME_None)
-	{
-		for (const FRemoteControlFunction& RCFunction : ExposedFunctions)
-		{
-			if (RCFunction.FieldPathInfo.IsEqual(Path))
-			{
-				Label = RCFunction.GetLabel();
-			}
-		}
-	}
-
-	return Label;
-}
-
-TOptional<FRemoteControlProperty> FRemoteControlTarget::GetProperty(FGuid PropertyId) const
-{
-	TOptional<FRemoteControlProperty> Field;
-	if (const FRemoteControlProperty* RCField = ExposedProperties.FindByHash(GetTypeHash(PropertyId), PropertyId))
-	{
-		Field = *RCField;
-	}
-	return Field;
-}
-
-TOptional<FRemoteControlFunction> FRemoteControlTarget::GetFunction(FGuid FunctionId) const
-{
-	TOptional<FRemoteControlFunction> Field;
-	if (const FRemoteControlFunction* RCField = ExposedFunctions.FindByHash(GetTypeHash(FunctionId), FunctionId))
-	{
-		Field = *RCField;
-	}
-	return Field;
-}
-
-PRAGMA_DISABLE_DEPRECATION_WARNINGS
-TOptional<FExposedProperty> FRemoteControlTarget::ResolveExposedProperty(FGuid PropertyId) const
-{
-	TOptional<FExposedProperty> OptionalExposedProperty;
-	ensure(Class);
-
-	if (TOptional<FRemoteControlProperty> RCProperty = GetProperty(PropertyId))
-	{
-		TArray<UObject*> FieldOwners = RCProperty->GetBoundObjects();
-		if (FieldOwners.Num() && FieldOwners[0])
-		{
-			//Always resolve exposed property. We might have been pointing to an array element that has been deleted
-			if (RCProperty->FieldPathInfo.Resolve(FieldOwners[0]))
-			{
-				OptionalExposedProperty = FExposedProperty{ RCProperty->FieldPathInfo.GetResolvedData().Field, FieldOwners };
-			}
-		}
-	}
-
-	return OptionalExposedProperty;
-}
-PRAGMA_ENABLE_DEPRECATION_WARNINGS
-
-PRAGMA_DISABLE_DEPRECATION_WARNINGS
-TOptional<FExposedFunction> FRemoteControlTarget::ResolveExposedFunction(FGuid FunctionId) const
-{
-	TOptional<FExposedFunction> OptionalExposedFunction;
-	ensure(Class);
-	if (TOptional<FRemoteControlFunction> RCFunction = GetFunction(FunctionId))
-	{
-		FExposedFunction ExposedFunction;
-		ExposedFunction.Function = RCFunction->GetFunction();
-		ExposedFunction.DefaultParameters = RCFunction->FunctionArguments;
-		ExposedFunction.OwnerObjects = ResolveBoundObjects();
-		OptionalExposedFunction = MoveTemp(ExposedFunction);
-	}
-
-	return OptionalExposedFunction;
-}
-PRAGMA_ENABLE_DEPRECATION_WARNINGS
-
-TArray<UObject*> FRemoteControlTarget::ResolveBoundObjects() const
-{
-	TArray<UObject*> ResolvedObjects;
-	ResolvedObjects.Reserve(Bindings.Num());
-	for (const FSoftObjectPath& Path : Bindings)
-	{
-		if (UObject* Obj = Path.ResolveObject())
-		{
-#if WITH_EDITOR
-			if (Obj->IsA<AActor>() && GEditor && GEditor->PlayWorld)
-			{
-				if (AActor* SimWorldActor = EditorUtilities::GetSimWorldCounterpartActor(Cast<AActor>(Obj)))
-				{
-					Obj = SimWorldActor;
-				}
-			}
-#endif		
-			ResolvedObjects.Add(Obj);
-		}
-	}
-
-	return ResolvedObjects;
-}
-
-void FRemoteControlTarget::BindObjects(const TArray<UObject*>& ObjectsToBind)
-{
-	if (CanBindObjects(ObjectsToBind))
-	{
-		for (UObject* TargetObject : ObjectsToBind)
-		{
-			Bindings.Add(FSoftObjectPath{ TargetObject });
-		}
-	}
-}
-
-bool FRemoteControlTarget::HasBoundObjects(const TArray<UObject*>& ObjectsToTest) const
-{
-	const TArray<UObject*>& TargetObjects = ResolveBoundObjects();
-	for (UObject* Object : ObjectsToTest)
-	{
-		if (!TargetObjects.Contains(Object))
-		{
-			return false;
-		}
-	}
-
-	return true;
-}
-
-bool FRemoteControlTarget::HasBoundObject(const UObject* ObjectToTest) const
-{
-	return ResolveBoundObjects().Contains(ObjectToTest);
-}
-
-bool FRemoteControlTarget::CanBindObjects(const TArray<UObject*>& ObjectsToTest) const
-{
-	if (!Class && Bindings.Num() == 0)
-	{
-		// If no target class is set, then anything can be exposed.
-		return true;
-	}
-
-	UClass* OwnersCommonBase = FindCommonBase(ObjectsToTest);
-	return OwnersCommonBase && OwnersCommonBase->IsChildOf(Class);
-}
-
-PRAGMA_DISABLE_DEPRECATION_WARNINGS
-FProperty* FRemoteControlTarget::FindPropertyRecursive(UStruct* Container, TArray<FString>& DesiredPropertyPath) const
-{
-	if (DesiredPropertyPath.Num() <= 0)
-	{
-		return nullptr;
-	}
-
-	const FString DesiredPropertyName = DesiredPropertyPath[0];
-	for (TFieldIterator<FProperty> It(Container, EFieldIteratorFlags::IncludeSuper, EFieldIteratorFlags::ExcludeDeprecated, EFieldIteratorFlags::ExcludeInterfaces); It; ++It)
-	{
-		if (It->GetName() == DesiredPropertyName)
-		{
-			if (DesiredPropertyPath.Num() == 1)
-			{
-				return *It;
-			}
-			else
-			{
-				//We're looking for nested property so might be a struct or array
-				if (FStructProperty* StructProp = CastField<FStructProperty>(*It))
-				{
-					DesiredPropertyPath.RemoveAt(0, 1, false);
-					return FindPropertyRecursive(StructProp->Struct, DesiredPropertyPath);
-				}
-			}
-		}
-	}
-
-	return nullptr;
-}
-PRAGMA_ENABLE_DEPRECATION_WARNINGS
-
 URemoteControlPreset::URemoteControlPreset()
 	: Layout(FRemoteControlPresetLayout{ this })
 	, PresetId(FGuid::NewGuid())
@@ -709,8 +448,6 @@ void URemoteControlPreset::PostLoad()
 
 	RegisterDelegates();
 
-	CacheFieldsData();
-
 	CacheFieldLayoutData();
 
 	InitializeEntitiesMetadata();
@@ -739,25 +476,6 @@ void URemoteControlPreset::BeginDestroy()
 	Super::BeginDestroy();
 }
 
-FName URemoteControlPreset::GetOwnerAlias(FGuid FieldId) const
-{
-	if (const FRCCachedFieldData* CachedData = FieldCache.Find(FieldId))
-	{
-		return CachedData->OwnerObjectAlias;
-	}
-
-	return NAME_None;
-}
-
-FGuid URemoteControlPreset::GetFieldId(FName FieldLabel) const
-{
-	if (const FGuid* Guid = NameToGuidMap.Find(FieldLabel))
-	{
-		return *Guid;
-	}
-	return FGuid();
-}
-
 TWeakPtr<FRemoteControlActor> URemoteControlPreset::ExposeActor(AActor* Actor, FRemoteControlPresetExposeArgs Args)
 {
 	check(Actor);
@@ -1375,59 +1093,6 @@ void URemoteControlPreset::Unexpose(const FGuid& EntityId)
 	}
 }
 
-FName URemoteControlPreset::CreateTarget(const TArray<UObject*>& TargetObjects)
-{
-	PRAGMA_DISABLE_DEPRECATION_WARNINGS
-	return CreateAndGetTarget(TargetObjects).Alias;
-	PRAGMA_ENABLE_DEPRECATION_WARNINGS
-}
-
-PRAGMA_DISABLE_DEPRECATION_WARNINGS
-FRemoteControlTarget& URemoteControlPreset::CreateAndGetTarget(const TArray<UObject*>& TargetObjects)
-{
-	check(TargetObjects.Num() != 0);
-
-	FName Alias = GenerateAliasForObjects(TargetObjects);
-	FRemoteControlTarget Target{ this };
-	Target.BindObjects(TargetObjects);
-	Target.Alias = Alias;
-
-	TArray<UClass*> Classes;
-	Classes.Reserve(TargetObjects.Num());
-	Algo::Transform(TargetObjects, Classes, [](const UObject* Object) { return Object->GetClass(); });
-	UClass* CommonClass = UClass::FindCommonBase(Classes);
-
-	Target.Class = CommonClass;
-
-	return RemoteControlTargets.Add(Alias, MoveTemp(Target));
-}
-PRAGMA_ENABLE_DEPRECATION_WARNINGS
-
-void URemoteControlPreset::DeleteTarget(FName TargetName)
-{	
-PRAGMA_DISABLE_DEPRECATION_WARNINGS
-	if (FRemoteControlTarget* Target = RemoteControlTargets.Find(TargetName))
-	{
-		for (auto It = Target->ExposedProperties.CreateConstIterator(); It; ++It)
-		{
-			Unexpose(It->GetLabel());
-		}
-	}
-
-	RemoteControlTargets.Remove(TargetName);
-PRAGMA_ENABLE_DEPRECATION_WARNINGS
-}
-
-void URemoteControlPreset::RenameTarget(FName TargetName, FName NewTargetName)
-{
-PRAGMA_DISABLE_DEPRECATION_WARNINGS
-	FRemoteControlTarget Target;
-	RemoteControlTargets.RemoveAndCopyValue(TargetName, Target);
-	Target.Alias = NewTargetName;
-	RemoteControlTargets.Add(NewTargetName, MoveTemp(Target));
-PRAGMA_ENABLE_DEPRECATION_WARNINGS
-}
-
 void URemoteControlPreset::CacheLayoutData()
 {
 	CacheFieldLayoutData();
@@ -1473,26 +1138,9 @@ void URemoteControlPreset::Serialize(FArchive& Ar)
 	Super::Serialize(Ar);
 
 	Ar.UsingCustomVersion(FRemoteControlObjectVersion::GUID);
-	int32 CustomVersion = Ar.CustomVer(FRemoteControlObjectVersion::GUID);
 
 	if (Ar.IsLoading())
 	{
-		if (CustomVersion < FRemoteControlObjectVersion::BeforeCustomVersionWasAdded)
-		{
-			// Create new targets and remove the component chain from the fields.
-			ConvertFieldsToRemoveComponentChain();
-		}
-
-		if (CustomVersion < FRemoteControlObjectVersion::ConvertRCFieldsToRCEntities)
-		{
-			ConvertFieldsToEntities();
-		}
-
-		if (CustomVersion < FRemoteControlObjectVersion::ConvertTargetsToBindings)
-		{
-			ConvertTargetsToBindings();
-		}
-
 		if (!PresetId.IsValid())
 		{
 			PresetId = FGuid::NewGuid();
@@ -1505,160 +1153,6 @@ FRemoteControlField* URemoteControlPreset::GetFieldPtr(FGuid FieldId)
 	return Registry->GetExposedEntity<FRemoteControlField>(FieldId).Get();
 }
 
-void URemoteControlPreset::ConvertFieldsToRemoveComponentChain()
-{
-PRAGMA_DISABLE_DEPRECATION_WARNINGS
-#if WITH_EDITOR
-	CacheFieldsData();
-
-	auto GroupPropertiesByObjects = [this]()
-	{
-		TMap<UObject*, TSet<FRemoteControlProperty>> ObjectAndProperties;
-
-		for (const TTuple<FName, FRemoteControlTarget>& Tuple : RemoteControlTargets)
-		{
-			const FRemoteControlTarget& Target = Tuple.Value;
-			for (const FRemoteControlProperty& Property : Target.ExposedProperties)
-			{
-				if (Property.ComponentHierarchy_DEPRECATED.Num())
-				{
-					TArray<UObject*> TargetObjects = Property.GetBoundObjects();
-					for (UObject* TargetObject : TargetObjects)
-					{
-						ObjectAndProperties.FindOrAdd(TargetObject).Add(Property);
-					}
-				}
-			}
-		}
-
-		return ObjectAndProperties;
-	};
-
-	auto MovePropertiesToTarget = [this](TSet<FRemoteControlProperty>&& Properties, FRemoteControlTarget& DestinationTarget)
-	{
-		for (FRemoteControlProperty& Property : Properties)
-		{
-			FRCCachedFieldData& CachedPropData = FieldCache.FindChecked(Property.GetId());
-			FRemoteControlTarget& OwnerTarget = RemoteControlTargets.FindChecked(CachedPropData.OwnerObjectAlias);
-
-			if (OwnerTarget.Alias != DestinationTarget.Alias)
-			{
-				OwnerTarget.RemoveField<FRemoteControlProperty>(OwnerTarget.ExposedProperties, Property.GetId());
-				Property.ComponentHierarchy_DEPRECATED.Empty();
-				DestinationTarget.ExposedProperties.Add(MoveTemp(Property));
-				CachedPropData.OwnerObjectAlias = DestinationTarget.Alias;
-			}
-		}
-	};
-
-	auto FindExistingDestinationTarget = [this](UObject* Object) -> FRemoteControlTarget*
-	{
-		TTuple<FName, FRemoteControlTarget>* Target = Algo::FindByPredicate(RemoteControlTargets,
-			[Object](const TTuple<FName, FRemoteControlTarget>& Target)
-			{
-				return Target.Value.HasBoundObject(Object);
-			});
-
-		return Target ? &Target->Value : nullptr;
-	};
-
-	auto RegroupPropertiesInTargets = [this, &MovePropertiesToTarget, &FindExistingDestinationTarget](TMap<UObject*, TSet<FRemoteControlProperty>>&& ObjectAndProperties)
-	{
-		
-		for (TTuple<UObject*, TSet<FRemoteControlProperty>>& ObjectMapEntry : ObjectAndProperties)
-		{
-			FRemoteControlTarget* DestinationTarget = FindExistingDestinationTarget(ObjectMapEntry.Key);
-
-			if (!DestinationTarget)
-			{
-				DestinationTarget = &CreateAndGetTarget({ ObjectMapEntry.Key });
-			}
-
-			MovePropertiesToTarget(MoveTemp(ObjectMapEntry.Value), *DestinationTarget);
-		}
-	};
-
-	auto RemoveEmptyTargets = [this]()
-	{
-		for (auto It = RemoteControlTargets.CreateIterator(); It; ++It)
-		{
-			if (It.Value().ExposedProperties.Num() == 0 && It.Value().ExposedFunctions.Num() == 0)
-			{
-				It.RemoveCurrent();
-			}
-		}
-	};
-
-	// Convert fields to not use the component chain anymore.
-	RegroupPropertiesInTargets(GroupPropertiesByObjects());
-	RemoveEmptyTargets();
-#endif
-PRAGMA_ENABLE_DEPRECATION_WARNINGS
-}
-
-void URemoteControlPreset::ConvertFieldsToEntities()
-{
-#if WITH_EDITOR
-	// Convert properties and functions to inherit from FRemoteControlEntities while preserving their old data.
-	PRAGMA_DISABLE_DEPRECATION_WARNINGS
-	for (TTuple<FName, FRemoteControlTarget>& Tuple : RemoteControlTargets)
-	{
-		for (FRemoteControlProperty& Property : Tuple.Value.ExposedProperties)
-		{
-			Property.Owner = this;
-		}
-
-		for (FRemoteControlFunction& Function : Tuple.Value.ExposedFunctions)
-		{
-			Function.Owner = this;
-		}
-	}
-	PRAGMA_ENABLE_DEPRECATION_WARNINGS
-#endif
-}
-
-void URemoteControlPreset::ConvertTargetsToBindings()
-{
-#if WITH_EDITOR
-	// Convert targets to bindings, and put everything in the registry.
-	PRAGMA_DISABLE_DEPRECATION_WARNINGS
-	for (TTuple<FName, FRemoteControlTarget>& Tuple : RemoteControlTargets)
-	{
-		TArray<URemoteControlBinding*> NewBindings;
-		NewBindings.Reserve(Tuple.Value.Bindings.Num());
-		for (const FSoftObjectPath& Path : Tuple.Value.Bindings)
-		{
-			// Load the asset 
-			if (URemoteControlBinding* Binding = FindOrAddBinding(TSoftObjectPtr<UObject>{Path}))
-			{
-				Binding->Name = Tuple.Value.Alias.ToString();
-				NewBindings.Add(Binding);
-			}
-		}
-
-		for (FRemoteControlProperty& Property : Tuple.Value.ExposedProperties)
-		{
-			Property.Bindings.Append(NewBindings);
-			Registry->AddExposedEntity(MoveTemp(Property), FRemoteControlProperty::StaticStruct());
-		}
-
-		for (FRemoteControlFunction& Function : Tuple.Value.ExposedFunctions)
-		{
-			Function.Bindings.Append(NewBindings);
-			Registry->AddExposedEntity(MoveTemp(Function), FRemoteControlFunction::StaticStruct());
-		}
-	}
-
-	RemoteControlTargets.Reset();
-
-	for (TSharedPtr<FRemoteControlActor> RCActor : Registry->GetExposedEntities<FRemoteControlActor>())
-	{
-		RCActor->Bindings = { FindOrAddBinding(TSoftObjectPtr<UObject>{ RCActor->Path }) }; 
-	}
-	PRAGMA_ENABLE_DEPRECATION_WARNINGS
-#endif
-}
-
 TSharedPtr<const FRemoteControlEntity> URemoteControlPreset::FindEntityById(const FGuid& EntityId, const UScriptStruct* EntityType) const
 {
 	return Registry->GetExposedEntity(EntityId, EntityType);
@@ -1727,47 +1221,6 @@ TOptional<FRemoteControlField> URemoteControlPreset::GetField(FGuid FieldId) con
 	return Field;
 }
 
-FName URemoteControlPreset::GenerateAliasForObjects(const TArray<UObject*>& Objects)
-{
-	auto Functor = [this](FName Candidate) { return RemoteControlTargets.Contains(Candidate); };
-
-	if (Objects.Num() == 1 && Objects[0])
-	{
-		UObject* TargetObject = Objects[0];
-		if (TargetObject->HasAnyFlags(RF_ClassDefaultObject | RF_ArchetypeObject))
-		{
-			return MakeUniqueName(FName(*TargetObject->GetName().RightChop(DefaultObjectPrefix.Len())), Functor);
-		}
-		else
-		{
-			FName Name;
-			if (TargetObject->IsA<UActorComponent>())
-			{
-				UActorComponent* Component = Cast<UActorComponent>(TargetObject);
-				Name = *Component->GetPathName(Component->GetOwner()->GetOuter());
-			}
-			else
-			{
-				Name = TargetObject->GetFName();
-			}
-
-			return MakeUniqueName(Name, Functor);
-		}
-	}
-	else
-	{
-		ensure(false);
-		return NAME_None;
-	}
-}
-
-FName URemoteControlPreset::GenerateUniqueFieldLabel(FName Alias, const FString& BaseName, bool bAppendAlias)
-{
-	auto NamePoolContainsFunctor = [this](FName Candidate) { return Registry->GetExposedEntity(GetExposedEntityId(Candidate)).IsValid(); };
-	FName BaseCandidate = bAppendAlias ? FName(*FString::Printf(TEXT("%s (%s)"), *BaseName, *Alias.ToString())) : FName(*FString::Printf(TEXT("%s"), *BaseName));
-	return MakeUniqueName(BaseCandidate, NamePoolContainsFunctor);
-}
-
 void URemoteControlPreset::OnExpose(const FExposeInfo& Info)
 {
 	FRCCachedFieldData CachedData;
@@ -1812,30 +1265,6 @@ void URemoteControlPreset::OnUnexpose(FGuid UnexposedFieldId)
 	}
 }
 
-void URemoteControlPreset::CacheFieldsData()
-{
-	if (FieldCache.Num())
-	{
-		return;
-	}
-
-	PRAGMA_DISABLE_DEPRECATION_WARNINGS
-	for (TTuple<FName, FRemoteControlTarget>& Target : RemoteControlTargets)
-	{
-		FieldCache.Reserve(FieldCache.Num() + Target.Value.ExposedProperties.Num() + Target.Value.ExposedFunctions.Num());
-		auto CacheField = [this, TargetName = Target.Key](const FRemoteControlField& Field)
-		{
-			FRCCachedFieldData& CachedData = FieldCache.FindOrAdd(Field.GetId());
-			CachedData.OwnerObjectAlias = TargetName;
-			NameToGuidMap.Add(Field.GetLabel(), Field.GetId());
-		};
-
-		Algo::ForEach(Target.Value.ExposedProperties, CacheField);
-		Algo::ForEach(Target.Value.ExposedFunctions, CacheField);
-	}
-	PRAGMA_ENABLE_DEPRECATION_WARNINGS
-}
-
 void URemoteControlPreset::CacheFieldLayoutData()
 {
 	for (FRemoteControlPresetGroup& Group : Layout.AccessGroups())
diff --git a/Engine/Plugins/VirtualProduction/RemoteControl/Source/RemoteControl/Public/RemoteControlPreset.h b/Engine/Plugins/VirtualProduction/RemoteControl/Source/RemoteControl/Public/RemoteControlPreset.h
--- a/Engine/Plugins/VirtualProduction/RemoteControl/Source/RemoteControl/Public/RemoteControlPreset.h
+++ b/Engine/Plugins/VirtualProduction/RemoteControl/Source/RemoteControl/Public/RemoteControlPreset.h
@@ -297,180 +297,7 @@ private:
 };
 
 /**
- * Represents objects grouped under a single alias that contain exposed functions and properties.
- */
-struct UE_DEPRECATED(4.27, "FRemoteControlTarget is deprecated. Expose properties directly on a remote control preset instead.") FRemoteControlTarget;
-USTRUCT()
-struct REMOTECONTROL_API FRemoteControlTarget
-{
-public:
-	GENERATED_BODY()
-
-	FRemoteControlTarget() = default;
-
-	FRemoteControlTarget(class URemoteControlPreset* InOwner)
-		: Owner(InOwner)
-		{}
-
-	/**
-	* Expose a property in this target.
-	* @param FieldPathInfo the path data from the owner object, including component chain, to this field. (ie. LightComponent0.Intensity)
-	* @param DesiredDisplayName the display name desired for this control. If the name is not unique preset-wise, a number will be appended.
-	*/
-	UE_DEPRECATED(4.27, "A component hierarchy is no longer needed to expose properties, use the overloaded function instead.")
-	TOptional<FRemoteControlProperty> ExposeProperty(FRCFieldPathInfo FieldPathInfo, TArray<FString> ComponentHierarchy, const FString& DesiredDisplayName, FGuid GroupId = FGuid());
-
-	/**
-	 * Expose a property in this target.
-	 * @param FieldPathInfo the path data from the owner object, including component chain, to this field. (ie. LightComponent0.Intensity)
-	 * @param DesiredDisplayName the display name desired for this control. If the name is not unique preset-wise, a number will be appended.
-	 */
-	TOptional<FRemoteControlProperty> ExposeProperty(FRCFieldPathInfo FieldPathInfo, const FString& DesiredDisplayName, FGuid GroupId = FGuid(), bool bAppendAliasToLabel = false);
-
-	/**
-	 * Expose a function in this target.
-	 * @param RelativeFieldPath the path from the owner object  to this field. (ie. Subcomponent.GetName)
-	 * @param DesiredDisplayName the display name desired for this control. If the name is not unique target-wide, a number will be appended.
-	 */
-	TOptional<FRemoteControlFunction> ExposeFunction(FString RelativeFieldPath, const FString& DesiredDisplayName, FGuid GroupId = FGuid(), bool bAppendAliasToLabel = false);
-
-	/**
-	 * Unexpose a field from this target.
-	 * @param TargetFieldId The ID of the field to remove.
-	 */
-	void Unexpose(FGuid TargetFieldId);
-
-	/**
-	 * Find a field label using a property name.
-	 * @param PropertyName the FName of the property to find.
-	 * @return the field's label or an empty name if not found.
-	 */
-	FName FindFieldLabel(FName FieldName) const;
-
-	/**
-	 * Find a field label using a a field path info.
-	 * @param Path the FieldPathInfo of the field to find.
-	 * @return the field's label or an empty name if not found.
-	 */
-	FName FindFieldLabel(const FRCFieldPathInfo& Path) const;
-
-	/**
-	 * Get a property using its label.
-	 * @param FieldLabel the property's label.
-	 * @return The target property if found.
-	 */
-	TOptional<FRemoteControlProperty> GetProperty(FGuid PropertyId) const;
-
-	/**
-	 * Get an exposed function using it's name.
-	 * @param FunctionName the function name to query.
-	 * @return The function if found.
-	 */
-	TOptional<FRemoteControlFunction> GetFunction(FGuid FunctionId) const;
-
-	/**
-	 * Resolve a remote controlled property to its FProperty and owner objects.
-	 * @param PropertyLabel the label of the remote controlled property.
-	 * @return The resolved exposed property if found.
-	 */
-	PRAGMA_DISABLE_DEPRECATION_WARNINGS
-	TOptional<FExposedProperty> ResolveExposedProperty(FGuid PropertyId) const;
-	PRAGMA_ENABLE_DEPRECATION_WARNINGS
-	/**
-	 * Resolve a remote controlled function to its UFunction and owner objects.
-	 * @param FunctionLabel the label of the remote controlled function.
-	 * @return The resolved exposed function if found.
-	 */
-	PRAGMA_DISABLE_DEPRECATION_WARNINGS
-	TOptional<FExposedFunction> ResolveExposedFunction(FGuid FunctionId) const;
-	PRAGMA_ENABLE_DEPRECATION_WARNINGS
-
-	/**
-	 * Resolve the target bindings to return the target's underlying objects.
-	 * @return The target's underlying objects.
-	 */
-	TArray<UObject*> ResolveBoundObjects() const;
-
-	/**
-	 * Adds objects to the target's underlying objects, binding them to the target's alias.
-	 * @param ObjectsToBind The objects to bind.
-	 */
-	void BindObjects(const TArray<UObject*>& ObjectsToBind);
-
-	/**
-	 * Verifies if the list of objects provided is bound under this target's alias.
-	 * @param ObjectsToTest The objects to validate.
-	 * @return true if the target has  underlying objects.
-	 */
-	bool HasBoundObjects(const TArray<UObject*>& ObjectsToTest) const;
-
-	/**
-	 * Verifies if an object is bound under this target's alias.
-	 * @param ObjectToTest The objects to validate.
-	 * @return true if the target has the object bound.
-	 */
-	bool HasBoundObject(const UObject* ObjectToTest) const;
-
-	/**
-	 * Returns whether the provided list of objects can be bound under this target.
-	 * @param ObjectsToTest The object list.
-	 * @return Whether the objects can be bound.
-	 */
-	bool CanBindObjects(const TArray<UObject*>& ObjectsToTest) const;
-
-private:
-
-	/** Remove a field from the specified fields array. */
-	template <typename Type> 
-	void RemoveField(TSet<Type>& Fields, FGuid FieldId)
-	{
-		Fields.RemoveByHash(GetTypeHash(FieldId), FieldId);
-	}
-
-	FProperty* FindPropertyRecursive(UStruct* Container, TArray<FString>& DesiredPropertyPath) const;
-
-public:
-	/**
-	 * The common class of the target's underlying objects.
-	 */
-	UPROPERTY()
-	UClass* Class = nullptr;
-
-	/**
-	 * The target's exposed functions.
-	 */
-	UPROPERTY()
-	TSet<FRemoteControlFunction> ExposedFunctions;
-
-	/**
-	 * The target's exposed properties.
-	 */
-	UPROPERTY()
-	TSet<FRemoteControlProperty> ExposedProperties;
-
-	/**
-	 * The alias for this target.
-	 */
-	UPROPERTY()
-	FName Alias;
-private:
-	/**
-	 * The objects bound under the target's alias.
-	 */
-	UPROPERTY()
-	TArray<FSoftObjectPath> Bindings;
-
-	/**
-	 * The preset that owns this target.
-	 */
-	UPROPERTY()
-	TWeakObjectPtr<URemoteControlPreset> Owner;
-
-	friend URemoteControlPreset;
-};
-
-/**
- * Holds targets that contain exposed functions and properties.
+ * Holds exposed functions and properties.
  */
 UCLASS(BlueprintType, EditInlineNew)
 class REMOTECONTROL_API URemoteControlPreset : public UObject
@@ -504,35 +331,6 @@ public:
 	 */
 	const FGuid& GetPresetId() const { return PresetId; }
 
-	/**
-	 * Get this preset's targets.
-	 */
-	UE_DEPRECATED(4.27, "FRemoteControlTarget is deprecated, use URemoteControlPreset::Bindings instead.")
-	PRAGMA_DISABLE_DEPRECATION_WARNINGS
-	TMap<FName, FRemoteControlTarget>& GetRemoteControlTargets() { return RemoteControlTargets; }
-	PRAGMA_ENABLE_DEPRECATION_WARNINGS
-
-	/**
-	 * Get this preset's targets.
-	 */
-	UE_DEPRECATED(4.27, "FRemoteControlTarget is deprecated, use URemoteControlPreset::Bindings instead.")
-	PRAGMA_DISABLE_DEPRECATION_WARNINGS
-	const TMap<FName, FRemoteControlTarget>& GetRemoteControlTargets() const { return RemoteControlTargets; }
-	PRAGMA_ENABLE_DEPRECATION_WARNINGS
-	/**
-	 * Get the target that owns this exposed field.
-	 */
-	UE_DEPRECATED(4.27, "FRemoteControlTarget is deprecated, access the entity's bindings by accessing FRemoteControlEntity::Bindings.")
-	FName GetOwnerAlias(FGuid FieldId) const;
-
-	/**
-	 * Get the ID for a field using its label.
-	 * @param FieldLabel the field's label.
-	 * @return the field's id, or an invalid GUID if not found.
-	 */
-	UE_DEPRECATED(4.27, "Use URemoteControlPreset::GetExposedEntityId instead.")
-	FGuid GetFieldId(FName FieldLabel) const;
-
 	/**
 	 * Expose an actor on this preset.
 	 * @param Actor the actor to expose.
@@ -722,41 +520,6 @@ public:
 	 */
 	void Unexpose(const FGuid& EntityId);
 
-	/**
-	 * Create a new target under this preset.
-	 * @param TargetObjects The objects to group under a common alias for the target.
-	 * @return The new target's alias if successful.
-	 * @note A target must be created with at least one object and they must have a common base class.
-	 */
-	UE_DEPRECATED(4.27, "Targets are deprecated in favor of exposing directly on the preset. You do not need to create a target beforehand.")
-	FName CreateTarget(const TArray<UObject*>& TargetObjects);
-
-	/**
-	 * Create a new target under this preset.
-	 * @param TargetObjects The objects to group under a common alias for the target.
-	 * @return The new target.
-	 * @note A target must be created with at least one object and they must have a common base class.
-	 */
-	UE_DEPRECATED(4.27, "Targets are deprecated in favor of exposing directly on the preset. You do not need to create a target beforehand.")
-	PRAGMA_DISABLE_DEPRECATION_WARNINGS
-	FRemoteControlTarget& CreateAndGetTarget(const TArray<UObject*>& TargetObjects);
-	PRAGMA_ENABLE_DEPRECATION_WARNINGS
-
-	/**
-	 * Remove a target from the preset.
-	 * @param TargetName The target to delete.
-	 */
-	UE_DEPRECATED(4.27, "Targets are deprecated, you do not need to delete them anymore.")
-	void DeleteTarget(FName TargetName);
-
-	/**
-	 * Rename a target.
-	 * @param TargetName The name of the target to rename.
-	 * @param NewTargetName The new target's name.
-	 */
-	UE_DEPRECATED(4.27, "Targets are deprecated in favor of bindings, you can now change the name of the binding directly.")
-	void RenameTarget(FName TargetName, FName NewTargetName);
-
 	/** Cache this preset's layout data. */
 	void CacheLayoutData();
 
@@ -883,12 +646,6 @@ private:
 	UPROPERTY(Instanced)
 	TObjectPtr<URCVirtualPropertyContainerBase> ControllerContainer;
 
-	/** Generate a unique alias for this target. */
-	FName GenerateAliasForObjects(const TArray<UObject*>& Objects);
-	
-	/** Generate a label for a field that is unique preset-wide. */
-	FName GenerateUniqueFieldLabel(FName Alias, const FString& BaseName, bool bAppendAlias);
-
 	/** Holds information about an exposed field. */
 	struct FExposeInfo
 	{
@@ -902,7 +659,6 @@ private:
 	void OnUnexpose(FGuid UnexposedFieldId);
 	
 	//~ Cache operations.
-	void CacheFieldsData();
 	void CacheFieldLayoutData();
 	
 	//~ Register/Unregister delegates
@@ -936,11 +692,6 @@ private:
 	/** Get a field ptr using it's id. */
 	FRemoteControlField* GetFieldPtr(FGuid FieldId);
 
-	//~ Utility functions to update the asset format.
-	void ConvertFieldsToRemoveComponentChain();
-	void ConvertFieldsToEntities();
-	void ConvertTargetsToBindings();
-
 	//~ Helper methods that interact with the expose registry.,
 	TSharedPtr<const FRemoteControlEntity> FindEntityById(const FGuid& EntityId, const UScriptStruct* EntityType = FRemoteControlEntity::StaticStruct()) const;
 	TSharedPtr<FRemoteControlEntity> FindEntityById(const FGuid& EntityId, const UScriptStruct* EntityType = FRemoteControlEntity::StaticStruct());
@@ -994,12 +745,6 @@ private:
 	UPROPERTY(AssetRegistrySearchable)
 	FGuid PresetId;
 
-PRAGMA_DISABLE_DEPRECATION_WARNINGS
-	/** The mappings of alias to targets. */
-	UPROPERTY()
-	TMap<FName, FRemoteControlTarget> RemoteControlTargets;
-PRAGMA_ENABLE_DEPRECATION_WARNINGS
-
 	/** The cache for information about an exposed field. */
 	UPROPERTY(Transient)
 	TMap<FGuid, FRCCachedFieldData> FieldCache;
@@ -1093,9 +838,6 @@ PRAGMA_ENABLE_DEPRECATION_WARNINGS
 	TSet<URemoteControlBinding*> PerFrameBindingsToClean;
 #endif
 
-PRAGMA_DISABLE_DEPRECATION_WARNINGS
-	friend FRemoteControlTarget;
-PRAGMA_ENABLE_DEPRECATION_WARNINGS
 	friend FRemoteControlPresetLayout;
 	friend FRemoteControlEntity;
 	friend class FRemoteControlPresetRebindingManager;
