Fix for "Add Controller" crash.

A recent engine change deprecating the ANY_PACKAGE requires that we use the fully qualified path to the Enum class while using FindObject.

diff --git a/Engine/Plugins/VirtualProduction/RemoteControl/Source/RemoteControlUI/Private/UI/Controller/SRCControllerPanelHeader.cpp b/Engine/Plugins/VirtualProduction/RemoteControl/Source/RemoteControlUI/Private/UI/Controller/SRCControllerPanelHeader.cpp
--- a/Engine/Plugins/VirtualProduction/RemoteControl/Source/RemoteControlUI/Private/UI/Controller/SRCControllerPanelHeader.cpp
+++ b/Engine/Plugins/VirtualProduction/RemoteControl/Source/RemoteControlUI/Private/UI/Controller/SRCControllerPanelHeader.cpp
@@ -111,7 +111,7 @@ TSharedRef<SWidget> SRCControllerPanelHeader::GetControllerMenuContentWidget() c
 	for (const TPair<EPropertyBagPropertyType, UObject*>& Pair : VirtualPropertyFieldClassNames)
 	{
 		// Display Name
-		UEnum* Enum = FindObjectChecked<UEnum>(nullptr, TEXT("EPropertyBagPropertyType"));
+		const UEnum* Enum = FindObjectChecked<UEnum>(nullptr, TEXT("/Script/StructUtils.EPropertyBagPropertyType"));
 		FString PropertyTypeString = Enum->GetValueAsName(Pair.Key).ToString();
 		PropertyTypeString.RemoveFromStart("EPropertyBagPropertyType::");
 
