~ Remote Control Logic Web API  ~
UE-155855

For manipulating Controllers through HTTP requests

The following routes have been added:

1. SetPresetController - Set the value of a Controller (on a Remote Control Preset asset).
  HTTP Parameters: Preset Name, Controller Name, Body: JSON;
   Example:
  (PUT) http://127.0.0.1:30010/remote/preset/PresetNameParam/controller/SomeBoolController
  Body (raw JSON):  { "PropertyValue": false}

2. GetPresetController - Fetch the value of a Controller (on a Remote Control Preset asset).
  HTTP Parameters: Preset Name, Controller Name, Body: None
   Example:
  (GET) http://127.0.0.1:30010/remote/preset/PresetNameParam/controller/SomeBoolController
  HTTP Response:  { "PropertyValue": false}

Interception support has been provided as well by adding an nDispaly cluster implementation.

Pending work:
Unit tests for new route handlers to be added to RemoteControlIntereceptionTest.cpp. Review will be updated with this.

diff --git a/Engine/Plugins/Runtime/nDisplay/Source/DisplayClusterRemoteControlInterceptor/Private/DisplayClusterRemoteControlInterceptor.cpp b/Engine/Plugins/Runtime/nDisplay/Source/DisplayClusterRemoteControlInterceptor/Private/DisplayClusterRemoteControlInterceptor.cpp
--- a/Engine/Plugins/Runtime/nDisplay/Source/DisplayClusterRemoteControlInterceptor/Private/DisplayClusterRemoteControlInterceptor.cpp
+++ b/Engine/Plugins/Runtime/nDisplay/Source/DisplayClusterRemoteControlInterceptor/Private/DisplayClusterRemoteControlInterceptor.cpp
@@ -111,6 +111,24 @@ ERCIResponse FDisplayClusterRemoteControlInterceptor::InvokeCall(FRCIFunctionMet
 	return ERCIResponse::Intercept;
 }
 
+ERCIResponse FDisplayClusterRemoteControlInterceptor::SetPresetController(FRCIControllerMetadata& InController)
+{
+	// Serialize command data to binary
+	TArray<uint8> Buffer;
+	FMemoryWriter MemoryWriter(Buffer);
+	MemoryWriter << InController;
+
+	// Queue interception sending
+	QueueInterceptEvent(FRCIControllerMetadata::Name, InController.GetUniquePath(), MoveTemp(Buffer));
+
+	if (bForceApply)
+	{
+		return ERCIResponse::Apply;
+	}
+
+	return ERCIResponse::Intercept;
+}
+
 void FDisplayClusterRemoteControlInterceptor::OnClusterEventBinaryHandler(const FDisplayClusterClusterEventBinary& Event)
 {
 	// Dispatch data to a proper handler
@@ -142,6 +160,10 @@ void FDisplayClusterRemoteControlInterceptor::OnClusterEventBinaryHandler(const
 				{
 					OnReplication_InvokeCall(EventPayload);
 				}
+				else if (MetadataName == FRCIControllerMetadata::Name)
+				{
+					OnReplication_SetPresetController(EventPayload);
+				}
 			}
 		}
 	}
@@ -230,6 +252,31 @@ void FDisplayClusterRemoteControlInterceptor::OnReplication_InvokeCall(const TAr
 	}
 }
 
+void FDisplayClusterRemoteControlInterceptor::OnReplication_SetPresetController(const TArray<uint8>& Buffer)
+{
+	UE_LOG(LogDisplayClusterRemoteControlInterceptor, VeryVerbose, TEXT("Processing replication event SetPresetController: %d bytes"), Buffer.Num());
+
+	// Deserialize command data
+	FMemoryReader MemoryReader(Buffer);
+	FRCIControllerMetadata ControllerMetadata;
+	MemoryReader << ControllerMetadata;
+
+	// Initialization
+	IModularFeatures& ModularFeatures = IModularFeatures::Get();
+	const FName ProcessorFeatureName = IRemoteControlInterceptionFeatureProcessor::GetName();
+	const int32 ProcessorsAmount = ModularFeatures.GetModularFeatureImplementationCount(IRemoteControlInterceptionFeatureProcessor::GetName());
+
+	// Send the command to the processor(s)
+	for (int32 ProcessorIdx = 0; ProcessorIdx < ProcessorsAmount; ++ProcessorIdx)
+	{
+		IRemoteControlInterceptionFeatureProcessor* const Processor = static_cast<IRemoteControlInterceptionFeatureProcessor*>(ModularFeatures.GetModularFeatureImplementation(ProcessorFeatureName, ProcessorIdx));
+		if (Processor)
+		{
+			Processor->SetPresetController(ControllerMetadata);
+		}
+	}
+}
+
 void FDisplayClusterRemoteControlInterceptor::SendReplicationQueue()
 {
 	if (!InterceptQueueMap.Num())
diff --git a/Engine/Plugins/Runtime/nDisplay/Source/DisplayClusterRemoteControlInterceptor/Private/DisplayClusterRemoteControlInterceptor.h b/Engine/Plugins/Runtime/nDisplay/Source/DisplayClusterRemoteControlInterceptor/Private/DisplayClusterRemoteControlInterceptor.h
--- a/Engine/Plugins/Runtime/nDisplay/Source/DisplayClusterRemoteControlInterceptor/Private/DisplayClusterRemoteControlInterceptor.h
+++ b/Engine/Plugins/Runtime/nDisplay/Source/DisplayClusterRemoteControlInterceptor/Private/DisplayClusterRemoteControlInterceptor.h
@@ -23,6 +23,7 @@ public:
 	virtual ERCIResponse SetObjectProperties(FRCIPropertiesMetadata& InProperties) override;
 	virtual ERCIResponse ResetObjectProperties(FRCIObjectMetadata& InObject) override;
 	virtual ERCIResponse InvokeCall(FRCIFunctionMetadata& InFunction) override;
+	virtual ERCIResponse SetPresetController(FRCIControllerMetadata& InController) override;
 	// ~IRemoteControlInterceptionCommands interface
 
 private:
@@ -39,6 +40,8 @@ private:
 	void OnReplication_ResetObjectProperties (const TArray<uint8>& Buffer);
 	// Process InvokeCall command replication data
 	void OnReplication_InvokeCall (const TArray<uint8>& Buffer);
+	// Process SetPresetController command replication data
+	void OnReplication_SetPresetController(const TArray<uint8>& Buffer);
 	// Send the queue of replication events at the end of the tick
 	void SendReplicationQueue();
 
diff --git a/Engine/Plugins/VirtualProduction/RemoteControl/Source/RemoteControl/Private/RCVirtualProperty.cpp b/Engine/Plugins/VirtualProduction/RemoteControl/Source/RemoteControl/Private/RCVirtualProperty.cpp
--- a/Engine/Plugins/VirtualProduction/RemoteControl/Source/RemoteControl/Private/RCVirtualProperty.cpp
+++ b/Engine/Plugins/VirtualProduction/RemoteControl/Source/RemoteControl/Private/RCVirtualProperty.cpp
@@ -2,12 +2,13 @@
 
 #include "RCVirtualProperty.h"
 
+#include "IStructDeserializerBackend.h"
+#include "RCVirtualPropertyContainer.h"
+#include "StructDeserializer.h"
 #include "StructSerializer.h"
 #include "UObject/TextProperty.h"
 #include "UObject/StructOnScope.h"
 
-#include "RCVirtualPropertyContainer.h"
-
 const FProperty* URCVirtualPropertyBase::GetProperty() const
 {
 	if (const FPropertyBagPropertyDesc* PropertyBagPropertyDesc = GetBagPropertyDesc())
@@ -118,6 +119,27 @@ void URCVirtualPropertyBase::SerializeToBackend(IStructSerializerBackend& OutBac
 	FStructSerializer::SerializeElement(GetContainerPtr(), GetProperty(), INDEX_NONE, OutBackend, Policies);
 }
 
+bool URCVirtualPropertyBase::DeserializeFromBackend(IStructDeserializerBackend& InBackend)
+{
+	FStructDeserializerPolicies Policies;
+	const FProperty* Property = GetProperty();
+	Policies.PropertyFilter = [&Property](const FProperty* CurrentProp, const FProperty* ParentProp)
+	{
+		return CurrentProp == Property;
+	};
+
+	const UPropertyBag* TypeInfo = GetPropertyBagInstance()->GetPropertyBagStruct();
+	UPropertyBag* TypeInfoNonConst = const_cast<UPropertyBag*>(TypeInfo);
+
+	bool bSuccess = FStructDeserializer::DeserializeElement(GetContainerPtr(), *TypeInfoNonConst, INDEX_NONE, InBackend, Policies);
+	if (bSuccess)
+	{
+		OnModifyPropertyValue();
+	}
+
+	return bSuccess;
+}
+
 bool URCVirtualPropertyBase::IsValueEqual(URCVirtualPropertyBase* InVirtualProperty)
 {
 	const FProperty* ThisProperty = GetProperty();
@@ -564,4 +586,4 @@ const FInstancedPropertyBag* URCVirtualPropertySelfContainer::GetPropertyBagInst
 TSharedPtr<FStructOnScope> URCVirtualPropertySelfContainer::CreateStructOnScope() const
 {
 	return MakeShared<FStructOnScope>(Bag.GetPropertyBagStruct(), Bag.GetMutableValue().GetMutableMemory());
-}
+}
\ No newline at end of file
diff --git a/Engine/Plugins/VirtualProduction/RemoteControl/Source/RemoteControl/Private/RemoteControlInterceptionProcessor.cpp b/Engine/Plugins/VirtualProduction/RemoteControl/Source/RemoteControl/Private/RemoteControlInterceptionProcessor.cpp
--- a/Engine/Plugins/VirtualProduction/RemoteControl/Source/RemoteControl/Private/RemoteControlInterceptionProcessor.cpp
+++ b/Engine/Plugins/VirtualProduction/RemoteControl/Source/RemoteControl/Private/RemoteControlInterceptionProcessor.cpp
@@ -11,7 +11,6 @@
 #include "StructDeserializer.h"
 #include "UObject/FieldPath.h"
 
-
 void FRemoteControlInterceptionProcessor::SetObjectProperties(FRCIPropertiesMetadata& PropsMetadata)
 {
 	// Set object reference
@@ -110,3 +109,14 @@ void FRemoteControlInterceptionProcessor::InvokeCall(FRCIFunctionMetadata& InFun
 		IRemoteControlModule::Get().InvokeCall(Call);
 	}
 }
+
+void FRemoteControlInterceptionProcessor::SetPresetController(FRCIControllerMetadata& InController)
+{
+	FMemoryReader Reader(InController.Payload);
+	FJsonStructDeserializerBackend Backend(Reader);
+
+	constexpr bool bAllowIntercept = false; // Run without replication
+
+	IRemoteControlModule::Get().SetPresetController(InController.Preset, InController.Controller, Backend, TArray<uint8>(), bAllowIntercept);
+}
+
diff --git a/Engine/Plugins/VirtualProduction/RemoteControl/Source/RemoteControl/Private/RemoteControlInterceptionProcessor.h b/Engine/Plugins/VirtualProduction/RemoteControl/Source/RemoteControl/Private/RemoteControlInterceptionProcessor.h
--- a/Engine/Plugins/VirtualProduction/RemoteControl/Source/RemoteControl/Private/RemoteControlInterceptionProcessor.h
+++ b/Engine/Plugins/VirtualProduction/RemoteControl/Source/RemoteControl/Private/RemoteControlInterceptionProcessor.h
@@ -20,5 +20,6 @@ public:
 	virtual void SetObjectProperties(FRCIPropertiesMetadata& InProperties) override;
 	virtual void ResetObjectProperties(FRCIObjectMetadata& InObject) override;
 	virtual void InvokeCall(FRCIFunctionMetadata& InFunction) override;
+	virtual void SetPresetController(FRCIControllerMetadata& InController) override;
 	// ~IRemoteControlInterceptionCommands interface
 };
diff --git a/Engine/Plugins/VirtualProduction/RemoteControl/Source/RemoteControl/Private/RemoteControlModule.cpp b/Engine/Plugins/VirtualProduction/RemoteControl/Source/RemoteControl/Private/RemoteControlModule.cpp
--- a/Engine/Plugins/VirtualProduction/RemoteControl/Source/RemoteControl/Private/RemoteControlModule.cpp
+++ b/Engine/Plugins/VirtualProduction/RemoteControl/Source/RemoteControl/Private/RemoteControlModule.cpp
@@ -2,13 +2,19 @@
 
 #include "RemoteControlModule.h"
 
+#include "AssetRegistry/AssetRegistryModule.h"
+#include "AssetRegistry/IAssetRegistry.h"
+#include "Backends/CborStructSerializerBackend.h"
+#include "Components/ActorComponent.h"
+#include "Components/LightComponent.h"
+#include "Features/IModularFeatures.h"
 #include "IRemoteControlInterceptionFeature.h"
 #include "IRemoteControlModule.h"
 #include "IStructDeserializerBackend.h"
 #include "IStructSerializerBackend.h"
-#include "Components/ActorComponent.h"
-#include "Components/LightComponent.h"
 #include "RCPropertyUtilities.h"
+#include "RCVirtualPropertyContainer.h"
+#include "RCVirtualProperty.h"
 #include "RemoteControlFieldPath.h"
 #include "RemoteControlInterceptionHelpers.h"
 #include "RemoteControlInterceptionProcessor.h"
@@ -17,13 +23,6 @@
 #include "Serialization/PropertyMapStructDeserializerBackendWrapper.h"
 #include "StructDeserializer.h"
 #include "StructSerializer.h"
-#include "AssetRegistry/AssetRegistryModule.h"
-#include "AssetRegistry/IAssetRegistry.h"
-#include "Backends/CborStructSerializerBackend.h"
-#include "Features/IModularFeatures.h"
-#include "Misc/CoreMisc.h"
-#include "Misc/ScopeExit.h"
-#include "Misc/TVariant.h"
 #include "UObject/Class.h"
 #include "UObject/FieldPath.h"
 #include "UObject/UnrealType.h"
@@ -1209,6 +1208,76 @@ bool FRemoteControlModule::SetObjectProperties(const FRCObjectReference& ObjectA
 	return false;
 }
 
+bool FRemoteControlModule::SetPresetController(const FName PresetName, const FName ControllerName, IStructDeserializerBackend& Backend, const TArray<uint8>& InPayload, const bool bAllowIntercept)
+{
+	URemoteControlPreset* Preset = IRemoteControlModule::Get().ResolvePreset(PresetName);
+	if (!ensure(Preset))
+	{
+		return false;
+	}
+
+	URCVirtualPropertyBase* VirtualProperty = Preset->GetVirtualProperty(ControllerName);
+	if (!ensure(VirtualProperty))
+	{
+		return false;
+	}
+
+	return SetPresetController(PresetName, VirtualProperty, Backend, InPayload, bAllowIntercept);
+}
+
+// Note: The actual Controller implementation (URCController) resides in the RemoteControlLogic module
+// and is not referenced directly here to avoid circular dependency between RemoteControl and RemoteControlLogic modules.
+// Controllers are simply accessed as their parent "URCVirtualPropertyBase" objects 
+
+bool FRemoteControlModule::SetPresetController(const FName PresetName, class URCVirtualPropertyBase* VirtualProperty, IStructDeserializerBackend& Backend, const TArray<uint8>& InPayload, const bool bAllowIntercept)
+{
+	if (!ensure(VirtualProperty))
+	{
+		return false;
+	}
+
+	URemoteControlPreset* Preset = VirtualProperty->PresetWeakPtr.Get();
+	if (!ensure(Preset))
+	{
+		return false;
+	}
+
+	if (bAllowIntercept)
+	{
+		// Build interception command
+		FRCIControllerMetadata ControllersMetadata(PresetName, VirtualProperty->PropertyName, InPayload);
+
+		// Initialization
+		IModularFeatures& ModularFeatures = IModularFeatures::Get();
+		const FName InterceptorFeatureName = IRemoteControlInterceptionFeatureInterceptor::GetName();
+		const int32 InterceptorsAmount = ModularFeatures.GetModularFeatureImplementationCount(IRemoteControlInterceptionFeatureInterceptor::GetName());
+
+		UE_LOG(LogRemoteControl, VeryVerbose, TEXT("SetPresetController - Num Interceptors: %d"), InterceptorsAmount);
+
+		// Pass interception command data to all available interceptors
+		bool bShouldIntercept = false;
+		for (int32 InterceptorIdx = 0; InterceptorIdx < InterceptorsAmount; ++InterceptorIdx)
+		{
+			IRemoteControlInterceptionFeatureInterceptor* const Interceptor = static_cast<IRemoteControlInterceptionFeatureInterceptor*>(ModularFeatures.GetModularFeatureImplementation(InterceptorFeatureName, InterceptorIdx));
+			if (Interceptor)
+			{
+				// Update response flag
+				bShouldIntercept |= (Interceptor->SetPresetController(ControllersMetadata) == ERCIResponse::Intercept);
+			}
+		}
+
+		// Don't process the RC message if any of interceptors returned ERCIResponse::Intercept
+		if (bShouldIntercept)
+		{
+			return true;
+		}
+	}
+
+	bool bSuccess = VirtualProperty->DeserializeFromBackend(Backend);
+
+	return bSuccess;
+}
+
 bool FRemoteControlModule::ResetObjectProperties(const FRCObjectReference& ObjectAccess, const bool bAllowIntercept)
 {
 	// Check the replication path before reset the property on this instance
diff --git a/Engine/Plugins/VirtualProduction/RemoteControl/Source/RemoteControl/Private/RemoteControlModule.h b/Engine/Plugins/VirtualProduction/RemoteControl/Source/RemoteControl/Private/RemoteControlModule.h
--- a/Engine/Plugins/VirtualProduction/RemoteControl/Source/RemoteControl/Private/RemoteControlModule.h
+++ b/Engine/Plugins/VirtualProduction/RemoteControl/Source/RemoteControl/Private/RemoteControlModule.h
@@ -33,6 +33,8 @@ public:
 	virtual bool GetObjectProperties(const FRCObjectReference& ObjectAccess, IStructSerializerBackend& Backend) override;
 	virtual bool SetObjectProperties(const FRCObjectReference& ObjectAccess, IStructDeserializerBackend& Backend, ERCPayloadType InPayloadType, const TArray<uint8>& InPayload, ERCModifyOperation Operation) override;
 	virtual bool ResetObjectProperties(const FRCObjectReference& ObjectAccess, const bool bAllowIntercept) override;
+	virtual bool SetPresetController(const FName PresetName, const FName ControllerName, IStructDeserializerBackend& Backend, const TArray<uint8>& InPayload, const bool bAllowIntercept) override;
+	virtual bool SetPresetController(const FName PresetName, class URCVirtualPropertyBase* VirtualProperty, IStructDeserializerBackend& Backend, const TArray<uint8>& InPayload, const bool bAllowIntercept) override;
 	virtual TOptional<FExposedFunction> ResolvePresetFunction(const FResolvePresetFieldArgs& Args) const override;
 	virtual TOptional<FExposedProperty> ResolvePresetProperty(const FResolvePresetFieldArgs& Args) const override;
 	virtual URemoteControlPreset* ResolvePreset(FName PresetName) const override;
diff --git a/Engine/Plugins/VirtualProduction/RemoteControl/Source/RemoteControl/Private/RemoteControlPreset.cpp b/Engine/Plugins/VirtualProduction/RemoteControl/Source/RemoteControl/Private/RemoteControlPreset.cpp
--- a/Engine/Plugins/VirtualProduction/RemoteControl/Source/RemoteControl/Private/RemoteControlPreset.cpp
+++ b/Engine/Plugins/VirtualProduction/RemoteControl/Source/RemoteControl/Private/RemoteControlPreset.cpp
@@ -541,7 +541,7 @@ void URemoteControlPreset::ResetVirtualProperties()
 {
 	if (ensure(ControllerContainer))
 	{
-		return ControllerContainer->Reset();
+		ControllerContainer->Reset();
 	}
 }
 
diff --git a/Engine/Plugins/VirtualProduction/RemoteControl/Source/RemoteControl/Private/Tests/RemoteControlInterceptionTest.cpp b/Engine/Plugins/VirtualProduction/RemoteControl/Source/RemoteControl/Private/Tests/RemoteControlInterceptionTest.cpp
--- a/Engine/Plugins/VirtualProduction/RemoteControl/Source/RemoteControl/Private/Tests/RemoteControlInterceptionTest.cpp
+++ b/Engine/Plugins/VirtualProduction/RemoteControl/Source/RemoteControl/Private/Tests/RemoteControlInterceptionTest.cpp
@@ -22,7 +22,6 @@
 #include "Serialization/MemoryWriter.h"
 #include "Serialization/MemoryReader.h"
 
-
 /**
  * Interceptor interface implementation
  */
@@ -78,6 +77,22 @@ protected:
 		
 		return ERCIResponse::Intercept;
 	}
+
+	virtual ERCIResponse SetPresetController(FRCIControllerMetadata& InController)
+	{
+		// Get processor feature
+		IRemoteControlInterceptionFeatureProcessor* const Processor = static_cast<IRemoteControlInterceptionFeatureProcessor*>(
+			IModularFeatures::Get().GetModularFeatureImplementation(IRemoteControlInterceptionFeatureProcessor::GetName(), 0));
+
+		// In case the processor feature has been registered, forward data directly to the processor
+		if (Processor)
+		{
+			Processor->SetPresetController(InController);
+		}
+
+		return ERCIResponse::Intercept;
+	}
+
 	// ~IRemoteControlInterceptionCommands interface
 };
 
diff --git a/Engine/Plugins/VirtualProduction/RemoteControl/Source/RemoteControl/Public/IRemoteControlModule.h b/Engine/Plugins/VirtualProduction/RemoteControl/Source/RemoteControl/Public/IRemoteControlModule.h
--- a/Engine/Plugins/VirtualProduction/RemoteControl/Source/RemoteControl/Public/IRemoteControlModule.h
+++ b/Engine/Plugins/VirtualProduction/RemoteControl/Source/RemoteControl/Public/IRemoteControlModule.h
@@ -111,6 +111,15 @@ struct FResolvePresetFieldArgs
 	FString FieldLabel;
 };
 
+/**
+ * Arguments for resolving a Controller on a preset.
+ */
+struct FResolvePresetControllerArgs
+{
+	FString PresetName;
+	FString ControllerName;
+};
+
 /**
  * Requested access mode to a remote property
  */
@@ -315,6 +324,28 @@ public:
 	 */
 	virtual bool ResetObjectProperties(const FRCObjectReference& ObjectAccess, const bool bAllowIntercept = false) = 0;
 
+	/**
+	* Set a controller's value on a Remote Control Preset
+	* @param PresetName - The Remote Control Preset asset's name
+	* @param ControllerName - The name of the controller being manipulated
+	* @param Backend - the struct deserializer backend to use to deserialize the object properties.
+	* @param InPayload - the JSON payoad
+	* @param bAllowIntercept - Whether the call needs to be replicated to interceptors or invoked directly
+	*	@return true if the operation succeeded
+	*/
+	virtual bool SetPresetController(const FName PresetName, const FName ControllerName, IStructDeserializerBackend& Backend, const TArray<uint8>& InPayload, const bool bAllowIntercept) = 0;
+
+	/**
+	 * Set a controller's value on a Remote Control Preset
+	 * @param PresetName - The Remote Control Preset asset's name
+	 * @param Controller - The controller object being manipulated
+	 * @param Backend - the struct deserializer backend to use to deserialize the object properties.
+	 * @param InPayload - the JSON payoad
+	* @param bAllowIntercept - Whether the call needs to be replicated to interceptors or invoked directly
+	 * @return true if the operation succeeded
+	 */
+	virtual bool SetPresetController(const FName PresetName, class URCVirtualPropertyBase* Controller, IStructDeserializerBackend& Backend, const TArray<uint8>& InPayload, const bool bAllowIntercept) = 0;
+
 	/**
 	 * Resolve the underlying function from a preset.
 	 * @return the underlying function and objects that the property is exposed on.
@@ -449,3 +480,4 @@ public:
 	/** Get map of the factories which is responsible for the Remote Control property creation */
 	virtual const TMap<FName, TSharedPtr<IRemoteControlPropertyFactory>>& GetEntityFactories() const = 0;
 };
+
diff --git a/Engine/Plugins/VirtualProduction/RemoteControl/Source/RemoteControl/Public/RCVirtualProperty.h b/Engine/Plugins/VirtualProduction/RemoteControl/Source/RemoteControl/Public/RCVirtualProperty.h
--- a/Engine/Plugins/VirtualProduction/RemoteControl/Source/RemoteControl/Public/RCVirtualProperty.h
+++ b/Engine/Plugins/VirtualProduction/RemoteControl/Source/RemoteControl/Public/RCVirtualProperty.h
@@ -7,10 +7,11 @@
  
 #include "RCVirtualProperty.generated.h"
 
+class FStructOnScope;
 class IStructSerializerBackend;
+class IStructDeserializerBackend;
 class URemoteControlPreset;
 class URCVirtualPropertyContainerBase;
-class FStructOnScope;
 
 /**
  * Base class for dynamic virtual properties
@@ -43,6 +44,8 @@ protected:
 	/** Pointer to Instanced property which is holds  bag of properties */
 	virtual const FInstancedPropertyBag* GetPropertyBagInstance() const;
 
+	virtual void OnModifyPropertyValue() {}
+
 public:
 	/** Returns const FProperty for this RC virtual property */
 	virtual const FProperty* GetProperty() const;
@@ -63,6 +66,13 @@ public:
 	 */
 	void SerializeToBackend(IStructSerializerBackend& OutBackend);
 
+	/**
+	 * Deserialize Virtual Property from a given Backend
+	 *
+	 * @param InBackend  - Deserializer containing a payload with value data
+	 */
+	bool DeserializeFromBackend(IStructDeserializerBackend& InBackend);
+
 	/** Compare this virtual property value with given property value */
 	bool IsValueEqual(URCVirtualPropertyBase* InVirtualProperty);
 
@@ -320,4 +330,4 @@ private:
 	/** Instanced property bag for store a bag of properties. */
 	UPROPERTY()
 	FInstancedPropertyBag Bag;
-};
+};
\ No newline at end of file
diff --git a/Engine/Plugins/VirtualProduction/RemoteControl/Source/RemoteControlLogic/Public/Controller/RCController.h b/Engine/Plugins/VirtualProduction/RemoteControl/Source/RemoteControlLogic/Public/Controller/RCController.h
--- a/Engine/Plugins/VirtualProduction/RemoteControl/Source/RemoteControlLogic/Public/Controller/RCController.h
+++ b/Engine/Plugins/VirtualProduction/RemoteControl/Source/RemoteControlLogic/Public/Controller/RCController.h
@@ -35,6 +35,12 @@ public:
 	/** Execute all behaviours for this controller. */
 	virtual void ExecuteBehaviours();
 
+	/** Handles modifications to controller value; evaluates all behaviours */
+	virtual void OnModifyPropertyValue() override
+	{
+		ExecuteBehaviours();
+	}
+
 public:
 	/** Set of the behaviours */
 	UPROPERTY()
diff --git a/Engine/Plugins/VirtualProduction/RemoteControl/Source/WebRemoteControl/Private/RemoteControlRequest.h b/Engine/Plugins/VirtualProduction/RemoteControl/Source/WebRemoteControl/Private/RemoteControlRequest.h
--- a/Engine/Plugins/VirtualProduction/RemoteControl/Source/WebRemoteControl/Private/RemoteControlRequest.h
+++ b/Engine/Plugins/VirtualProduction/RemoteControl/Source/WebRemoteControl/Private/RemoteControlRequest.h
@@ -817,3 +817,31 @@
 	UPROPERTY()
 	int32 TransactionId = -1;
 };
+
+/**
+ * Struct representation of SetPresetController HTTP request
+ */
+USTRUCT()
+struct FRCPresetSetControllerRequest : public FRCRequest
+{
+	GENERATED_BODY()
+
+	FRCPresetSetControllerRequest()
+	{
+		AddStructParameter(PropertyValueLabel());
+	}
+
+	/**
+	 * Get the label for the PropertyValue struct.
+	 */
+	static FString PropertyValueLabel() { return TEXT("PropertyValue"); }
+
+public:
+
+	/**
+	 * The name of the Controller being set (for a given Remote Control Preset asset)
+	 */
+	UPROPERTY()
+	FString ControllerName;
+};
+
diff --git a/Engine/Plugins/VirtualProduction/RemoteControl/Source/WebRemoteControl/Private/WebRemoteControl.cpp b/Engine/Plugins/VirtualProduction/RemoteControl/Source/WebRemoteControl/Private/WebRemoteControl.cpp
--- a/Engine/Plugins/VirtualProduction/RemoteControl/Source/WebRemoteControl/Private/WebRemoteControl.cpp
+++ b/Engine/Plugins/VirtualProduction/RemoteControl/Source/WebRemoteControl/Private/WebRemoteControl.cpp
@@ -6,6 +6,8 @@
 #include "Modules/ModuleManager.h"
 
 #include "IRemoteControlModule.h"
+#include "RCVirtualPropertyContainer.h"
+#include "RCVirtualProperty.h"
 #include "RemoteControlReflectionUtils.h"
 #include "RemoteControlRoute.h"
 #include "RemoteControlSettings.h"
@@ -758,6 +760,21 @@ void FWebRemoteControlModule::RegisterRoutes()
 		FRequestHandlerDelegate::CreateRaw(this, &FWebRemoteControlModule::HandleDeleteTransientPresetRoute)
 	});
 
+	// Remote Control Logic - Controllers
+	RegisterRoute({
+	TEXT("Set a controller value on a preset."),
+	FHttpPath(TEXT("/remote/preset/:preset/controller/:controller")),
+	EHttpServerRequestVerbs::VERB_PUT,
+	FRequestHandlerDelegate::CreateRaw(this, &FWebRemoteControlModule::HandlePresetSetControllerRoute)
+		});
+
+	RegisterRoute({
+		TEXT("Get a controller value from a preset."),
+		FHttpPath(TEXT("/remote/preset/:preset/controller/:controller")),
+		EHttpServerRequestVerbs::VERB_GET,
+		FRequestHandlerDelegate::CreateRaw(this, &FWebRemoteControlModule::HandlePresetGetControllerRoute)
+		});
+
 	//**************************************
 	// Special websocket route just using http request
 	RegisterWebsocketRoute({
@@ -1974,6 +1991,151 @@ bool FWebRemoteControlModule::HandleDeleteTransientPresetRoute(const FHttpServer
 	return true;
 }
 
+// Remote Control Logic - Controllers
+bool FWebRemoteControlModule::HandlePresetSetControllerRoute(const FHttpServerRequest& Request, const FHttpResultCallback& OnComplete)
+{
+	// 1. Prepare response object
+	TUniquePtr<FHttpServerResponse> Response = WebRemoteControlInternalUtils::CreateHttpResponse();
+
+	// 2. Validate Content Type
+	if (!WebRemoteControlInternalUtils::ValidateContentType(Request, TEXT("application/json"), OnComplete))
+	{
+		WebRemoteControlInternalUtils::CreateUTF8ErrorMessage(TEXT("Expected content type to be application/json"), Response->Body);
+		OnComplete(MoveTemp(Response));
+
+		return false;
+	}
+
+	// 3. Deserialize Json POST data into request struct
+	FRCPresetSetControllerRequest SetControllerRequest;
+	if (!WebRemoteControlInternalUtils::DeserializeRequest(Request, &OnComplete, SetControllerRequest))
+	{
+		WebRemoteControlInternalUtils::CreateUTF8ErrorMessage(TEXT("Unable to process JSON body. Expected format: type to be application/json"), Response->Body);
+		OnComplete(MoveTemp(Response));
+
+		return false;
+	}
+
+	// 4. Resolve url arguments
+	FResolvePresetControllerArgs Args;
+	Args.PresetName = Request.PathParams.FindChecked(TEXT("preset"));
+	Args.ControllerName = Request.PathParams.FindChecked(TEXT("controller"));
+
+	// 5. Acquire Remote Control Preset object
+	URemoteControlPreset* Preset = IRemoteControlModule::Get().ResolvePreset(*Args.PresetName);
+	if (Preset == nullptr)
+	{
+		Response->Code = EHttpServerResponseCodes::NotFound;
+		WebRemoteControlInternalUtils::CreateUTF8ErrorMessage(TEXT("Unable to resolve the preset."), Response->Body);
+		OnComplete(MoveTemp(Response));
+
+		return false;
+	}
+
+	// 6. Acquire Controller
+	URCVirtualPropertyBase* Controller = Preset->GetVirtualProperty(*Args.ControllerName);
+	if (!Controller)
+	{
+		Response->Code = EHttpServerResponseCodes::NotFound;
+		WebRemoteControlInternalUtils::CreateUTF8ErrorMessage(TEXT("Unable to resolve the controller input."), Response->Body);
+		OnComplete(MoveTemp(Response));
+
+		return false;
+	}
+
+	// 7. Reformat Payload to represent our internal structure
+	TArray<uint8> NewPayload;
+	const FName PropertyValueKey = WebRemoteControlStructUtils::Prop_PropertyValue;
+
+	RemotePayloadSerializer::ReplaceFirstOccurence(SetControllerRequest.TCHARBody, PropertyValueKey.ToString(), Args.ControllerName, NewPayload);
+	FMemoryReader Reader(NewPayload);
+	FJsonStructDeserializerBackend Backend(Reader);
+
+	// 8. Invoke via RC Module which provides Interception functionality
+	const bool bSuccess = IRemoteControlModule::Get().SetPresetController(*Args.PresetName, Controller, Backend, NewPayload, true /*support interception*/);
+
+	// 9. Finalize HTTP Response
+	if (bSuccess)
+	{
+		Response->Code = EHttpServerResponseCodes::Ok;
+	}
+	else
+	{
+		WebRemoteControlInternalUtils::CreateUTF8ErrorMessage(FString::Printf(TEXT("Error while trying to set controller %s."), *Args.ControllerName), Response->Body);
+	}
+
+	OnComplete(MoveTemp(Response));
+
+	return true;
+
+}
+
+bool FWebRemoteControlModule::HandlePresetGetControllerRoute(const FHttpServerRequest& Request, const FHttpResultCallback& OnComplete)
+{
+	// 1. Prepare response object
+	TUniquePtr<FHttpServerResponse> Response = WebRemoteControlInternalUtils::CreateHttpResponse();
+
+	// 2. Resolve url arguments
+	FResolvePresetControllerArgs Args;
+	Args.PresetName = Request.PathParams.FindChecked(TEXT("preset"));
+	Args.ControllerName = Request.PathParams.FindChecked(TEXT("controller"));
+
+	// 3. Acquire Remote Control Preset object
+	URemoteControlPreset* Preset = IRemoteControlModule::Get().ResolvePreset(*Args.PresetName);
+	if (Preset == nullptr)
+	{
+		Response->Code = EHttpServerResponseCodes::NotFound;
+		WebRemoteControlInternalUtils::CreateUTF8ErrorMessage(TEXT("Unable to resolve the preset."), Response->Body);
+		OnComplete(MoveTemp(Response));
+
+		return true;
+	}
+
+	// 4. Acquire Controller
+	URCVirtualPropertyBase* Controller = Preset->GetVirtualProperty(*Args.ControllerName);
+	if (!Controller)
+	{
+		Response->Code = EHttpServerResponseCodes::NotFound;
+		WebRemoteControlInternalUtils::CreateUTF8ErrorMessage(TEXT("Unable to resolve the controller input."), Response->Body);
+		OnComplete(MoveTemp(Response));
+
+		return true;
+	}
+
+	// 5. Define a StructOnScope for representing our generic value output
+	FWebRCGenerateStructArgs StructArgs;
+	const FName PropertyValueKey = WebRemoteControlStructUtils::Prop_PropertyValue;
+	StructArgs.GenericProperties.Emplace(PropertyValueKey, Controller->GetProperty());
+
+	FString& StructName = ControllersSerializerStructNameCache.FindOrAdd(Controller->PropertyName);
+	if(StructName.IsEmpty())
+	{
+		StructName = FString::Format(TEXT("{0}_{1}"), { *PropertyValueKey.ToString(), Controller->Id.ToString() });
+		ControllersSerializerStructNameCache.Add(Controller->PropertyName, StructName);
+	}
+	
+	FStructOnScope Result(UE::WebRCReflectionUtils::GenerateStruct(*StructName, StructArgs));
+
+	// 6. Copy raw memory from the controller onto our dynamic struct
+	const FProperty* TargetValueProp = Result.GetStruct()->FindPropertyByName(PropertyValueKey);
+	uint8* DestPtr = TargetValueProp->ContainerPtrToValuePtr<uint8>((void*)Result.GetStructMemory());
+
+	Controller->CopyCompleteValue(TargetValueProp, DestPtr);
+
+	// 7. Serialize to JSON
+	TArray<uint8> JsonBuffer;
+	FMemoryWriter Writer(JsonBuffer);
+	WebRemoteControlInternalUtils::SerializeStructOnScope(Result, Writer);
+
+	// 8. Finalize HTTP Response
+	WebRemoteControlUtils::ConvertToUTF8(JsonBuffer, Response->Body);
+	Response->Code = EHttpServerResponseCodes::Ok;
+
+	OnComplete(MoveTemp(Response));
+
+	return true;
+}
+
 void FWebRemoteControlModule::HandleWebSocketHttpMessage(const FRemoteControlWebSocketMessage& WebSocketMessage)
 {
 	// Sets the acting client id for the duration of the message handling.
diff --git a/Engine/Plugins/VirtualProduction/RemoteControl/Source/WebRemoteControl/Private/WebRemoteControl.h b/Engine/Plugins/VirtualProduction/RemoteControl/Source/WebRemoteControl/Private/WebRemoteControl.h
--- a/Engine/Plugins/VirtualProduction/RemoteControl/Source/WebRemoteControl/Private/WebRemoteControl.h
+++ b/Engine/Plugins/VirtualProduction/RemoteControl/Source/WebRemoteControl/Private/WebRemoteControl.h
@@ -131,6 +131,8 @@ private:
 	bool HandlePassphraseRoute(const FHttpServerRequest& Request, const FHttpResultCallback& OnComplete);
 	bool HandleCreateTransientPresetRoute(const FHttpServerRequest& Request, const FHttpResultCallback& OnComplete);
 	bool HandleDeleteTransientPresetRoute(const FHttpServerRequest& Request, const FHttpResultCallback& OnComplete);
+	bool HandlePresetSetControllerRoute(const FHttpServerRequest& Request, const FHttpResultCallback& OnComplete);
+	bool HandlePresetGetControllerRoute(const FHttpServerRequest& Request, const FHttpResultCallback& OnComplete);
 
 	//~ Websocket route handlers
 	void HandleWebSocketHttpMessage(const struct FRemoteControlWebSocketMessage& WebSocketMessage);
@@ -203,4 +205,8 @@ private:
 
 	/** Receiver that responds to beacon messages from the stage app. */
 	FStageAppBeaconReceiver StageAppBeaconReceiver;
+
+	/** Name cache for Get Controller Result structs
+	* Key: Controller name; Value: Name of the dynamic struct holding our result*/
+	TMap<FName, FString> ControllersSerializerStructNameCache;
 };
\ No newline at end of file
diff --git a/Engine/Plugins/VirtualProduction/RemoteControlInterception/Source/RemoteControlInterception/Public/IRemoteControlInterceptionCommands.h b/Engine/Plugins/VirtualProduction/RemoteControlInterception/Source/RemoteControlInterception/Public/IRemoteControlInterceptionCommands.h
--- a/Engine/Plugins/VirtualProduction/RemoteControlInterception/Source/RemoteControlInterception/Public/IRemoteControlInterceptionCommands.h
+++ b/Engine/Plugins/VirtualProduction/RemoteControlInterception/Source/RemoteControlInterception/Public/IRemoteControlInterceptionCommands.h
@@ -7,6 +7,7 @@
 struct FRCIObjectMetadata;
 struct FRCIFunctionMetadata;
 struct FRCIPropertiesMetadata;
+struct FRCIControllerMetadata;
 
 /**
  * Interception flags that define how RemoteControl should behave after a message was intercepted.
@@ -90,6 +91,15 @@ public:
 	 * @return - Return types depends on the interface implementation (template parameter)
 	 */
 	virtual TResponseType InvokeCall(FRCIFunctionMetadata& InFunction) = 0;
+
+	/**
+	 * SetPresetController command
+	 *
+	 * @param InController - Metadata of the controller being manipulated
+	 *
+	 * @return - Return types depends on the interface implementation (template parameter)
+	 */
+	virtual TResponseType SetPresetController(FRCIControllerMetadata& InController) = 0;
 };
 
 
@@ -299,3 +309,70 @@ private:
 	/** Object Path + Function Path */
 	FName UniquePath;
 };
+
+/**
+ * Controller properties for interception of Controller get/set operations
+ */
+struct FRCIControllerMetadata
+{
+public:
+
+	FRCIControllerMetadata() = default;
+
+	FRCIControllerMetadata(const FName InPreset, const FName InController, const TArray<uint8>& InPayload)
+		: Preset(InPreset)
+		, Controller(InController)
+		, Payload(InPayload)
+	{
+		UniquePath = *FString::Printf(TEXT("%s:%s"), *Preset.ToString(), *Controller.ToString());
+	}
+
+public:
+	/** Get Unique Path for Controller metadata */
+	const FName& GetUniquePath() const
+	{
+		return UniquePath;
+	}
+
+	/** Returns true if reference is valid */
+	bool IsValid() const
+	{
+		return !Preset.IsNone() && !Controller.IsNone() && Payload.Num() > 0;
+	}
+
+	/**
+	* FRCIControllerMetadata serialization
+	*
+	* @param Ar				- The archive
+	* @param Interception	- Instance to serialize/deserialize
+	*
+	* @return FArchive instance.
+	*/
+	friend FArchive& operator<<(FArchive& Ar, FRCIControllerMetadata& Interception)
+	{
+		Ar << Interception.Preset;
+		Ar << Interception.Controller;
+		// const cast is needed for keeping the const& on a constructor and as a class parameter
+		Ar << const_cast<TArray<uint8>&>(Interception.Payload);
+		return Ar;
+	}
+
+public:
+	/**Name of the Remote Control Preset holding the Controller */
+	FName Preset;
+
+	/** The Controller being manipulated in the Preset */
+	FName Controller;
+
+	/** Request Payload */
+	const TArray<uint8> Payload;
+
+public:
+	/** Structure Name - Used by replication event handlers to deserialize event data into
+	* the correct structure for a given event*/
+	static constexpr TCHAR const* Name = TEXT("FRCIControllerMetadata");
+
+private:
+	/** Preset Name + Controller Name*/
+	FName UniquePath;
+};
\ No newline at end of file
