- Fix IsChildOf called on nullptr in RemoteControl

diff --git a/Engine/Plugins/VirtualProduction/RemoteControl/Source/RemoteControl/Private/RemoteControlModule.cpp b/Engine/Plugins/VirtualProduction/RemoteControl/Source/RemoteControl/Private/RemoteControlModule.cpp
--- a/Engine/Plugins/VirtualProduction/RemoteControl/Source/RemoteControl/Private/RemoteControlModule.cpp
+++ b/Engine/Plugins/VirtualProduction/RemoteControl/Source/RemoteControl/Private/RemoteControlModule.cpp
@@ -1169,14 +1169,18 @@ bool FRemoteControlModule::SetObjectProperties(const FRCObjectReference& ObjectA
 				}
 
 				// Update the world lighting if we're modifying a color.
-				if (MutableObjectReference.IsValid() && MutableObjectReference.Property->GetOwnerClass()->IsChildOf(ULightComponentBase::StaticClass()))
+				if (MutableObjectReference.IsValid() )
 				{
-					UWorld* World = MutableObjectReference.Object->GetWorld();
-					if (World && World->Scene)
+					if (MutableObjectReference.Object->IsA<ULightComponent>())
 					{
-						if (MutableObjectReference.Object->IsA<ULightComponent>())
+						UClass* OwnerClass = MutableObjectReference.Property->GetOwnerClass();
+						if (OwnerClass && OwnerClass->IsChildOf(ULightComponentBase::StaticClass()))
 						{
-							World->Scene->UpdateLightColorAndBrightness(Cast<ULightComponent>(MutableObjectReference.Object.Get()));
+							UWorld* World = MutableObjectReference.Object->GetWorld();
+							if (World && World->Scene)
+							{
+								World->Scene->UpdateLightColorAndBrightness(Cast<ULightComponent>(MutableObjectReference.Object.Get()));
+							}
 						}
 					}
 				}
