[UE-156705] Quality of Life upgrades for Logic UI

Featuring:

* Rename Controller [UE-156717] - Provides ability to rename Controllers via API support and Rename UI command.

* Controllers Table - header & resizing support

~~~

[Rename Controller]

Functional notes:
* Provides rename via F2 key (UI command). Rename via RMB Context Menu wiill be part of a subsequent CL.
* Web API has been updated in this CL to support using the user-friendly Controller name (renamed by users as desired) rather than the internal proprety name (which is currently immutable, see impl notes below)

Impl notes:
* Added a new property DisplayName to VirtualProperty class to support specifiying user frienedly Controller Names.
* I deemed this preferrable to modifying the underlying PropertyBag description because that approach will reallocate memory for all the Controllers each time the user tries to rename even just one. This is due to the way PropertyBag is designed (each collection of properties is conidered as a unique type with a specific hash value.The entire set is reallocated if any modification is made to the property bag).

[Controllers Table]
* Used Header Rows for automatic resizing
* Preferred over using manual splitter (as in EP list), no issues noted thusfar

diff --git a/Engine/Plugins/VirtualProduction/RemoteControl/Source/RemoteControl/Private/RCVirtualProperty.cpp b/Engine/Plugins/VirtualProduction/RemoteControl/Source/RemoteControl/Private/RCVirtualProperty.cpp
--- a/Engine/Plugins/VirtualProduction/RemoteControl/Source/RemoteControl/Private/RCVirtualProperty.cpp
+++ b/Engine/Plugins/VirtualProduction/RemoteControl/Source/RemoteControl/Private/RCVirtualProperty.cpp
@@ -494,12 +494,7 @@ bool URCVirtualPropertyBase::SetValueColor(const FColor& InColor)
 
 FName URCVirtualPropertyBase::GetPropertyName() const
 {
-	if (const FProperty* Property = GetProperty())
-	{
-		return Property->GetClass()->GetFName();
-	}
-	
-	return NAME_None;
+	return PropertyName;
 }
 
 const FInstancedPropertyBag* URCVirtualPropertyInContainer::GetPropertyBagInstance() const
diff --git a/Engine/Plugins/VirtualProduction/RemoteControl/Source/RemoteControl/Private/RCVirtualPropertyContainer.cpp b/Engine/Plugins/VirtualProduction/RemoteControl/Source/RemoteControl/Private/RCVirtualPropertyContainer.cpp
--- a/Engine/Plugins/VirtualProduction/RemoteControl/Source/RemoteControl/Private/RCVirtualPropertyContainer.cpp
+++ b/Engine/Plugins/VirtualProduction/RemoteControl/Source/RemoteControl/Private/RCVirtualPropertyContainer.cpp
@@ -139,7 +139,7 @@ void URCVirtualPropertyContainerBase::Reset()
 	Bag.Reset();
 }
 
-URCVirtualPropertyBase* URCVirtualPropertyContainerBase::GetVirtualProperty(const FName& InPropertyName) const
+URCVirtualPropertyBase* URCVirtualPropertyContainerBase::GetVirtualProperty(const FName InPropertyName) const
 {
 	for (URCVirtualPropertyBase* VirtualProperty : VirtualProperties)
 	{
@@ -152,6 +152,19 @@ URCVirtualPropertyBase* URCVirtualPropertyContainerBase::GetVirtualProperty(cons
 	return nullptr;
 }
 
+URCVirtualPropertyBase* URCVirtualPropertyContainerBase::GetVirtualPropertyByDisplayName(const FName InDisplayName) const
+{
+	for (URCVirtualPropertyBase* VirtualProperty : VirtualProperties)
+	{
+		if (VirtualProperty->DisplayName == InDisplayName)
+		{
+			return VirtualProperty;
+		}
+	}
+
+	return nullptr;
+}
+
 int32 URCVirtualPropertyContainerBase::GetNumVirtualProperties() const
 {
 	const int32 NumPropertiesInBag = Bag.GetNumPropertiesInBag();
diff --git a/Engine/Plugins/VirtualProduction/RemoteControl/Source/RemoteControl/Private/RemoteControlModule.cpp b/Engine/Plugins/VirtualProduction/RemoteControl/Source/RemoteControl/Private/RemoteControlModule.cpp
--- a/Engine/Plugins/VirtualProduction/RemoteControl/Source/RemoteControl/Private/RemoteControlModule.cpp
+++ b/Engine/Plugins/VirtualProduction/RemoteControl/Source/RemoteControl/Private/RemoteControlModule.cpp
@@ -1220,7 +1220,7 @@ bool FRemoteControlModule::SetPresetController(const FName PresetName, const FNa
 		return false;
 	}
 
-	URCVirtualPropertyBase* VirtualProperty = Preset->GetVirtualProperty(ControllerName);
+	URCVirtualPropertyBase* VirtualProperty = Preset->GetVirtualPropertyByDisplayName(ControllerName);
 	if (!ensure(VirtualProperty))
 	{
 		return false;
diff --git a/Engine/Plugins/VirtualProduction/RemoteControl/Source/RemoteControl/Private/RemoteControlPreset.cpp b/Engine/Plugins/VirtualProduction/RemoteControl/Source/RemoteControl/Private/RemoteControlPreset.cpp
--- a/Engine/Plugins/VirtualProduction/RemoteControl/Source/RemoteControl/Private/RemoteControlPreset.cpp
+++ b/Engine/Plugins/VirtualProduction/RemoteControl/Source/RemoteControl/Private/RemoteControlPreset.cpp
@@ -499,7 +499,7 @@ FName URemoteControlPreset::GenerateUniqueLabel(const FName InDesiredName) const
 	return Registry->GenerateUniqueLabel(InDesiredName);
 }
 
-URCVirtualPropertyBase* URemoteControlPreset::GetVirtualProperty(const FName& InPropertyName) const
+URCVirtualPropertyBase* URemoteControlPreset::GetVirtualProperty(const FName InPropertyName) const
 {
 	if (!ensure(ControllerContainer))
 	{
@@ -509,6 +509,16 @@ URCVirtualPropertyBase* URemoteControlPreset::GetVirtualProperty(const FName& In
 	return ControllerContainer->GetVirtualProperty(InPropertyName);
 }
 
+URCVirtualPropertyBase* URemoteControlPreset::GetVirtualPropertyByDisplayName(const FName InDisplayName) const
+{
+	if (!ensure(ControllerContainer))
+	{
+		return nullptr;
+	}
+
+	return ControllerContainer->GetVirtualPropertyByDisplayName(InDisplayName);
+}
+
 URCVirtualPropertyInContainer* URemoteControlPreset::AddVirtualProperty(TSubclassOf<URCVirtualPropertyInContainer> InPropertyClass, const EPropertyBagPropertyType InValueType, UObject* InValueTypeObject /*= nullptr*/, const FName InPropertyName /*= NAME_None*/)
 {
 	if (!ensure(ControllerContainer))
diff --git a/Engine/Plugins/VirtualProduction/RemoteControl/Source/RemoteControl/Public/RCVirtualProperty.h b/Engine/Plugins/VirtualProduction/RemoteControl/Source/RemoteControl/Public/RCVirtualProperty.h
--- a/Engine/Plugins/VirtualProduction/RemoteControl/Source/RemoteControl/Public/RCVirtualProperty.h
+++ b/Engine/Plugins/VirtualProduction/RemoteControl/Source/RemoteControl/Public/RCVirtualProperty.h
@@ -243,6 +243,10 @@ public:
 	/** Pointer to Remote Control Preset */
 	UPROPERTY()
 	TWeakObjectPtr<URemoteControlPreset> PresetWeakPtr;
+
+	/** User friendly name of the Controller*/
+	UPROPERTY()
+	FName DisplayName;
 };
 
 /**
diff --git a/Engine/Plugins/VirtualProduction/RemoteControl/Source/RemoteControl/Public/RCVirtualPropertyContainer.h b/Engine/Plugins/VirtualProduction/RemoteControl/Source/RemoteControl/Public/RCVirtualPropertyContainer.h
--- a/Engine/Plugins/VirtualProduction/RemoteControl/Source/RemoteControl/Public/RCVirtualPropertyContainer.h
+++ b/Engine/Plugins/VirtualProduction/RemoteControl/Source/RemoteControl/Public/RCVirtualPropertyContainer.h
@@ -73,7 +73,12 @@ public:
 	/**
 	 * Returns virtual property by specified name.
 	 */ 
-	virtual URCVirtualPropertyBase* GetVirtualProperty(const FName& InPropertyName) const;
+	virtual URCVirtualPropertyBase* GetVirtualProperty(const FName InPropertyName) const;
+
+	/**
+	 * Returns virtual property by user-friendly display name (Controller Name)
+	 */
+	virtual URCVirtualPropertyBase* GetVirtualPropertyByDisplayName(const FName InDisplayName) const;
 
 	/**
 	 * Returns number of virtual properties.
diff --git a/Engine/Plugins/VirtualProduction/RemoteControl/Source/RemoteControl/Public/RemoteControlPreset.h b/Engine/Plugins/VirtualProduction/RemoteControl/Source/RemoteControl/Public/RemoteControlPreset.h
--- a/Engine/Plugins/VirtualProduction/RemoteControl/Source/RemoteControl/Public/RemoteControlPreset.h
+++ b/Engine/Plugins/VirtualProduction/RemoteControl/Source/RemoteControl/Public/RemoteControlPreset.h
@@ -608,8 +608,11 @@ public:
 	* The goal is to hide the Controller Container and provide a simple interface for Controller access to UI and Web.
 	*/
 
+	/** Fetches a virtual property by internal property name. */
+	URCVirtualPropertyBase* GetVirtualProperty(const FName InPropertyName) const;
+
 	/** Fetches a virtual property by specified name. */
-	URCVirtualPropertyBase* GetVirtualProperty(const FName& InPropertyName) const;
+	URCVirtualPropertyBase* GetVirtualPropertyByDisplayName(const FName InDisplayName) const;
 
 	/** Adds a Virtual Property (Controller) to the Remote Control Preset */
 	URCVirtualPropertyInContainer* AddVirtualProperty(TSubclassOf<URCVirtualPropertyInContainer> InPropertyClass, const EPropertyBagPropertyType InValueType, UObject* InValueTypeObject = nullptr, const FName InPropertyName = NAME_None);
diff --git a/Engine/Plugins/VirtualProduction/RemoteControl/Source/RemoteControlUI/Private/UI/Controller/RCControllerModel.cpp b/Engine/Plugins/VirtualProduction/RemoteControl/Source/RemoteControlUI/Private/UI/Controller/RCControllerModel.cpp
--- a/Engine/Plugins/VirtualProduction/RemoteControl/Source/RemoteControlUI/Private/UI/Controller/RCControllerModel.cpp
+++ b/Engine/Plugins/VirtualProduction/RemoteControl/Source/RemoteControlUI/Private/UI/Controller/RCControllerModel.cpp
@@ -6,12 +6,24 @@
 #include "RemoteControlPreset.h"
 #include "RCVirtualProperty.h"
 #include "UI/SRemoteControlPanel.h"
+#include "Widgets/Text/SInlineEditableTextBlock.h"
 
-FRCControllerModel::FRCControllerModel(const FName& InPropertyName, const TSharedRef<IDetailTreeNode>& InTreeNode, const TSharedPtr<SRemoteControlPanel> InRemoteControlPanel)
+#define LOCTEXT_NAMESPACE "FRCControllerModel"
+
+FRCControllerModel::FRCControllerModel(URCVirtualPropertyBase* InVirtualProperty, const TSharedRef<IDetailTreeNode>& InTreeNode, const TSharedPtr<SRemoteControlPanel> InRemoteControlPanel)
 	: FRCLogicModeBase(InRemoteControlPanel)
-	, DetailTreeNodeWeakPtr(InTreeNode)
-	, PropertyName(InPropertyName)
+	, VirtualPropertyWeakPtr(InVirtualProperty)
+	, DetailTreeNodeWeakPtr(InTreeNode)	
 {
+	if(ensure(InVirtualProperty))
+	{
+		if(InVirtualProperty->DisplayName.IsNone())
+			InVirtualProperty->DisplayName = InVirtualProperty->PropertyName;
+
+		SAssignNew(ControllerNameTextBox, SInlineEditableTextBlock)
+			.Text(FText::FromName(InVirtualProperty->DisplayName))
+			.OnTextCommitted_Raw(this, &FRCControllerModel::OnControllerNameCommitted);
+	}
 }
 
 TSharedRef<SWidget> FRCControllerModel::GetWidget() const
@@ -39,39 +51,28 @@ TSharedRef<SWidget> FRCControllerModel::GetWidget() const
 				NodeWidgets.WholeRowWidget.ToSharedRef()
 			];
 	}
-	
-	return 	SNew(SHorizontalBox)
-		.Clipping(EWidgetClipping::OnDemand)
-		// Field name
-		+ SHorizontalBox::Slot()
-		.VAlign(VAlign_Center)
-		.AutoWidth()		
-		[
-			NodeWidgets.NameWidget.ToSharedRef()
-		]
-		// Value Widget
-		+ SHorizontalBox::Slot()
-		.VAlign(VAlign_Center)
-		.AutoWidth()
-		.Padding(FMargin(3.f))
-		[
-			FieldWidget
-		];
+
+	return FieldWidget;
 }
 
-URCVirtualPropertyBase* FRCControllerModel::GetVirtualProperty() const
+TSharedRef<SWidget> FRCControllerModel::GetNameWidget() const
 {
-	if (const URemoteControlPreset* Preset = GetPreset())
-	{
-		return Preset->GetVirtualProperty(PropertyName);
-	}
+	return ControllerNameTextBox.ToSharedRef();
+}
 
-	return nullptr;
+URCVirtualPropertyBase* FRCControllerModel::GetVirtualProperty() const
+{
+	return VirtualPropertyWeakPtr.Get();
 }
 
-const FName& FRCControllerModel::GetPropertyName() const
+const FName FRCControllerModel::GetPropertyName() const
 {
-	return PropertyName;
+	if (const URCVirtualPropertyBase* VirtualProperty = GetVirtualProperty())
+	{
+		return VirtualProperty->PropertyName;
+	}
+
+	return NAME_None;
 }
 
 TSharedPtr<FRCBehaviourModel> FRCControllerModel::GetSelectedBehaviourModel() const
@@ -82,4 +83,33 @@ TSharedPtr<FRCBehaviourModel> FRCControllerModel::GetSelectedBehaviourModel() co
 void FRCControllerModel::UpdateSelectedBehaviourModel(TSharedPtr<FRCBehaviourModel> InModel)
 {
 	SelectedBehaviourModelWeakPtr = InModel;
-}
\ No newline at end of file
+}
+
+void FRCControllerModel::OnControllerNameCommitted(const FText& InNewControllerName, ETextCommit::Type InCommitInfo)
+{
+	if (URemoteControlPreset* Preset = GetPreset())
+	{
+		if(URCVirtualPropertyBase* Controller = GetVirtualProperty())
+		{
+			Controller->DisplayName = *InNewControllerName.ToString();
+			ControllerNameTextBox->SetText(InNewControllerName);
+		}
+	}
+}
+
+void FRCControllerModel::EnterRenameMode()
+{
+	ControllerNameTextBox->EnterEditingMode();
+}
+
+FName FRCControllerModel::GetControllerDisplayName()
+{
+	if (URCVirtualPropertyBase* Controller = GetVirtualProperty())
+	{
+		return Controller->DisplayName;
+	}
+
+	return NAME_None;
+}
+
+#undef LOCTEXT_NAMESPACE
\ No newline at end of file
diff --git a/Engine/Plugins/VirtualProduction/RemoteControl/Source/RemoteControlUI/Private/UI/Controller/RCControllerModel.h b/Engine/Plugins/VirtualProduction/RemoteControl/Source/RemoteControlUI/Private/UI/Controller/RCControllerModel.h
--- a/Engine/Plugins/VirtualProduction/RemoteControl/Source/RemoteControlUI/Private/UI/Controller/RCControllerModel.h
+++ b/Engine/Plugins/VirtualProduction/RemoteControl/Source/RemoteControlUI/Private/UI/Controller/RCControllerModel.h
@@ -4,9 +4,12 @@
 
 #include "UI/BaseLogicUI/RCLogicModeBase.h"
 #include "UObject/WeakFieldPtr.h"
+#include "Widgets/DeclarativeSyntaxSupport.h"
+#include "Widgets/SWidget.h"
 
 class FRCBehaviourModel;
 class IDetailTreeNode;
+class SInlineEditableTextBlock;
 class SRemoteControlPanel;
 class URCBehaviour;
 class URCControllerContainer;
@@ -23,18 +26,20 @@ class URCVirtualPropertyBase;
 class FRCControllerModel : public FRCLogicModeBase
 {
 public:
-	FRCControllerModel(const FName& PropertyName, const TSharedRef<IDetailTreeNode>& InTreeNode, const TSharedPtr<SRemoteControlPanel> InRemoteControlPanel);
+	FRCControllerModel(URCVirtualPropertyBase* InVirtualProperty, const TSharedRef<IDetailTreeNode>& InTreeNode, const TSharedPtr<SRemoteControlPanel> InRemoteControlPanel);
 	
 	/** The widget to be rendered for this Controller
 	* Used to represent a single row when added to the Controllers Panel List
 	*/
 	virtual TSharedRef<SWidget> GetWidget() const override;
 
+	TSharedRef<SWidget> GetNameWidget() const;
+
 	/** Returns the Controller (Data Model) associated with us*/
 	virtual URCVirtualPropertyBase* GetVirtualProperty() const;
 	
 	/** Returns the Controller Name (from Virtual Property) represented by us*/
-	virtual const FName& GetPropertyName() const;
+	virtual const FName GetPropertyName() const;
 
 	/** Returns the currently selected Behaviour (UI model) */
 	TSharedPtr<FRCBehaviourModel> GetSelectedBehaviourModel() const;
@@ -42,9 +47,18 @@ public:
 	/** Updates our internal record of the currently selected Behaviour (UI Model) */
 	void UpdateSelectedBehaviourModel(TSharedPtr<FRCBehaviourModel> InModel);
 
+	/** Allows users to enter text into the Controller Name Box */
+	void EnterRenameMode();
+
+	/** User-friendly Name of the underlying Controller */
+	FName GetControllerDisplayName();
+
 private:
-	/** The Parent Container holding our Controller (Virtual Property) Data Model */
-	TWeakObjectPtr<URCVirtualPropertyContainerBase> PropertyContainerWeakPtr;
+	/* Text commit event for Controller Name text box*/
+	void OnControllerNameCommitted(const FText& InNewControllerName, ETextCommit::Type InCommitInfo);
+
+	/** Virtual Property (Data Model) represented by us*/
+	TWeakObjectPtr<URCVirtualPropertyBase> VirtualPropertyWeakPtr;
 	
 	/** The row generator used to build a generic Value Widget*/
 	TWeakPtr<IDetailTreeNode> DetailTreeNodeWeakPtr;
@@ -52,6 +66,6 @@ private:
 	/** The currently selected Behaviour (UI model) */
 	TWeakPtr<FRCBehaviourModel>  SelectedBehaviourModelWeakPtr;
 
-	/** The Controller (Virtual Property Name) represented by us*/
-	FName PropertyName;
+	/** Controller name - editable text box */
+	TSharedPtr<SInlineEditableTextBlock> ControllerNameTextBox;
 };
diff --git a/Engine/Plugins/VirtualProduction/RemoteControl/Source/RemoteControlUI/Private/UI/Controller/SRCControllerPanel.cpp b/Engine/Plugins/VirtualProduction/RemoteControl/Source/RemoteControlUI/Private/UI/Controller/SRCControllerPanel.cpp
--- a/Engine/Plugins/VirtualProduction/RemoteControl/Source/RemoteControlUI/Private/UI/Controller/SRCControllerPanel.cpp
+++ b/Engine/Plugins/VirtualProduction/RemoteControl/Source/RemoteControlUI/Private/UI/Controller/SRCControllerPanel.cpp
@@ -45,5 +45,16 @@ bool SRCControllerPanel::IsListFocused() const
 
 void SRCControllerPanel::DeleteSelectedPanelItem()
 {
-	ControllerPanelList->DeleteSelectedPanelItem();
-}
\ No newline at end of file
+	if (ControllerPanelList)
+	{
+		ControllerPanelList->DeleteSelectedPanelItem();
+	}
+}
+
+void SRCControllerPanel::EnterRenameMode()
+{
+	if (ControllerPanelList)
+	{
+		ControllerPanelList->EnterRenameMode();
+	}
+}
diff --git a/Engine/Plugins/VirtualProduction/RemoteControl/Source/RemoteControlUI/Private/UI/Controller/SRCControllerPanel.h b/Engine/Plugins/VirtualProduction/RemoteControl/Source/RemoteControlUI/Private/UI/Controller/SRCControllerPanel.h
--- a/Engine/Plugins/VirtualProduction/RemoteControl/Source/RemoteControlUI/Private/UI/Controller/SRCControllerPanel.h
+++ b/Engine/Plugins/VirtualProduction/RemoteControl/Source/RemoteControlUI/Private/UI/Controller/SRCControllerPanel.h
@@ -32,4 +32,6 @@ public:
 
 	/** Delete Item UI command implementation for this panel */
 	virtual void DeleteSelectedPanelItem() override;
+
+	void EnterRenameMode();
 };
diff --git a/Engine/Plugins/VirtualProduction/RemoteControl/Source/RemoteControlUI/Private/UI/Controller/SRCControllerPanelList.cpp b/Engine/Plugins/VirtualProduction/RemoteControl/Source/RemoteControlUI/Private/UI/Controller/SRCControllerPanelList.cpp
--- a/Engine/Plugins/VirtualProduction/RemoteControl/Source/RemoteControlUI/Private/UI/Controller/SRCControllerPanelList.cpp
+++ b/Engine/Plugins/VirtualProduction/RemoteControl/Source/RemoteControlUI/Private/UI/Controller/SRCControllerPanelList.cpp
@@ -13,6 +13,44 @@
 #include "UI/BaseLogicUI/RCLogicModeBase.h"
 #include "UI/SRemoteControlPanel.h"
 
+#define LOCTEXT_NAMESPACE "SRCControllerPanelList"
+
+namespace UE::RCControllerPanelList
+{
+	static const FName ControllerNameColumn(TEXT("Controller Name"));
+	static const FName ControllerValueColumn(TEXT("Controller Value"));
+
+	class SControllerItemListRow : public SMultiColumnTableRow<TSharedRef<FRCControllerModel>>
+	{
+	private:
+		TSharedPtr<FRCControllerModel> ControllerItem;
+
+	public:
+		void Construct(const FTableRowArgs& InArgs, const TSharedRef<STableViewBase>& OwnerTableView, TSharedRef<FRCControllerModel> InControllerItem)
+		{
+			ControllerItem = InControllerItem;
+			FSuperRowType::Construct(InArgs, OwnerTableView);
+		}
+
+		TSharedRef<SWidget> GenerateWidgetForColumn(const FName& ColumnName) override
+		{
+			if (!ensure(ControllerItem.IsValid()))
+				return SNullWidget::NullWidget;
+
+			if (ColumnName == UE::RCControllerPanelList::ControllerNameColumn)
+			{
+				return ControllerItem->GetNameWidget();
+			}
+			else if (ColumnName == UE::RCControllerPanelList::ControllerValueColumn)
+			{
+				return ControllerItem->GetWidget();
+			}
+
+			return SNullWidget::NullWidget;
+		}
+	};
+}
+
 BEGIN_SLATE_FUNCTION_BUILD_OPTIMIZATION
 
 void SRCControllerPanelList::Construct(const FArguments& InArgs, const TSharedRef<SRCControllerPanel> InControllerPanel)
@@ -22,10 +60,23 @@ void SRCControllerPanelList::Construct(const FArguments& InArgs, const TSharedRe
 	ControllerPanelWeakPtr = InControllerPanel;
 	
 	ListView = SNew(SListView<TSharedPtr<FRCControllerModel>>)
-		.ListItemsSource( &ControllerItems )
+		.ListItemsSource(&ControllerItems)
 		.OnSelectionChanged(this, &SRCControllerPanelList::OnTreeSelectionChanged)
-		.OnGenerateRow(this, &SRCControllerPanelList::OnGenerateWidgetForList );
-	
+		.OnGenerateRow(this, &SRCControllerPanelList::OnGenerateWidgetForList)
+		.SelectionMode(ESelectionMode::Single) // Current setup supports only single selection (and related display) of a Controller in the list
+		.HeaderRow
+		(
+			SNew(SHeaderRow)
+
+			+ SHeaderRow::Column(UE::RCControllerPanelList::ControllerNameColumn)
+			.DefaultLabel(LOCTEXT("Controller Name Column Name", "Name"))
+			.FillWidth(0.25f)
+
+			+ SHeaderRow::Column(UE::RCControllerPanelList::ControllerValueColumn)
+			.DefaultLabel(LOCTEXT("Controller Value Column Name", "Input"))
+			.FillWidth(0.75f)
+		);
+
 	ChildSlot
 	[
 		ListView.ToSharedRef()
@@ -74,7 +125,11 @@ void SRCControllerPanelList::Reset()
 		{
 			FProperty* Property = Child->CreatePropertyHandle()->GetProperty();
 			check(Property);
-			ControllerItems.Add(MakeShared<FRCControllerModel>(Property->GetFName(), Child, RemoteControlPanel));
+
+			if (URCVirtualPropertyBase* Controller = Preset->GetVirtualProperty(Property->GetFName()))
+			{
+				ControllerItems.Add(MakeShared<FRCControllerModel>(Controller, Child, RemoteControlPanel));
+			}
 		}
 	}
 
@@ -83,10 +138,13 @@ void SRCControllerPanelList::Reset()
 
 TSharedRef<ITableRow> SRCControllerPanelList::OnGenerateWidgetForList(TSharedPtr<FRCControllerModel> InItem, const TSharedRef<STableViewBase>& OwnerTable)
 {
-	return SNew(STableRow<TSharedPtr<FString>>, OwnerTable)
-	[
-		InItem->GetWidget()
-	];
+	typedef UE::RCControllerPanelList::SControllerItemListRow ControllerRowType;
+
+	const TSharedRef<ControllerRowType> NewRow =
+		SNew(ControllerRowType, OwnerTable, InItem.ToSharedRef())
+		.Padding(FMargin(3.f));
+
+	return NewRow;
 }
 
 void SRCControllerPanelList::OnTreeSelectionChanged(TSharedPtr<FRCControllerModel> InItem, ESelectInfo::Type)
@@ -188,3 +246,13 @@ void SRCControllerPanelList::DeleteSelectedPanelItem()
 {
 	DeleteItemFromLogicPanel<FRCControllerModel>(ControllerItems, ListView->GetSelectedItems());
 }
+
+void SRCControllerPanelList::EnterRenameMode()
+{
+	if (TSharedPtr<FRCControllerModel> SelectedItem = SelectedControllerItemWeakPtr.Pin())
+	{
+		SelectedItem->EnterRenameMode();
+	}
+}
+
+#undef LOCTEXT_NAMESPACE
\ No newline at end of file
diff --git a/Engine/Plugins/VirtualProduction/RemoteControl/Source/RemoteControlUI/Private/UI/Controller/SRCControllerPanelList.h b/Engine/Plugins/VirtualProduction/RemoteControl/Source/RemoteControlUI/Private/UI/Controller/SRCControllerPanelList.h
--- a/Engine/Plugins/VirtualProduction/RemoteControl/Source/RemoteControlUI/Private/UI/Controller/SRCControllerPanelList.h
+++ b/Engine/Plugins/VirtualProduction/RemoteControl/Source/RemoteControlUI/Private/UI/Controller/SRCControllerPanelList.h
@@ -83,5 +83,10 @@ public:
 
 	/** Deletes currently selected items from the list view*/
 	void DeleteSelectedPanelItem();
+
+	void EnterRenameMode();
+
+private:
+	
 };
 
diff --git a/Engine/Plugins/VirtualProduction/RemoteControl/Source/RemoteControlUI/Private/UI/SRemoteControlPanel.cpp b/Engine/Plugins/VirtualProduction/RemoteControl/Source/RemoteControlUI/Private/UI/SRemoteControlPanel.cpp
--- a/Engine/Plugins/VirtualProduction/RemoteControl/Source/RemoteControlUI/Private/UI/SRemoteControlPanel.cpp
+++ b/Engine/Plugins/VirtualProduction/RemoteControl/Source/RemoteControlUI/Private/UI/SRemoteControlPanel.cpp
@@ -1966,6 +1966,12 @@ bool SRemoteControlPanel::CanDeleteEntity() const
 
 void SRemoteControlPanel::RenameEntity_Execute() const
 {
+	if (ControllerPanel->IsListFocused())
+	{
+		ControllerPanel->EnterRenameMode();
+		return;
+	}
+
 	if (SelectedEntity->GetRCType() == SRCPanelTreeNode::FieldChild) // Field Child does not contain entity ID, that is why it should not be processed
 	{
 		return;
@@ -1976,6 +1982,11 @@ void SRemoteControlPanel::RenameEntity_Execute() const
 
 bool SRemoteControlPanel::CanRenameEntity() const
 {
+	if (ControllerPanel->IsListFocused())
+	{
+		return true;
+	}
+
 	if (SelectedEntity.IsValid() && Preset.IsValid() && bIsInEditMode)
 	{
 		// Do not allow default group to be renamed.
diff --git a/Engine/Plugins/VirtualProduction/RemoteControl/Source/WebRemoteControl/Private/WebRemoteControl.cpp b/Engine/Plugins/VirtualProduction/RemoteControl/Source/WebRemoteControl/Private/WebRemoteControl.cpp
--- a/Engine/Plugins/VirtualProduction/RemoteControl/Source/WebRemoteControl/Private/WebRemoteControl.cpp
+++ b/Engine/Plugins/VirtualProduction/RemoteControl/Source/WebRemoteControl/Private/WebRemoteControl.cpp
@@ -2030,7 +2030,7 @@ bool FWebRemoteControlModule::HandlePresetSetControllerRoute(const FHttpServerRe
 	}
 
 	// 6. Acquire Controller
-	URCVirtualPropertyBase* Controller = Preset->GetVirtualProperty(*Args.ControllerName);
+	URCVirtualPropertyBase* Controller = Preset->GetVirtualPropertyByDisplayName(*Args.ControllerName);
 	if (!Controller)
 	{
 		Response->Code = EHttpServerResponseCodes::NotFound;
@@ -2043,8 +2043,9 @@ bool FWebRemoteControlModule::HandlePresetSetControllerRoute(const FHttpServerRe
 	// 7. Reformat Payload to represent our internal structure
 	TArray<uint8> NewPayload;
 	const FName PropertyValueKey = WebRemoteControlStructUtils::Prop_PropertyValue;
+	const FName PropertyNameInternal = Controller->GetPropertyName();
 
-	RemotePayloadSerializer::ReplaceFirstOccurence(SetControllerRequest.TCHARBody, PropertyValueKey.ToString(), Args.ControllerName, NewPayload);
+	RemotePayloadSerializer::ReplaceFirstOccurence(SetControllerRequest.TCHARBody, PropertyValueKey.ToString(), PropertyNameInternal.ToString(), NewPayload);
 	FMemoryReader Reader(NewPayload);
 	FJsonStructDeserializerBackend Backend(Reader);
 
@@ -2089,7 +2090,7 @@ bool FWebRemoteControlModule::HandlePresetGetControllerRoute(const FHttpServerRe
 	}
 
 	// 4. Acquire Controller
-	URCVirtualPropertyBase* Controller = Preset->GetVirtualProperty(*Args.ControllerName);
+	URCVirtualPropertyBase* Controller = Preset->GetVirtualPropertyByDisplayName(*Args.ControllerName);
 	if (!Controller)
 	{
 		Response->Code = EHttpServerResponseCodes::NotFound;
