Update to allow remote control presets to exist within other assets and toolkit hosts (not level editor.)

diff --git a/Engine/Plugins/VirtualProduction/RemoteControl/Source/RemoteControl/Private/RemoteControlBinding.cpp b/Engine/Plugins/VirtualProduction/RemoteControl/Source/RemoteControl/Private/RemoteControlBinding.cpp
--- a/Engine/Plugins/VirtualProduction/RemoteControl/Source/RemoteControl/Private/RemoteControlBinding.cpp
+++ b/Engine/Plugins/VirtualProduction/RemoteControl/Source/RemoteControl/Private/RemoteControlBinding.cpp
@@ -34,6 +34,8 @@ namespace
 
 	/**
 	 * Find the counterpart actor/component in PIE/Editor
+	 * 
+	 * @TODO Does this need to be updated to possibly include a non-editor related preset?
 	 */
 	UObject* FindObjectInCounterpartWorld(UObject* Object, ECounterpartWorldTarget WorldTarget)
 	{
@@ -368,36 +370,12 @@ UClass* URemoteControlLevelDependantBinding::GetSupportedOwnerClass() const
 	return BindingContext.OwnerActorClass.LoadSynchronous();
 }
 
-UWorld* URemoteControlLevelDependantBinding::GetCurrentWorld()
+UWorld* URemoteControlLevelDependantBinding::GetCurrentWorld() const
 {
-	// Since this is used to retrieve the binding in the map, we never use the PIE world in editor.
-	UWorld* World = nullptr;
-
-#if WITH_EDITOR
-	if (GEditor && FApp::CanEverRender())
-	{
-		World = GEditor->GetEditorWorldContext(false).World();
-	}
-#endif
-
-	if (World)
-	{
-		return World;
-	}
+	URemoteControlPreset* Preset = Cast<URemoteControlPreset>(GetOuter());
 
-	if (GEngine)
-	{
-		for (const FWorldContext& WorldContext : GEngine->GetWorldContexts())
-		{
-			if (WorldContext.WorldType == EWorldType::Game)
-			{
-				World = WorldContext.World();
-				break;
-			}
-		}
-	}
-		
-	return World;
+	// Since this is used to retrieve the binding in the map, we never use the PIE world in editor.
+	return URemoteControlPreset::GetPresetWorld(Preset, false);
 }
 
 void URemoteControlLevelDependantBinding::SetBoundObject(const TSoftObjectPtr<ULevel>& Level, const TSoftObjectPtr<UObject>& BoundObject)
diff --git a/Engine/Plugins/VirtualProduction/RemoteControl/Source/RemoteControl/Private/RemoteControlModule.cpp b/Engine/Plugins/VirtualProduction/RemoteControl/Source/RemoteControl/Private/RemoteControlModule.cpp
--- a/Engine/Plugins/VirtualProduction/RemoteControl/Source/RemoteControl/Private/RemoteControlModule.cpp
+++ b/Engine/Plugins/VirtualProduction/RemoteControl/Source/RemoteControl/Private/RemoteControlModule.cpp
@@ -639,6 +639,85 @@ void FRemoteControlModule::UnregisterPreset(FName Name)
 {
 }
 
+bool FRemoteControlModule::RegisterEmbeddedPreset(URemoteControlPreset* Preset, bool bReplaceExisting)
+{
+	if (!Preset)
+	{
+		return false;
+	}
+
+	FName PresetName = Preset->GetPresetName();
+
+	if (PresetName == NAME_None)
+	{
+		return false;
+	}
+
+	if (const TWeakObjectPtr<URemoteControlPreset>* FoundPreset = EmbeddedPresets.Find(PresetName))
+	{
+		if (FoundPreset->IsValid() && !bReplaceExisting)
+		{
+			return false;
+		}
+	}
+
+	EmbeddedPresets.Emplace(PresetName, TWeakObjectPtr<URemoteControlPreset>(Preset));
+
+	FGuid PresetId = Preset->GetPresetId();
+
+	if (PresetId.IsValid())
+	{
+		CachedPresetNamesById.Emplace(PresetId, PresetName);
+	}
+
+	return true;
+}
+
+void FRemoteControlModule::UnregisterEmbeddedPreset(FName Name)
+{
+	if (Name == NAME_None)
+	{
+		return;
+	}
+
+	// Check the cached preset ids and remove it if our name matches the stored id
+	TWeakObjectPtr<URemoteControlPreset>* FoundPreset = EmbeddedPresets.Find(Name);
+
+	if (FoundPreset && FoundPreset->IsValid())
+	{
+		UnregisterEmbeddedPreset(FoundPreset->Get());
+	}
+}
+
+void FRemoteControlModule::UnregisterEmbeddedPreset(URemoteControlPreset* Preset)
+{
+	if (!Preset)
+	{
+		return;
+	}
+
+	FName PresetName = Preset->GetPresetName();
+
+	if (PresetName == NAME_None)
+	{
+		return;
+	}
+
+	FGuid PresetId = Preset->GetPresetId();
+
+	if (PresetId.IsValid())
+	{
+		FName* FoundPresetName = CachedPresetNamesById.Find(PresetId);
+
+		if (FoundPresetName && *FoundPresetName == PresetName)
+		{
+			CachedPresetNamesById.Remove(PresetId);
+		}
+	}
+
+	EmbeddedPresets.Remove(PresetName);
+}
+
 bool FRemoteControlModule::ResolveCall(const FString& ObjectPath, const FString& FunctionName, FRCCallReference& OutCallRef, FString* OutErrorText)
 {
 	bool bSuccess = true;
@@ -1381,6 +1460,14 @@ TOptional<FExposedProperty> FRemoteControlModule::ResolvePresetProperty(const FR
 
 URemoteControlPreset* FRemoteControlModule::ResolvePreset(FName PresetName) const
 {
+	if (const TWeakObjectPtr<URemoteControlPreset>* EmbeddedPreset = EmbeddedPresets.Find(PresetName))
+	{
+		if (EmbeddedPreset->IsValid())
+		{
+			return EmbeddedPreset->Get();
+		}
+	}
+
 	if (const TArray<FAssetData>* Assets = CachedPresetsByName.Find(PresetName))
 	{
 		for (const FAssetData& Asset : *Assets)
@@ -1392,6 +1479,10 @@ URemoteControlPreset* FRemoteControlModule::ResolvePreset(FName PresetName) cons
 		}
 	}
 
+	/**
+	 * No need to cache hosted preset names - they should already be in a short enough list
+	 * mapped to their name.
+	 */
 	if (URemoteControlPreset* FoundPreset = RemoteControlUtil::GetPresetByName(PresetName))
 	{
 		CachedPresetsByName.FindOrAdd(PresetName).AddUnique(FoundPreset);
@@ -1407,6 +1498,14 @@ URemoteControlPreset* FRemoteControlModule::ResolvePreset(const FGuid& PresetId)
 
 	if (const FName* AssetName = CachedPresetNamesById.Find(PresetId))
 	{
+		if (const TWeakObjectPtr<URemoteControlPreset>* EmbeddedPreset = EmbeddedPresets.Find(*AssetName))
+		{
+			if (EmbeddedPreset->IsValid())
+			{
+				return EmbeddedPreset->Get();
+			}
+		}
+
 		if (const TArray<FAssetData>* Assets = CachedPresetsByName.Find(*AssetName))
 		{
 			for (const FAssetData& Asset : *Assets)
@@ -1423,6 +1522,21 @@ URemoteControlPreset* FRemoteControlModule::ResolvePreset(const FGuid& PresetId)
 		}
 	}
 
+	for (const auto& Pair : EmbeddedPresets)
+	{
+		if (Pair.Value.IsValid())
+		{
+			URemoteControlPreset* EmbeddedPreset = Pair.Value.Get();
+			FGuid EmbeddedPresetId = Pair.Value->GetPresetId();
+
+			if (EmbeddedPresetId.IsValid() && EmbeddedPresetId == PresetId)
+			{
+				CachedPresetNamesById.Emplace(PresetId, Pair.Key);
+				return EmbeddedPreset;
+			}
+		}
+	}
+
 	if (URemoteControlPreset* FoundPreset = RemoteControlUtil::GetPresetById(PresetId))
 	{
 		CachedPresetNamesById.Emplace(PresetId, FoundPreset->GetName());
@@ -1548,6 +1662,17 @@ void FRemoteControlModule::GetPresetAssets(TArray<FAssetData>& OutPresetAssets,
 	}
 }
 
+void FRemoteControlModule::GetEmbeddedPresets(TArray<TWeakObjectPtr<URemoteControlPreset>>& OutEmbeddedPresets) const
+{
+	for (const auto& Pair : EmbeddedPresets)
+	{
+		if (Pair.Value.IsValid())
+		{
+			OutEmbeddedPresets.Add(Pair.Value);
+		}
+	}
+}
+
 const TMap<FName, FEntityMetadataInitializer>& FRemoteControlModule::GetDefaultMetadataInitializers() const
 {
 	return DefaultMetadataInitializers;
@@ -1709,13 +1834,73 @@ void FRemoteControlModule::OnAssetRemoved(const FAssetData& AssetData)
 	}
 }
 
-void FRemoteControlModule::OnAssetRenamed(const FAssetData& AssetData, const FString&)
+void FRemoteControlModule::OnAssetRenamed(const FAssetData& AssetData, const FString& OldName)
 {
 	if (AssetData.AssetClassPath != URemoteControlPreset::StaticClass()->GetClassPathName())
 	{
+		// OldName comes in full path format /path/package.object
+		// Assets are registered in package format /path/package
+		FString PackageName = OldName;
+		int32 DotIdx = INDEX_NONE;
+		
+		if (PackageName.FindChar('.', DotIdx))
+		{
+			PackageName = PackageName.Left(DotIdx);
+		}
+
+		FName PackageFName = FName(PackageName);
+		
+		// Find already registered preset with old package name
+		TWeakObjectPtr<URemoteControlPreset> EmbeddedPreset;
+
+		for (const auto& Pair : EmbeddedPresets)
+		{
+			if (Pair.Key == PackageFName)
+			{
+				EmbeddedPreset = Pair.Value;
+
+				// Remove found asset
+				EmbeddedPresets.Remove(Pair.Key);
+				bool bRemovedId = false;
+				
+				// If preset is still valid and it has a valid id, just remove it.
+				if (EmbeddedPreset.IsValid())
+				{
+					FGuid EmbeddedPresetId = EmbeddedPreset->GetPresetId();
+
+					if (EmbeddedPresetId.IsValid())
+					{
+						CachedPresetNamesById.Remove(EmbeddedPresetId);
+						bRemovedId = true;
+					}
+				}
+
+				// If it's not valid or doesn't have a valid id, search for it and remove it.
+				if (!bRemovedId)
+				{
+					for (const auto& PairInner : CachedPresetNamesById)
+					{
+						if (PairInner.Value == Pair.Key)
+						{
+							CachedPresetNamesById.Remove(PairInner.Key);
+						}
+					}
+				}
+
+				break;
+			}
+		}
+
+		// Reregister embedded asset if it's still valid.
+		if (EmbeddedPreset.IsValid())
+		{
+			RegisterEmbeddedPreset(EmbeddedPreset.Get(), true);
+		}
+
 		return;
 	}
 
+	// Update regular asset-based preset
 	const FGuid PresetId = RemoteControlUtil::GetPresetId(AssetData);
 	if (FName* OldPresetName = CachedPresetNamesById.Find(PresetId))
 	{
diff --git a/Engine/Plugins/VirtualProduction/RemoteControl/Source/RemoteControl/Private/RemoteControlModule.h b/Engine/Plugins/VirtualProduction/RemoteControl/Source/RemoteControl/Private/RemoteControlModule.h
--- a/Engine/Plugins/VirtualProduction/RemoteControl/Source/RemoteControl/Private/RemoteControlModule.h
+++ b/Engine/Plugins/VirtualProduction/RemoteControl/Source/RemoteControl/Private/RemoteControlModule.h
@@ -26,6 +26,9 @@ public:
 	virtual FOnPresetUnregistered& OnPresetUnregistered() override;
 	virtual bool RegisterPreset(FName Name, URemoteControlPreset* Preset) override;
 	virtual void UnregisterPreset(FName Name) override;
+	virtual bool RegisterEmbeddedPreset(URemoteControlPreset* Preset, bool bReplaceExisting) override;
+	virtual void UnregisterEmbeddedPreset(FName Name) override;
+	virtual void UnregisterEmbeddedPreset(URemoteControlPreset* Preset) override;
 	virtual bool ResolveCall(const FString& ObjectPath, const FString& FunctionName, FRCCallReference& OutCallRef, FString* OutErrorText) override;
 	virtual bool InvokeCall(FRCCall& InCall, ERCPayloadType InPayloadType = ERCPayloadType::Json, const TArray<uint8>& InInterceptPayload = TArray<uint8>()) override;
 	virtual bool ResolveObject(ERCAccess AccessType, const FString& ObjectPath, const FString& PropertyName, FRCObjectReference& OutObjectRef, FString* OutErrorText = nullptr) override;
@@ -46,6 +49,7 @@ public:
 	virtual bool IsPresetTransient(const FGuid& PresetId) const override;
 	virtual void GetPresets(TArray<TSoftObjectPtr<URemoteControlPreset>>& OutPresets) const override;
 	virtual void GetPresetAssets(TArray<FAssetData>& OutPresetAssets, bool bIncludeTransient = true) const override;
+	virtual void GetEmbeddedPresets(TArray<TWeakObjectPtr<URemoteControlPreset>>& OutEmbeddedPresets) const override;
 	virtual const TMap<FName, FEntityMetadataInitializer>& GetDefaultMetadataInitializers() const override;
 	virtual bool RegisterDefaultEntityMetadata(FName MetadataKey, FEntityMetadataInitializer MetadataInitializer) override;
 	virtual void UnregisterDefaultEntityMetadata(FName MetadataKey) override;
@@ -65,6 +69,7 @@ private:
 	void OnAssetAdded(const FAssetData& AssetData);
 	void OnAssetRemoved(const FAssetData& AssetData);
 	void OnAssetRenamed(const FAssetData& AssetData, const FString&);
+	void OnEmbeddedPresetRenamed(URemoteControlPreset* Preset);
 
 	/** Destroy a transient preset using an object reference. */
 	bool DestroyTransientPreset(URemoteControlPreset* Preset);
@@ -112,6 +117,12 @@ private:
 	/** Cache of ids to preset names. */
 	mutable TMap<FGuid, FName> CachedPresetNamesById;
 
+	/**
+	 * Listed of hosted (non-asset based) presets. 
+	 * We can assume that all hosted presets are "active," so the list should be relatively short.
+	 **/
+	TMap<FName, TWeakObjectPtr<URemoteControlPreset>> EmbeddedPresets;
+
 	/** Temporary presets that aren't saved as assets or directly visible to the editor's user. */
 	TSet<FAssetData> TransientPresets;
 
diff --git a/Engine/Plugins/VirtualProduction/RemoteControl/Source/RemoteControl/Private/RemoteControlPreset.cpp b/Engine/Plugins/VirtualProduction/RemoteControl/Source/RemoteControl/Private/RemoteControlPreset.cpp
--- a/Engine/Plugins/VirtualProduction/RemoteControl/Source/RemoteControl/Private/RemoteControlPreset.cpp
+++ b/Engine/Plugins/VirtualProduction/RemoteControl/Source/RemoteControl/Private/RemoteControlPreset.cpp
@@ -476,6 +476,16 @@ void URemoteControlPreset::BeginDestroy()
 	Super::BeginDestroy();
 }
 
+FName URemoteControlPreset::GetPresetName() const
+{
+	if (!IsEmbeddedPreset())
+	{
+		return GetFName();
+	}
+
+	return GetPackage()->GetFName();
+}
+
 TWeakPtr<FRemoteControlActor> URemoteControlPreset::ExposeActor(AActor* Actor, FRemoteControlPresetExposeArgs Args)
 {
 	check(Actor);
@@ -973,6 +983,75 @@ void URemoteControlPreset::PostLoadProperties()
 	}
 }
 
+UWorld* URemoteControlPreset::GetPresetWorld(const URemoteControlPreset* Preset, bool bAllowPIE /*= false*/)
+{
+	if (Preset)
+	{
+		if (UObject* Outer = Preset->GetOuter())
+		{
+			if (UWorld* OuterWorld = Outer->GetWorld())
+			{
+				return OuterWorld;
+			}
+		}
+	}
+
+#if WITH_EDITOR
+	if (GEditor)
+	{
+		if (bAllowPIE)
+		{
+			FWorldContext* PIEWorldContext = GEditor->GetPIEWorldContext();
+
+			if (PIEWorldContext)
+			{
+				return PIEWorldContext->World();
+			}
+		}
+
+		return GEditor->GetEditorWorldContext(false).World();
+	}
+
+	if (GWorld)
+	{
+		return GWorld;
+	}
+#endif
+
+	for (const FWorldContext& WorldContext : GEngine->GetWorldContexts())
+	{
+		if (WorldContext.WorldType == EWorldType::Game)
+		{
+			return WorldContext.World();
+		}
+	}
+
+	return nullptr;
+}
+
+UWorld* URemoteControlPreset::GetPresetWorld(bool bAllowPIE /*= false*/) const
+{
+	return URemoteControlPreset::GetPresetWorld(this, bAllowPIE);
+}
+
+bool URemoteControlPreset::IsEmbeddedPreset() const
+{
+	UObject* Outer = GetOuter();
+
+	if (!Outer)
+	{
+		return false;
+	}
+
+	// An asset
+	if (Outer->IsA<UPackage>())
+	{
+		return false;
+	}
+
+	return true;
+}
+
 FName URemoteControlPreset::GetEntityName(const FName InDesiredName, UObject* InObject, const FRCFieldPathInfo& InFieldPath) const
 {
 	FName DesiredName = InDesiredName;
diff --git a/Engine/Plugins/VirtualProduction/RemoteControl/Source/RemoteControl/Private/RemoteControlPresetRebindingManager.cpp b/Engine/Plugins/VirtualProduction/RemoteControl/Source/RemoteControl/Private/RemoteControlPresetRebindingManager.cpp
--- a/Engine/Plugins/VirtualProduction/RemoteControl/Source/RemoteControl/Private/RemoteControlPresetRebindingManager.cpp
+++ b/Engine/Plugins/VirtualProduction/RemoteControl/Source/RemoteControl/Private/RemoteControlPresetRebindingManager.cpp
@@ -24,44 +24,17 @@ static TAutoConsoleVariable<int32> CVarRemoteControlRebindingUseLegacyAlgo(TEXT(
 
 namespace RCPresetRebindingManager
 {
-	UWorld* GetCurrentWorld()
-	{
-		// Never use PIE world when searching for a new binding.
-		UWorld* World = nullptr;
-
-#if WITH_EDITOR
-		if (GEditor)
-		{
-			World = GEditor->GetEditorWorldContext(false).World();
-		}
-#endif
-
-		if (World)
-		{
-			return World;
-		}
-
-		for (const FWorldContext& WorldContext : GEngine->GetWorldContexts())
-		{
-			if (WorldContext.WorldType == EWorldType::Game)
-			{
-				return World;
-			}
-		}
-		
-		return nullptr;
-	}
 
-	bool IsValidObjectForRebinding(UObject* InObject)
+	bool IsValidObjectForRebinding(UObject* InObject, UWorld* PresetWorld)
 	{
 		return InObject
 			&& InObject->GetTypedOuter<UPackage>() != GetTransientPackage()
-			&& InObject->GetWorld() == GetCurrentWorld();
+			&& InObject->GetWorld() == PresetWorld;
 	}
 
-	bool IsValidActorForRebinding(AActor* InActor)
+	bool IsValidActorForRebinding(AActor* InActor, UWorld* PresetWorld)
 	{
-		return IsValidObjectForRebinding(InActor) &&
+		return IsValidObjectForRebinding(InActor, PresetWorld) &&
 #if WITH_EDITOR
             InActor->IsEditable() &&
             InActor->IsListedInSceneOutliner() &&
@@ -71,16 +44,16 @@ namespace RCPresetRebindingManager
             !InActor->HasAnyFlags(RF_Transient);
 	}
 
-	bool IsValidComponentForRebinding(UActorComponent* InComponent)
+	bool IsValidComponentForRebinding(UActorComponent* InComponent, UWorld* PresetWorld)
 	{
-		if (!IsValidObjectForRebinding(InComponent))
+		if (!IsValidObjectForRebinding(InComponent, PresetWorld))
 		{
 			return false;
 		}
 		
 		if (AActor* OuterActor = InComponent->GetOwner())
 		{
-			if (!IsValidActorForRebinding(OuterActor))
+			if (!IsValidActorForRebinding(OuterActor, PresetWorld))
 			{
 				return false;
 			}
@@ -98,15 +71,15 @@ namespace RCPresetRebindingManager
 	 * Returns a map of actors grouped by their relevant class.
 	 * (Relevant class being the base class that supports a given entity)
 	 */
-	TMap<UClass*, TArray<UObject*>> GetLevelObjectsGroupedByClass(const TArray<UClass*>& RelevantClasses)
+	TMap<UClass*, TArray<UObject*>> GetLevelObjectsGroupedByClass(const TArray<UClass*>& RelevantClasses, UWorld* PresetWorld)
 	{
 		TMap<UClass*, TArray<UObject*>> ObjectMap;
 
-		auto AddRelevantClassesForObject = [&ObjectMap, &RelevantClasses] (UObject* Object)
+		auto AddRelevantClassesForObject = [&ObjectMap, &RelevantClasses, PresetWorld] (UObject* Object)
 		{
 			if (!Object
-				|| (Object->IsA<AActor>() && !IsValidActorForRebinding(CastChecked<AActor>(Object)))
-				|| (Object->IsA<UActorComponent>() && !IsValidComponentForRebinding(CastChecked<UActorComponent>(Object))))
+				|| (Object->IsA<AActor>() && !IsValidActorForRebinding(CastChecked<AActor>(Object), PresetWorld))
+				|| (Object->IsA<UActorComponent>() && !IsValidComponentForRebinding(CastChecked<UActorComponent>(Object), PresetWorld)))
 			{
 				return;
 			}
@@ -117,9 +90,9 @@ namespace RCPresetRebindingManager
 			}
 		};
 
-		if (UWorld* World = GetCurrentWorld())
+		if (PresetWorld)
 		{
-			for (TActorIterator<AActor> It(World, AActor::StaticClass(), EActorIteratorFlags::SkipPendingKill); It; ++It)
+			for (TActorIterator<AActor> It(PresetWorld, AActor::StaticClass(), EActorIteratorFlags::SkipPendingKill); It; ++It)
 			{
         		AddRelevantClassesForObject(*It);
 			}
@@ -134,18 +107,18 @@ namespace RCPresetRebindingManager
 		return ObjectMap;
 	}
 
-	bool GetActorsOfClass(UClass* InTargetClass, TArray<AActor*>& OutActors)
+	bool GetActorsOfClass(UWorld* PresetWorld, UClass* InTargetClass, TArray<AActor*>& OutActors)
 	{
-		auto IsObjectOfClass = [InTargetClass](UObject* Object)
+		auto IsObjectOfClass = [InTargetClass, PresetWorld](UObject* Object)
 		{
 			return Object
-				&& IsValidActorForRebinding(CastChecked<AActor>(Object))
+				&& IsValidActorForRebinding(CastChecked<AActor>(Object), PresetWorld)
 				&& Object->IsA(InTargetClass);
 		};
 
-		if (UWorld* World = GetCurrentWorld())
+		if (PresetWorld)
 		{
-			for (TActorIterator<AActor> It(World, AActor::StaticClass(), EActorIteratorFlags::SkipPendingKill); It; ++It)
+			for (TActorIterator<AActor> It(PresetWorld, AActor::StaticClass(), EActorIteratorFlags::SkipPendingKill); It; ++It)
 			{
 				if (IsObjectOfClass(*It))
 				{
@@ -262,7 +235,7 @@ void FRemoteControlPresetRebindingManager::Rebind_Legacy(URemoteControlPreset* P
 	// Fetch any relevant object from the level based on if their class is relevant for the set of entities we have to rebind.
 	TArray<UClass*> EntitiesOwnerClasses;
 	EntitiesGroupedBySupportedOwnerClass.GenerateKeyArray(EntitiesOwnerClasses);
-	Context.ObjectsGroupedByRelevantClass = RCPresetRebindingManager::GetLevelObjectsGroupedByClass(EntitiesOwnerClasses);
+	Context.ObjectsGroupedByRelevantClass = RCPresetRebindingManager::GetLevelObjectsGroupedByClass(EntitiesOwnerClasses, Preset->GetPresetWorld(false));
 
 	for (TPair<UClass*, TArray<TSharedPtr<FRemoteControlEntity>>>& Entry : EntitiesGroupedBySupportedOwnerClass)
 	{
@@ -318,7 +291,7 @@ void FRemoteControlPresetRebindingManager::Rebind_NewAlgo(URemoteControlPreset*
 	}
 
 	// First, try to 'rebind for all properties' bindings that are already bound.
-	// to make sure that we don't accidently rebind a property to a different actor
+	// to make sure that we don't accidentally rebind a property to a different actor
 	for (URemoteControlLevelDependantBinding* ValidBinding : ValidBindings)
 	{
 		UObject* ResolvedObject = ValidBinding->Resolve();
@@ -459,7 +432,7 @@ void FRemoteControlPresetRebindingManager::Rebind_NewAlgo(URemoteControlPreset*
 		{
 			const bool bRebindingComponent = !InvalidBinding->BindingContext.ComponentName.IsNone();
 			TArray<AActor*> ObjectsWithSupportedClass;
-			if (!RCPresetRebindingManager::GetActorsOfClass(OwnerClass, ObjectsWithSupportedClass))
+			if (!RCPresetRebindingManager::GetActorsOfClass(Preset->GetPresetWorld(false), OwnerClass, ObjectsWithSupportedClass))
 			{
 				continue;
 			}
diff --git a/Engine/Plugins/VirtualProduction/RemoteControl/Source/RemoteControl/Public/IRemoteControlModule.h b/Engine/Plugins/VirtualProduction/RemoteControl/Source/RemoteControl/Public/IRemoteControlModule.h
--- a/Engine/Plugins/VirtualProduction/RemoteControl/Source/RemoteControl/Public/IRemoteControlModule.h
+++ b/Engine/Plugins/VirtualProduction/RemoteControl/Source/RemoteControl/Public/IRemoteControlModule.h
@@ -251,6 +251,19 @@ public:
 	UE_DEPRECATED(4.27, "UnregisterPreset is deprecated.")
 	virtual void UnregisterPreset(FName Name) = 0;
 
+	/**
+	 * Registers a hosted (non-asset based) preset with an attached name.
+	 * Embedded presets should only be registered while they are active.
+	 * @return whether registration was successful.
+	 */
+	virtual bool RegisterEmbeddedPreset(URemoteControlPreset* Preset, bool bReplaceExisting = false) = 0;
+
+	/**
+	 * Unregisters a hosted (non-asset based) preset with the given name.
+	 */
+	virtual void UnregisterEmbeddedPreset(FName Name) = 0;
+	virtual void UnregisterEmbeddedPreset(URemoteControlPreset* Preset) = 0;
+
 	/**
 	 * Resolve a RemoteCall Object and Function.
 	 * This will look for object and function to resolve. 
@@ -423,7 +436,12 @@ public:
 	 * @arg bIncludeTransient Whether to include transient presets.
 	 */
 	virtual void GetPresetAssets(TArray<FAssetData>& OutPresetAssets, bool bIncludeTransient = true) const = 0;
-	
+
+	/**
+	 * Gets all hosted presets currently registered with the module.
+	 */
+	virtual void GetEmbeddedPresets(TArray<TWeakObjectPtr<URemoteControlPreset>>& OutEmbeddedPresets) const = 0;
+
 	/**
 	 * Get the map of registered default entity metadata initializers. 
 	 */
diff --git a/Engine/Plugins/VirtualProduction/RemoteControl/Source/RemoteControl/Public/RemoteControlBinding.h b/Engine/Plugins/VirtualProduction/RemoteControl/Source/RemoteControl/Public/RemoteControlBinding.h
--- a/Engine/Plugins/VirtualProduction/RemoteControl/Source/RemoteControl/Public/RemoteControlBinding.h
+++ b/Engine/Plugins/VirtualProduction/RemoteControl/Source/RemoteControl/Public/RemoteControlBinding.h
@@ -172,8 +172,9 @@ public:
 	 */
 	UClass* GetSupportedOwnerClass() const;
 
-	/** Get the current world. */
-	static UWorld* GetCurrentWorld();
+	/** Get the current world associated with the preset or the editor world as a fallback. */
+	UWorld* GetCurrentWorld() const;
+
 private:
 	TSoftObjectPtr<UObject> ResolveForCurrentWorld() const;
 
diff --git a/Engine/Plugins/VirtualProduction/RemoteControl/Source/RemoteControl/Public/RemoteControlPreset.h b/Engine/Plugins/VirtualProduction/RemoteControl/Source/RemoteControl/Public/RemoteControlPreset.h
--- a/Engine/Plugins/VirtualProduction/RemoteControl/Source/RemoteControl/Public/RemoteControlPreset.h
+++ b/Engine/Plugins/VirtualProduction/RemoteControl/Source/RemoteControl/Public/RemoteControlPreset.h
@@ -331,6 +331,11 @@ public:
 	 */
 	const FGuid& GetPresetId() const { return PresetId; }
 
+	/**
+	 * Returns the FName that represents this asset or hosted preset.
+	 */
+	FName GetPresetName() const;
+
 	/**
 	 * Expose an actor on this preset.
 	 * @param Actor the actor to expose.
@@ -833,6 +838,26 @@ private:
 	/** Frame counter for delaying property change checks. */
 	int32 PropertyChangeWatchFrameCounter = 0;
 
+public:
+
+	static UWorld* GetPresetWorld(const URemoteControlPreset* Preset = nullptr, bool bAllowPIE = false);
+	UWorld* GetPresetWorld(bool bAllowPIE = false) const;
+	
+	/** Returns true if the preset is hosted within another asset. */
+	bool IsEmbeddedPreset() const;
+
+#if WITH_EDITOR
+	const TArray<FName>& GetDetailsTabIdentifierOverrides() const { return DetailsTabIdentifierOverrides; }
+	void SetDetailsTabIdentifierOverrides(const TArray<FName>& NewOverrides) { DetailsTabIdentifierOverrides = NewOverrides; }
+#endif
+
+private:
+
+#if WITH_EDITORONLY_DATA
+	UPROPERTY()
+	TArray<FName> DetailsTabIdentifierOverrides;
+#endif
+
 #if WITH_EDITOR
 	/** List of blueprints for which we have registered events. */
 	TSet<TWeakObjectPtr<UBlueprint>> BlueprintsWithRegisteredDelegates;
diff --git a/Engine/Plugins/VirtualProduction/RemoteControl/Source/RemoteControlUI/Private/AssetEditor/RemoteControlPresetEditorToolkit.cpp b/Engine/Plugins/VirtualProduction/RemoteControl/Source/RemoteControlUI/Private/AssetEditor/RemoteControlPresetEditorToolkit.cpp
--- a/Engine/Plugins/VirtualProduction/RemoteControl/Source/RemoteControlUI/Private/AssetEditor/RemoteControlPresetEditorToolkit.cpp
+++ b/Engine/Plugins/VirtualProduction/RemoteControl/Source/RemoteControlUI/Private/AssetEditor/RemoteControlPresetEditorToolkit.cpp
@@ -33,11 +33,21 @@ void FRemoteControlPresetEditorToolkit::InitRemoteControlPresetEditor(const EToo
 	// Required, will cause the previous toolkit to close bringing down the RemoteControlPreset and unsubscribing the
 	// tab spawner. Without this, the InitAssetEditor call below will trigger an ensure as the RemoteControlPreset
 	// tab ID will already be registered within EditorTabManager
-	FLevelEditorModule& LevelEditorModule = FModuleManager::GetModuleChecked<FLevelEditorModule>("LevelEditor");
-	TSharedPtr<FTabManager> EditorTabManager = LevelEditorModule.GetLevelEditorTabManager();
-	if (EditorTabManager->FindExistingLiveTab(FRemoteControlUIModule::RemoteControlPanelTabName).IsValid())
+	TSharedPtr<FTabManager> HostTabManager;
+
+	if (ToolkitHost.IsValid())
+	{
+		HostTabManager = ToolkitHost.Pin()->GetTabManager();
+	}
+	else
+	{
+		FLevelEditorModule& LevelEditorModule = FModuleManager::GetModuleChecked<FLevelEditorModule>("LevelEditor");
+		HostTabManager = LevelEditorModule.GetLevelEditorTabManager();
+	}
+
+	if (HostTabManager.IsValid() && HostTabManager->FindExistingLiveTab(FRemoteControlUIModule::RemoteControlPanelTabName).IsValid())
 	{
-		EditorTabManager->TryInvokeTab(FRemoteControlUIModule::RemoteControlPanelTabName)->RequestCloseTab();
+		HostTabManager->TryInvokeTab(FRemoteControlUIModule::RemoteControlPanelTabName)->RequestCloseTab();
 	}
 
 	RegisterTabSpawners(InitToolkitHost->GetTabManager().ToSharedRef());
@@ -78,17 +88,21 @@ FRemoteControlPresetEditorToolkit::~FRemoteControlPresetEditorToolkit()
 
 void FRemoteControlPresetEditorToolkit::RegisterTabSpawners(const TSharedRef<FTabManager>& InTabManager)
 {
-	WorkspaceMenuCategory = InTabManager->AddLocalWorkspaceMenuCategory(LOCTEXT("WorkspaceMenu_RemoteControlPanel", "Remote Control Panel"));
+	// In case of double registration
+	if (!InTabManager->HasTabSpawner(FRemoteControlUIModule::RemoteControlPanelTabName))
+	{
+		WorkspaceMenuCategory = InTabManager->AddLocalWorkspaceMenuCategory(LOCTEXT("WorkspaceMenu_RemoteControlPanel", "Remote Control Panel"));
 
-	InTabManager->RegisterTabSpawner(FRemoteControlUIModule::RemoteControlPanelTabName, FOnSpawnTab::CreateSP(this, &FRemoteControlPresetEditorToolkit::HandleTabManagerSpawnPanelTab))
-		.SetDisplayName(LOCTEXT("RemoteControlPanelMainTab", "Remote Control Panel"))
-		.SetGroup(WorkspaceMenuCategory.ToSharedRef())
-		.SetIcon(FSlateIcon(FAppStyle::GetAppStyleSetName(), "LevelEditor.Tabs.Details"));
+		InTabManager->RegisterTabSpawner(FRemoteControlUIModule::RemoteControlPanelTabName, FOnSpawnTab::CreateSP(this, &FRemoteControlPresetEditorToolkit::HandleTabManagerSpawnPanelTab))
+			.SetDisplayName(LOCTEXT("RemoteControlPanelMainTab", "Remote Control Panel"))
+			.SetGroup(WorkspaceMenuCategory.ToSharedRef())
+			.SetIcon(FSlateIcon(FAppStyle::GetAppStyleSetName(), "LevelEditor.Tabs.Details"));
 
-	InTabManager->RegisterTabSpawner(FRemoteControlUIModule::EntityDetailsTabName, FOnSpawnTab::CreateSP(this, &FRemoteControlPresetEditorToolkit::HandleTabManagerSpawnDetailsTab))
-		.SetDisplayName(LOCTEXT("RemoteControlPanelDetailsTab", "Entity Details"))
-		.SetGroup(WorkspaceMenuCategory.ToSharedRef())
-		.SetIcon(FSlateIcon(FAppStyle::GetAppStyleSetName(), "LevelEditor.Tabs.Details"));
+		InTabManager->RegisterTabSpawner(FRemoteControlUIModule::EntityDetailsTabName, FOnSpawnTab::CreateSP(this, &FRemoteControlPresetEditorToolkit::HandleTabManagerSpawnDetailsTab))
+			.SetDisplayName(LOCTEXT("RemoteControlPanelDetailsTab", "Entity Details"))
+			.SetGroup(WorkspaceMenuCategory.ToSharedRef())
+			.SetIcon(FSlateIcon(FAppStyle::GetAppStyleSetName(), "LevelEditor.Tabs.Details"));
+	}
 }
 
 void FRemoteControlPresetEditorToolkit::UnregisterTabSpawners(const TSharedRef<FTabManager>& InTabManager)
@@ -179,24 +193,34 @@ void FRemoteControlPresetEditorToolkit::InvokePanelTab()
 		}
 	};
 
-	// Create a new DockTab and add the RemoteControlPreset widget to it.
-	FLevelEditorModule& LevelEditorModule = FModuleManager::GetModuleChecked<FLevelEditorModule>("LevelEditor");
-	TSharedPtr<FTabManager> EditorTabManager = LevelEditorModule.GetLevelEditorTabManager();
+	TSharedPtr<FTabManager> HostTabManager;	
 
-	if (TSharedPtr<SDockTab> Tab = EditorTabManager->TryInvokeTab(FRemoteControlUIModule::RemoteControlPanelTabName))
+	if (ToolkitHost.IsValid())
 	{
-		Tab->SetContent(PanelTab.ToSharedRef());
-		Tab->SetOnTabClosed(SDockTab::FOnTabClosedCallback::CreateStatic(&Local::OnRemoteControlPresetClosed, TWeakPtr<IAssetEditorInstance>(SharedThis(this))));
-	
-		// Force the parent window to have minimum width.
-		TSharedPtr<SWindow> Window = Tab->GetParentWindow();
-		if (Window.IsValid())
-		{
-			FWindowSizeLimits SizeLimits = Window->GetSizeLimits();
+		HostTabManager = ToolkitHost.Pin()->GetTabManager();
+	}
+	else
+	{
+		FLevelEditorModule& LevelEditorModule = FModuleManager::GetModuleChecked<FLevelEditorModule>("LevelEditor");
+		HostTabManager = LevelEditorModule.GetLevelEditorTabManager();
+	}
 
-			SizeLimits.SetMinWidth(SRemoteControlPanel::MinimumPanelWidth);
+	// Create a new DockTab and add the RemoteControlPreset widget to it.
+	if (HostTabManager.IsValid())
+	{
+		if (TSharedPtr<SDockTab> Tab = HostTabManager->TryInvokeTab(FRemoteControlUIModule::RemoteControlPanelTabName))
+		{
+			Tab->SetContent(PanelTab.ToSharedRef());
+			Tab->SetOnTabClosed(SDockTab::FOnTabClosedCallback::CreateStatic(&Local::OnRemoteControlPresetClosed, TWeakPtr<IAssetEditorInstance>(SharedThis(this))));
 
-			Window->SetSizeLimits(SizeLimits);
+			// Force the parent window to have minimum width.
+			TSharedPtr<SWindow> Window = Tab->GetParentWindow();
+			if (Window.IsValid())
+			{
+				FWindowSizeLimits SizeLimits = Window->GetSizeLimits();
+				SizeLimits.SetMinWidth(SRemoteControlPanel::MinimumPanelWidth);
+				Window->SetSizeLimits(SizeLimits);
+			}
 		}
 	}
 }
diff --git a/Engine/Plugins/VirtualProduction/RemoteControl/Source/RemoteControlUI/Private/RemoteControlUIModule.cpp b/Engine/Plugins/VirtualProduction/RemoteControl/Source/RemoteControlUI/Private/RemoteControlUIModule.cpp
--- a/Engine/Plugins/VirtualProduction/RemoteControl/Source/RemoteControlUI/Private/RemoteControlUIModule.cpp
+++ b/Engine/Plugins/VirtualProduction/RemoteControl/Source/RemoteControlUI/Private/RemoteControlUIModule.cpp
@@ -39,13 +39,15 @@
 const FName FRemoteControlUIModule::EntityDetailsTabName = "RemoteControl_EntityDetails";
 const FName FRemoteControlUIModule::RemoteControlPanelTabName = "RemoteControl_RemoteControlPanel";
 
-static const FName DetailsTabIdentifiers[] = {
+static const FName DetailsTabIdentifiers_LevelEditor[] = {
 	"LevelEditorSelectionDetails",
 	"LevelEditorSelectionDetails2",
 	"LevelEditorSelectionDetails3",
 	"LevelEditorSelectionDetails4"
 };
 
+static const int32 DetailsTabIdentifiers_LevelEditor_Count = 4;
+
 FRCExposesPropertyArgs::FRCExposesPropertyArgs()
 	: PropertyHandle(nullptr)
 	, OwnerObject(nullptr)
@@ -212,7 +214,32 @@ namespace RemoteControlUIModule
 
 	bool IsTransientObjectAllowListed(UObject* Object)
 	{
-		return Object && Object->IsA<UDEditorParameterValue>();
+		if (!Object)
+		{
+			return false;
+		}
+
+		if (Object->IsA<UDEditorParameterValue>())
+		{
+			return true;
+		}
+
+		/*
+		 * Embedded presets may be part of a transient package.
+		 * If the object's world's outer is transient, it's probably hosted. Allow it.
+		 */
+		if (UWorld* ObjectWorld = Object->GetWorld())
+		{
+			if (UObject* ObjectWorldOuter = ObjectWorld->GetOuter())
+			{
+				if (ObjectWorldOuter->HasAnyFlags(RF_Transient))
+				{
+					return true;
+				}
+			}
+		}
+
+		return false;
 	}
 }
 
@@ -293,6 +320,12 @@ TSharedRef<SRemoteControlPanel> FRemoteControlUIModule::CreateRemoteControlPanel
 	if (!SharedDetailsPanel.IsValid())
 	{
 		FPropertyEditorModule& PropertyEditor = FModuleManager::LoadModuleChecked<FPropertyEditorModule>("PropertyEditor");
+		TArray<FName> DetailsTabIdentifiers = Preset->GetDetailsTabIdentifierOverrides();
+
+		if (DetailsTabIdentifiers.Num() == 0)
+		{
+			DetailsTabIdentifiers.Append(DetailsTabIdentifiers_LevelEditor, DetailsTabIdentifiers_LevelEditor_Count);
+		}
 
 		for (const FName& DetailsTabIdentifier : DetailsTabIdentifiers)
 		{
diff --git a/Engine/Plugins/VirtualProduction/RemoteControl/Source/RemoteControlUI/Private/UI/SRCPanelExposedActor.cpp b/Engine/Plugins/VirtualProduction/RemoteControl/Source/RemoteControlUI/Private/UI/SRCPanelExposedActor.cpp
--- a/Engine/Plugins/VirtualProduction/RemoteControl/Source/RemoteControlUI/Private/UI/SRCPanelExposedActor.cpp
+++ b/Engine/Plugins/VirtualProduction/RemoteControl/Source/RemoteControlUI/Private/UI/SRCPanelExposedActor.cpp
@@ -6,8 +6,11 @@
 #include "PropertyCustomizationHelpers.h"
 #include "RemoteControlActor.h"
 #include "RemoteControlPreset.h"
+#include "SceneOutlinerModule.h"
+#include "SceneOutlinerPublicTypes.h"
 #include "ScopedTransaction.h"
 #include "SRCPanelDragHandle.h"
+#include "Widgets/Input/SButton.h"
 
 #define LOCTEXT_NAMESPACE "RemoteControlPanelActor"
 
@@ -65,6 +68,35 @@ void SRCPanelExposedActor::Refresh()
 
 TSharedRef<SWidget> SRCPanelExposedActor::RecreateWidget(const FString& Path)
 {
+	// Create combo button to open actor picker with correct custom world
+	if (Preset.IsValid() && Preset->IsEmbeddedPreset())
+	{
+		FText ActorName = LOCTEXT("MissingReference", "Unable to find reference");
+
+		if (TSharedPtr<FRemoteControlActor> RCActor = WeakActor.Pin())
+		{
+			if (RCActor->GetActor())
+			{
+				ActorName = FText::AsCultureInvariant(RCActor->GetActor()->GetActorLabel());
+			}
+		}
+
+		TSharedRef<STextBlock> ActorNameText = SNew(STextBlock)
+			.Text(ActorName)
+			.Font(FCoreStyle::GetDefaultFontStyle("Regular", 8));
+
+		TSharedRef<SComboButton> ActorButton = SNew(SComboButton)
+			.ButtonStyle(FAppStyle::Get(), TEXT("SimpleButton"))
+			.ButtonContent()
+			[
+				ActorNameText
+			]
+			.OnGetMenuContent(this, &SRCPanelExposedActor::CreateEmbeddedPresetActorPicker);
+
+		return CreateEntityWidget(MoveTemp(ActorButton));
+	}
+
+	// Otherwise use SObjectPropertyEntryBox which does not support non-level editor worlds.
 	TSharedRef<SObjectPropertyEntryBox> EntryBox = SNew(SObjectPropertyEntryBox)
 		.ObjectPath(Path)
 		.AllowedClass(AActor::StaticClass())
@@ -74,7 +106,8 @@ TSharedRef<SWidget> SRCPanelExposedActor::RecreateWidget(const FString& Path)
 		.DisplayBrowse(true)
 		.NewAssetFactories(TArray<UFactory*>());
 
-	float MinWidth = 0.f, MaxWidth = 0.f;
+	float MinWidth = 0.f;
+	float MaxWidth = 0.f;
 	EntryBox->GetDesiredWidth(MinWidth, MaxWidth);
 
 	TSharedRef<SWidget> ValueWidget =
@@ -92,13 +125,34 @@ void SRCPanelExposedActor::OnChangeActor(const FAssetData& AssetData)
 {
 	if (AActor* Actor = Cast<AActor>(AssetData.GetAsset()))
 	{
-		if (TSharedPtr<FRemoteControlActor> RCActor = WeakActor.Pin())
-		{
-			FScopedTransaction Transaction(LOCTEXT("ChangeExposedActor", "Change Exposed Actor"));
-			RCActor->SetActor(Actor);
-			ChildSlot.AttachWidget(RecreateWidget(Actor->GetPathName()));
-		}
+		OnChangeActor(Actor);
 	}
 }
 
+void SRCPanelExposedActor::OnChangeActor(AActor* Actor)
+{
+	if (TSharedPtr<FRemoteControlActor> RCActor = WeakActor.Pin())
+	{
+		FScopedTransaction Transaction(LOCTEXT("ChangeExposedActor", "Change Exposed Actor"));
+		RCActor->SetActor(Actor);
+		ChildSlot.AttachWidget(RecreateWidget(Actor->GetPathName()));
+	}
+}
+
+TSharedRef<SWidget> SRCPanelExposedActor::CreateEmbeddedPresetActorPicker()
+{
+	FSceneOutlinerModule& SceneOutlinerModule = FModuleManager::LoadModuleChecked<FSceneOutlinerModule>("SceneOutliner");
+	FSceneOutlinerInitializationOptions InitOptions;
+	UWorld* PresetWorld = URemoteControlPreset::GetPresetWorld(Preset.Get());
+
+	TSharedRef<SWidget> ActorPicker =
+		SceneOutlinerModule.CreateActorPicker(
+			InitOptions,
+			FOnActorPicked::CreateSP(this, &SRCPanelExposedActor::OnChangeActor),
+			PresetWorld
+		);
+
+	return ActorPicker;
+}
+
 #undef LOCTEXT_NAMESPACE /*RemoteControlPanelActor*/
\ No newline at end of file
diff --git a/Engine/Plugins/VirtualProduction/RemoteControl/Source/RemoteControlUI/Private/UI/SRCPanelExposedActor.h b/Engine/Plugins/VirtualProduction/RemoteControl/Source/RemoteControlUI/Private/UI/SRCPanelExposedActor.h
--- a/Engine/Plugins/VirtualProduction/RemoteControl/Source/RemoteControlUI/Private/UI/SRCPanelExposedActor.h
+++ b/Engine/Plugins/VirtualProduction/RemoteControl/Source/RemoteControlUI/Private/UI/SRCPanelExposedActor.h
@@ -41,6 +41,10 @@ private:
 	TSharedRef<SWidget> RecreateWidget(const FString& Path);
 	/** Handle the user selecting a different actor to expose. */
 	void OnChangeActor(const FAssetData& AssetData);
+	/** Handle direct actor value event instead of by asset. */
+	void OnChangeActor(AActor* Actor);
+	/** Creates an actor picker for hosted presets because the engine default one cannot use custom worlds. */
+	TSharedRef<SWidget> CreateEmbeddedPresetActorPicker();
 private:
 	/** Weak reference to the preset that exposes the actor. */
 	TWeakObjectPtr<URemoteControlPreset> WeakPreset;
diff --git a/Engine/Plugins/VirtualProduction/RemoteControl/Source/RemoteControlUI/Private/UI/SRCPanelExposedEntity.cpp b/Engine/Plugins/VirtualProduction/RemoteControl/Source/RemoteControlUI/Private/UI/SRCPanelExposedEntity.cpp
--- a/Engine/Plugins/VirtualProduction/RemoteControl/Source/RemoteControlUI/Private/UI/SRCPanelExposedEntity.cpp
+++ b/Engine/Plugins/VirtualProduction/RemoteControl/Source/RemoteControlUI/Private/UI/SRCPanelExposedEntity.cpp
@@ -176,7 +176,7 @@ TSharedRef<SWidget> SRCPanelExposedEntity::CreateRebindAllPropertiesForActorMenu
 		.MaxDesiredHeight(400.0f)
 		.WidthOverride(300.0f)
 		[
-			SceneOutlinerModule.CreateActorPicker(Options, FOnActorPicked::CreateRaw(this, &SRCPanelExposedEntity::OnActorSelectedForRebindAllProperties))
+			SceneOutlinerModule.CreateActorPicker(Options, FOnActorPicked::CreateSP(this, &SRCPanelExposedEntity::OnActorSelectedForRebindAllProperties), URemoteControlPreset::GetPresetWorld(Preset.Get()))
 		];
 }
 
@@ -214,15 +214,15 @@ TSharedRef<SWidget> SRCPanelExposedEntity::CreateRebindMenuContent()
 	FSceneOutlinerModule& SceneOutlinerModule = FModuleManager::Get().LoadModuleChecked<FSceneOutlinerModule>("SceneOutliner");
 	FSceneOutlinerInitializationOptions Options;
 	Options.Filters = MakeShared<FSceneOutlinerFilters>();
-
 	Options.Filters->AddFilterPredicate<FActorTreeItem>(FActorTreeItem::FFilterPredicate::CreateRaw(this, &SRCPanelExposedEntity::IsActorSelectable));
+	UWorld* PresetWorld = URemoteControlPreset::GetPresetWorld(Preset.Get());
 
 	return SNew(SBox)
-	.MaxDesiredHeight(400.0f)
-	.WidthOverride(300.0f)
-	[
-		SceneOutlinerModule.CreateActorPicker(Options, FOnActorPicked::CreateRaw(this, &SRCPanelExposedEntity::OnActorSelected), nullptr)
-	];
+		.MaxDesiredHeight(400.0f)
+		.WidthOverride(300.0f)
+		[
+			SceneOutlinerModule.CreateActorPicker(Options, FOnActorPicked::CreateSP(this, &SRCPanelExposedEntity::OnActorSelected), PresetWorld)
+		];
 }
 
 bool SRCPanelExposedEntity::OnVerifyItemLabelChanged(const FText& InLabel, FText& OutErrorMessage)
diff --git a/Engine/Plugins/VirtualProduction/RemoteControl/Source/RemoteControlUI/Private/UI/SRCPanelFunctionPicker.cpp b/Engine/Plugins/VirtualProduction/RemoteControl/Source/RemoteControlUI/Private/UI/SRCPanelFunctionPicker.cpp
--- a/Engine/Plugins/VirtualProduction/RemoteControl/Source/RemoteControlUI/Private/UI/SRCPanelFunctionPicker.cpp
+++ b/Engine/Plugins/VirtualProduction/RemoteControl/Source/RemoteControlUI/Private/UI/SRCPanelFunctionPicker.cpp
@@ -8,9 +8,12 @@
 #include "Engine/World.h"
 #include "Engine/GameEngine.h"
 #include "Framework/Application/SlateApplication.h"
+#include "Toolkits/IToolkitHost.h"
 #include "Kismet/GameplayStatics.h"
 #include "GameFramework/Actor.h"
 #include "Styling/SlateColor.h"
+#include "RemoteControlPreset.h"
+#include "SRemoteControlPanel.h"
 #include "Subsystems/Subsystem.h"
 #include "UObject/ObjectMacros.h"
 #include "UObject/UObjectIterator.h"
@@ -169,32 +172,13 @@ namespace FunctionPickerUtils
 		/** This object's name. */
 		FString Name;
 	};
-
-	UWorld* GetWorld(bool bIgnorePIE = false)
-	{
-		// Get the current world.
-		if (GIsEditor)
-		{
-			FWorldContext* PIEWorldContext = GEditor->GetPIEWorldContext();
-			if (!bIgnorePIE && PIEWorldContext)
-			{
-				return PIEWorldContext->World();
-			}
-			return GEditor->GetEditorWorldContext().World();
-		}
-		else if (UGameEngine* GameEngine = Cast<UGameEngine>(GEngine))
-		{
-			return GameEngine->GetGameWorld();
-		}
-		return nullptr;
-	}
-
 }
 
 void SRCPanelFunctionPicker::Construct(const FArguments& InArgs)
 {
 	using namespace FunctionPickerUtils;
 
+	RemoteControlPanel = InArgs._RemoteControlPanel;
 	ObjectClass = InArgs._ObjectClass;
 	Label = InArgs._Label;
 	bAllowDefaultObjects = InArgs._AllowDefaultObjects;
@@ -299,7 +283,7 @@ void SRCPanelFunctionPicker::Refresh()
 	if (Class->IsChildOf(AActor::StaticClass()) && GEditor)
 	{
 		TArray<AActor*> ActorList;
-		UGameplayStatics::GetAllActorsOfClass(GetWorld(), Class, ActorList);
+		UGameplayStatics::GetAllActorsOfClass(GetPresetWorld(), Class, ActorList);
 
 		for (AActor* Actor : ActorList)
 		{
@@ -320,4 +304,16 @@ void SRCPanelFunctionPicker::Refresh()
 	}
 }
 
+UWorld* SRCPanelFunctionPicker::GetPresetWorld(bool bIgnorePIE) const
+{
+	TSharedPtr<SRemoteControlPanel> RCP = RemoteControlPanel.Pin();
+
+	if (RCP.IsValid())
+	{
+		return URemoteControlPreset::GetPresetWorld(RCP->GetPreset(), !bIgnorePIE);
+	}
+
+	return URemoteControlPreset::GetPresetWorld(nullptr, !bIgnorePIE);
+}
+
 #undef LOCTEXT_NAMESPACE /* RemoteControlPanel */
\ No newline at end of file
diff --git a/Engine/Plugins/VirtualProduction/RemoteControl/Source/RemoteControlUI/Private/UI/SRCPanelFunctionPicker.h b/Engine/Plugins/VirtualProduction/RemoteControl/Source/RemoteControlUI/Private/UI/SRCPanelFunctionPicker.h
--- a/Engine/Plugins/VirtualProduction/RemoteControl/Source/RemoteControlUI/Private/UI/SRCPanelFunctionPicker.h
+++ b/Engine/Plugins/VirtualProduction/RemoteControl/Source/RemoteControlUI/Private/UI/SRCPanelFunctionPicker.h
@@ -7,6 +7,7 @@
 DECLARE_DELEGATE_TwoParams(FOnSelectFunction, UObject*, UFunction*);
 
 struct FRCFunctionPickerTreeNode;
+class SRemoteControlPanel;
 
 /**
  * Widget that displays a picker for blueprint functions.
@@ -18,6 +19,7 @@ public:
 		: _AllowDefaultObjects(false)
 	{}
 		SLATE_EVENT(FOnSelectFunction, OnSelectFunction)
+		SLATE_ARGUMENT(TSharedPtr<SRemoteControlPanel>, RemoteControlPanel)
 		SLATE_ARGUMENT(UClass*, ObjectClass)
 		SLATE_ARGUMENT(FText, Label)
 		SLATE_ARGUMENT(bool, AllowDefaultObjects)
@@ -53,4 +55,8 @@ private:
 	bool bAllowDefaultObjects = false;
 	/** Keep track of the last time this widget tick in order to focus it if it's been more than a frame. */
 	double LastTimeSinceTick = 0.0;
+	/** Keep a weak reference to the RCPanel for access to toolkit host. */
+	TWeakPtr<SRemoteControlPanel> RemoteControlPanel;
+	/** Checks the remote control panel to get a reference to the preset world. */
+	UWorld* GetPresetWorld(bool bIgnorePIE = false) const;
 };
\ No newline at end of file
diff --git a/Engine/Plugins/VirtualProduction/RemoteControl/Source/RemoteControlUI/Private/UI/SRemoteControlPanel.cpp b/Engine/Plugins/VirtualProduction/RemoteControl/Source/RemoteControlUI/Private/UI/SRemoteControlPanel.cpp
--- a/Engine/Plugins/VirtualProduction/RemoteControl/Source/RemoteControlUI/Private/UI/SRemoteControlPanel.cpp
+++ b/Engine/Plugins/VirtualProduction/RemoteControl/Source/RemoteControlUI/Private/UI/SRemoteControlPanel.cpp
@@ -39,6 +39,8 @@
 #include "RemoteControlPanelStyle.h"
 #include "RemoteControlPreset.h"
 #include "RemoteControlSettings.h"
+#include "SceneOutlinerModule.h"
+#include "SceneOutlinerPublicTypes.h"
 #include "ScopedTransaction.h"
 #include "SClassViewer.h"
 #include "SRCLogger.h"
@@ -90,11 +92,6 @@ namespace RemoteControlPanelUtils
             && !Actor->IsA(AWorldSettings::StaticClass());	// Don't add the WorldSettings actor, even though it is technically editable
 	};
 
-	UWorld* GetEditorWorld()
-	{
-		return GEditor ? GEditor->GetEditorWorldContext(false).World() : nullptr;
-	}
-
 	TSharedPtr<FStructOnScope> GetEntityOnScope(const TSharedPtr<FRemoteControlEntity>& Entity, const UScriptStruct* EntityType)
 	{
 		if (ensure(Entity && EntityType))
@@ -438,7 +435,9 @@ void SRemoteControlPanel::Construct(const FArguments& InArgs, URemoteControlPres
 			.AutoWidth()
 			[
 				SNew(SButton)
-				.Visibility_Lambda([this]() { return bIsInEditMode ? EVisibility::Visible : EVisibility::Collapsed; })
+				.Visibility_Lambda([this]() { 
+					return (bIsInEditMode && Preset.IsValid() && !Preset->IsEmbeddedPreset()) ? EVisibility::Visible : EVisibility::Collapsed; 
+				})
 				.ButtonStyle(FAppStyle::Get(), "FlatButton")
 				.OnClicked(this, &SRemoteControlPanel::OnSavePreset)
 				[
@@ -567,7 +566,7 @@ void SRemoteControlPanel::Construct(const FArguments& InArgs, URemoteControlPres
 				[
 					SAssignNew(PresetNameTextBlock, STextBlock)
 					.Font(FAppStyle::Get().GetFontStyle("DetailsView.CategoryFontStyle"))
-					.Text(FText::FromName(Preset->GetFName()))
+					.Text(FText::FromName(Preset->GetPresetName()))
 				]
 				+ SHorizontalBox::Slot()
 				.AutoWidth()
@@ -970,17 +969,20 @@ TSharedRef<SWidget> SRemoteControlPanel::CreateExposeButton()
 	FMenuBuilder MenuBuilder(true, nullptr);
 	
 	SAssignNew(BlueprintPicker, SRCPanelFunctionPicker)
+		.RemoteControlPanel(SharedThis(this))
 		.AllowDefaultObjects(true)
 		.Label(LOCTEXT("FunctionLibrariesLabel", "Function Libraries"))
 		.ObjectClass(UBlueprintFunctionLibrary::StaticClass())
 		.OnSelectFunction_Raw(this, &SRemoteControlPanel::ExposeFunction);
 
 	SAssignNew(SubsystemFunctionPicker, SRCPanelFunctionPicker)
+		.RemoteControlPanel(SharedThis(this))
 		.Label(LOCTEXT("SubsystemFunctionLabel", "Subsystem Functions"))
 		.ObjectClass(USubsystem::StaticClass())
 		.OnSelectFunction_Raw(this, &SRemoteControlPanel::ExposeFunction);
 
 	SAssignNew(ActorFunctionPicker, SRCPanelFunctionPicker)
+		.RemoteControlPanel(SharedThis(this))
 		.Label(LOCTEXT("ActorFunctionsLabel", "Actor Functions"))
 		.ObjectClass(AActor::StaticClass())
 		.OnSelectFunction_Raw(this, &SRemoteControlPanel::ExposeFunction);
@@ -1022,16 +1024,44 @@ TSharedRef<SWidget> SRemoteControlPanel::CreateExposeButton()
 			LOCTEXT("ActorFunctionSubMenuToolTip", "Expose an actor's function."),
 			ActorFunctionPicker.ToSharedRef()
 		);
-		
-		MenuBuilder.AddWidget(
-			SNew(SObjectPropertyEntryBox)
+
+		// SObjectPropertyEntryBox does not support non-level editor worlds.
+		if (!Preset.IsValid() || !Preset->IsEmbeddedPreset())
+		{
+			MenuBuilder.AddWidget(
+				SNew(SObjectPropertyEntryBox)
 				.AllowedClass(AActor::StaticClass())
 				.OnObjectChanged(this, &SRemoteControlPanel::OnExposeActor)
 				.AllowClear(false)
 				.DisplayUseSelected(true)
 				.DisplayBrowse(true)
 				.NewAssetFactories(TArray<UFactory*>()),
-			LOCTEXT("ActorEntry", "Actor"));
+				LOCTEXT("ActorEntry", "Actor"));
+		}
+		else
+		{
+			MenuBuilder.AddSubMenu(
+				LOCTEXT("SelectActor", "Select Actor"),
+				LOCTEXT("SelectActorTooltip", "Select an actor to Remote Control."),
+				FNewMenuDelegate::CreateLambda(
+					[this](FMenuBuilder& SubMenuBuilder)
+					{
+						FSceneOutlinerModule& SceneOutlinerModule = FModuleManager::LoadModuleChecked<FSceneOutlinerModule>("SceneOutliner");
+						FSceneOutlinerInitializationOptions InitOptions;
+						UWorld* PresetWorld = URemoteControlPreset::GetPresetWorld(Preset.Get());
+
+						SubMenuBuilder.AddWidget(
+							SceneOutlinerModule.CreateActorPicker(
+								InitOptions,
+								FOnActorPicked::CreateSP(this, &SRemoteControlPanel::ExposeActor),
+								PresetWorld
+							),
+							FText::GetEmpty(), true, false
+						);
+					}
+				)
+			);
+		}
 
 		CreatePickerSubMenu(
 			LOCTEXT("ClassPickerEntry", "Actors By Class"),
@@ -1099,7 +1129,7 @@ TSharedRef<SWidget> SRemoteControlPanel::CreateExposeByClassWidget()
 	TSharedRef<SWidget> Widget = FModuleManager::LoadModuleChecked<FClassViewerModule>("ClassViewer").CreateClassViewer(Options, FOnClassPicked::CreateLambda(
 		[this](UClass* ChosenClass)
 		{
-			if (UWorld* World = RemoteControlPanelUtils::GetEditorWorld())
+			if (UWorld* World = URemoteControlPreset::GetPresetWorld(Preset.Get()))
 			{
 				for (TActorIterator<AActor> It(World, ChosenClass, EActorIteratorFlags::SkipPendingKill); It; ++It)
 				{
@@ -1127,8 +1157,9 @@ TSharedRef<SWidget> SRemoteControlPanel::CreateExposeByClassWidget()
 
 void SRemoteControlPanel::CacheLevelClasses()
 {
-	CachedClassesInLevel.Empty();	
-	if (UWorld* World = RemoteControlPanelUtils::GetEditorWorld())
+	CachedClassesInLevel.Empty();
+
+	if (UWorld* World = URemoteControlPreset::GetPresetWorld(Preset.Get()))
 	{
 		for (TActorIterator<AActor> It(World, AActor::StaticClass(), EActorIteratorFlags::SkipPendingKill); It; ++It)
 		{
@@ -1599,12 +1630,9 @@ void SRemoteControlPanel::UpdateActorFunctionPicker()
 
 void SRemoteControlPanel::OnAssetRenamed(const FAssetData& Asset, const FString&)
 {
-	if (Asset.GetAsset() == Preset.Get())
+	if (Preset.IsValid() && PresetNameTextBlock)
 	{
-		if (PresetNameTextBlock)
-		{
-			PresetNameTextBlock->SetText(FText::FromName(Asset.AssetName));	
-		}
+		PresetNameTextBlock->SetText(FText::FromName(Preset->GetPresetName()));
 	}
 }
 
diff --git a/Engine/Plugins/VirtualProduction/RemoteControl/Source/WebRemoteControl/Private/RemoteControlModels.h b/Engine/Plugins/VirtualProduction/RemoteControl/Source/WebRemoteControl/Private/RemoteControlModels.h
--- a/Engine/Plugins/VirtualProduction/RemoteControl/Source/WebRemoteControl/Private/RemoteControlModels.h
+++ b/Engine/Plugins/VirtualProduction/RemoteControl/Source/WebRemoteControl/Private/RemoteControlModels.h
@@ -497,7 +497,7 @@ struct FRCPresetDescription
 	{
 		checkSlow(Preset)
 
-		Name = Preset->GetName();
+		Name = Preset->GetPresetName().ToString();
 		Path = Preset->GetPathName();
 		ID = Preset->GetPresetId().ToString();
 
@@ -558,7 +558,17 @@ struct FRCShortPresetDescription
 	{
 		if (InPreset)
 		{
-			Name = InPreset->GetFName();
+			Name = InPreset->GetPresetName();
+			ID = InPreset->GetPresetId().ToString();
+			Path = *InPreset->GetPathName();
+		}
+	}
+
+	FRCShortPresetDescription(const TWeakObjectPtr<URemoteControlPreset> InPreset)
+	{
+		if (InPreset.IsValid())
+		{
+			Name = InPreset->GetPresetName();
 			ID = InPreset->GetPresetId().ToString();
 			Path = *InPreset->GetPathName();
 		}
diff --git a/Engine/Plugins/VirtualProduction/RemoteControl/Source/WebRemoteControl/Private/RemoteControlResponse.h b/Engine/Plugins/VirtualProduction/RemoteControl/Source/WebRemoteControl/Private/RemoteControlResponse.h
--- a/Engine/Plugins/VirtualProduction/RemoteControl/Source/WebRemoteControl/Private/RemoteControlResponse.h
+++ b/Engine/Plugins/VirtualProduction/RemoteControl/Source/WebRemoteControl/Private/RemoteControlResponse.h
@@ -48,9 +48,10 @@
 	
 	FListPresetsResponse() = default;
 
-	FListPresetsResponse(const TArray<FAssetData>& InPresets)
+	FListPresetsResponse(const TArray<FAssetData>& InPresets, const TArray<TWeakObjectPtr<URemoteControlPreset>> InEmbeddedPresets)
 	{
 		Presets.Append(InPresets);
+		Presets.Append(InEmbeddedPresets);
 	}
 
 	/**
@@ -254,7 +255,7 @@
 	{
 		if (InPreset)
 		{
-			PresetName = InPreset->GetFName();
+			PresetName = InPreset->GetPresetName();
 			PresetId = InPreset->GetPresetId().ToString();
 			Metadata = InPreset->Metadata;
 		}
@@ -368,7 +369,7 @@
 		: Type(TEXT("PresetEntitiesModified"))
 	{
 		checkSlow(InPreset);
-		PresetName = InPreset->GetFName();
+		PresetName = InPreset->GetPresetName();
 		PresetId = InPreset->GetPresetId().ToString();
 		ModifiedEntities = FRCPresetModifiedEntitiesDescription{InPreset, InModifiedEntities};
 	}
diff --git a/Engine/Plugins/VirtualProduction/RemoteControl/Source/WebRemoteControl/Private/WebRemoteControl.cpp b/Engine/Plugins/VirtualProduction/RemoteControl/Source/WebRemoteControl/Private/WebRemoteControl.cpp
--- a/Engine/Plugins/VirtualProduction/RemoteControl/Source/WebRemoteControl/Private/WebRemoteControl.cpp
+++ b/Engine/Plugins/VirtualProduction/RemoteControl/Source/WebRemoteControl/Private/WebRemoteControl.cpp
@@ -1593,7 +1593,10 @@ bool FWebRemoteControlModule::HandleGetPresetsRoute(const FHttpServerRequest& Re
 	TArray<FAssetData> PresetAssets;
 	IRemoteControlModule::Get().GetPresetAssets(PresetAssets, false);
 
-	WebRemoteControlUtils::SerializeMessage(FListPresetsResponse{ PresetAssets }, Response->Body);
+	TArray<TWeakObjectPtr<URemoteControlPreset>> EmbeddedPresets;
+	IRemoteControlModule::Get().GetEmbeddedPresets(EmbeddedPresets);
+
+	WebRemoteControlUtils::SerializeMessage(FListPresetsResponse{ PresetAssets, EmbeddedPresets }, Response->Body);
 
 	OnComplete(MoveTemp(Response));
 	return true;
diff --git a/Engine/Plugins/VirtualProduction/RemoteControl/Source/WebRemoteControl/Private/WebSocketMessageHandler.cpp b/Engine/Plugins/VirtualProduction/RemoteControl/Source/WebRemoteControl/Private/WebSocketMessageHandler.cpp
--- a/Engine/Plugins/VirtualProduction/RemoteControl/Source/WebRemoteControl/Private/WebSocketMessageHandler.cpp
+++ b/Engine/Plugins/VirtualProduction/RemoteControl/Source/WebRemoteControl/Private/WebSocketMessageHandler.cpp
@@ -176,7 +176,7 @@ namespace WebSocketMessageHandlerStructUtils
 
 		FStructOnScope FieldsChangedOnScope{ TopLevelStruct };
 		SetStringPropertyValue(Prop_Type, FieldsChangedOnScope, TEXT("PresetFieldsChanged"));
-		SetStringPropertyValue(Prop_PresetName, FieldsChangedOnScope, *Preset->GetFName().ToString());
+		SetStringPropertyValue(Prop_PresetName, FieldsChangedOnScope, *Preset->GetPresetName().ToString());
 		SetStringPropertyValue(Prop_PresetId, FieldsChangedOnScope, *Preset->GetPresetId().ToString());
 		SetStringPropertyValue(Prop_SequenceNumber, FieldsChangedOnScope, FString::Printf(TEXT("%lld"), SequenceNumber));
 		SetStructArrayPropertyValue(Prop_ChangedFields, FieldsChangedOnScope, PropertyValuesOnScope);
@@ -222,7 +222,7 @@ namespace WebSocketMessageHandlerStructUtils
 		FStructOnScope FieldsChangedOnScope{ TopLevelStruct };
 
 		SetStringPropertyValue(Prop_Type, FieldsChangedOnScope, TEXT("PresetActorModified"));
-		SetStringPropertyValue(Prop_PresetName, FieldsChangedOnScope, *Preset->GetFName().ToString());
+		SetStringPropertyValue(Prop_PresetName, FieldsChangedOnScope, *Preset->GetPresetName().ToString());
 		SetStringPropertyValue(Prop_PresetId, FieldsChangedOnScope, *Preset->GetPresetId().ToString());
 		SetStructArrayPropertyValue(Prop_ModifiedActors, FieldsChangedOnScope, ModifiedActorsOnScope);
 
@@ -232,15 +232,6 @@ namespace WebSocketMessageHandlerStructUtils
 
 namespace WebSocketMessageHandlerMiscUtils
 {
-	UWorld* GetEditorWorld()
-	{
-#if WITH_EDITOR
-		return GEditor ? GEditor->GetEditorWorldContext(false).World() : nullptr;
-#else
-		return nullptr;
-#endif
-	}
-
 	void* GetPresetPropertyClassPointer(URemoteControlPreset* Preset, const FGuid& PropertyId)
 	{
 		FRCObjectReference ObjectRef;
@@ -570,12 +561,6 @@ void FWebSocketMessageHandler::HandleWebSocketActorRegister(const FRemoteControl
 		return;
 	}
 
-	const UWorld* World = WebSocketMessageHandlerMiscUtils::GetEditorWorld();
-	if (!World)
-	{
-		return;
-	}
-
 	FWatchedClassData* WatchedClassData = ActorNotificationMap.Find(ActorClass);
 	FString ClassPath = ActorClass->GetPathName();
 
@@ -592,10 +577,23 @@ void FWebSocketMessageHandler::HandleWebSocketActorRegister(const FRemoteControl
 	// Register events for each actor and send the existing list of actors as "added" so the client is caught up
 	FRCActorsChangedEvent Event;
 	FRCActorsChangedData& ChangeData = Event.Changes.Add(ClassPath);
-	for (AActor* Actor : TActorRange<AActor>(World, ActorClass))
+
+	for (const FWorldContext& WorldContext : GEngine->GetWorldContexts())
 	{
-		ChangeData.AddedActors.Add(FRCActorDescription(Actor));
-		StartWatchingActor(Actor, ActorClass);
+		UWorld* World = WorldContext.World();
+		
+		// To support non-level editor worlds, instead of just checking the level editor's world,
+		// now every Editor-type world is checked.
+		if (!World || WorldContext.WorldType != EWorldType::Editor)
+		{
+			continue;
+		}
+
+		for (AActor* Actor : TActorRange<AActor>(World, ActorClass))
+		{
+			ChangeData.AddedActors.Add(FRCActorDescription(Actor));
+			StartWatchingActor(Actor, ActorClass);
+		}
 	}
 
 	TArray<uint8> Payload;
@@ -1191,7 +1189,7 @@ void FWebSocketMessageHandler::ProcessAddedProperties()
 		}
 
 		TArray<uint8> Payload;
-		WebRemoteControlUtils::SerializeMessage(FRCPresetFieldsAddedEvent{ Preset->GetFName(), Preset->GetPresetId(), AddedPropertiesDescription }, Payload);
+		WebRemoteControlUtils::SerializeMessage(FRCPresetFieldsAddedEvent{ Preset->GetPresetName(), Preset->GetPresetId(), AddedPropertiesDescription }, Payload);
 		BroadcastToPresetListeners(Entry.Key, Payload);
 	}
 
@@ -1216,7 +1214,7 @@ void FWebSocketMessageHandler::ProcessRemovedProperties()
 		ensure(Entry.Value.Key.Num() == Entry.Value.Value.Num());
 		
 		TArray<uint8> Payload;
-		WebRemoteControlUtils::SerializeMessage(FRCPresetFieldsRemovedEvent{ Preset->GetFName(), Preset->GetPresetId(), Entry.Value.Value, Entry.Value.Key }, Payload);
+		WebRemoteControlUtils::SerializeMessage(FRCPresetFieldsRemovedEvent{ Preset->GetPresetName(), Preset->GetPresetId(), Entry.Value.Value, Entry.Value.Key }, Payload);
 		BroadcastToPresetListeners(Entry.Key, Payload);
 	}
 	
@@ -1239,7 +1237,7 @@ void FWebSocketMessageHandler::ProcessRenamedFields()
 		}
 
 		TArray<uint8> Payload;
-		WebRemoteControlUtils::SerializeMessage(FRCPresetFieldsRenamedEvent{Preset->GetFName(), Preset->GetPresetId(), Entry.Value}, Payload);
+		WebRemoteControlUtils::SerializeMessage(FRCPresetFieldsRenamedEvent{Preset->GetPresetName(), Preset->GetPresetId(), Entry.Value}, Payload);
 		BroadcastToPresetListeners(Entry.Key, Payload);
 	}
 
@@ -1535,10 +1533,36 @@ bool FWebSocketMessageHandler::WriteActorPropertyChangePayload(URemoteControlPre
 void FWebSocketMessageHandler::OnActorAdded(AActor* Actor)
 {
 #if WITH_EDITOR
-	const UWorld* World = WebSocketMessageHandlerMiscUtils::GetEditorWorld();
-	if (Actor->GetWorld() != World)
+	// To support non-level editor worlds, we now check world type instead of specifically checking for the level editor.
+	if (Actor)
 	{
-		return;
+		bool bFoundWorld = false;
+
+		if (UWorld* ActorWorld = Actor->GetWorld())
+		{
+			for (const FWorldContext& WorldContext : GEngine->GetWorldContexts())
+			{
+				UWorld* ContextWorld = WorldContext.World();
+
+				if (ContextWorld != ActorWorld)
+				{
+					continue;
+				}
+
+				if (WorldContext.WorldType != EWorldType::Editor)
+				{
+					return;
+				}
+
+				bFoundWorld = true;
+				break;
+			}
+		}
+
+		if (!bFoundWorld)
+		{
+			return;
+		}
 	}
 #endif
 
@@ -1635,11 +1659,6 @@ void FWebSocketMessageHandler::OnActorListChanged()
 {
 	// We don't know exactly what changed, so manually check all the actors we know about
 
-	const UWorld* World = WebSocketMessageHandlerMiscUtils::GetEditorWorld();
-	if (!World)
-	{
-		return;
-	}
 
 	TSet<AActor*> RemainingActors;
 	WatchedActors.GetKeys(RemainingActors);
@@ -1647,13 +1666,26 @@ void FWebSocketMessageHandler::OnActorListChanged()
 	TArray<AActor*> NewActors;
 
 	// Find any new actors
-	for (AActor* Actor : TActorRange<AActor>(World))
+	for (const FWorldContext& WorldContext : GEngine->GetWorldContexts())
 	{
-		if (!WatchedActors.Contains(Actor))
+		UWorld* World = WorldContext.World();
+
+		// To support non-level editor worlds, instead of just checking the level editor's world,
+		// now every Editor-type world is checked.
+		if (!World || WorldContext.WorldType != EWorldType::Editor)
 		{
-			NewActors.Add(Actor);
+			continue;
+		}
+		
+		for (AActor* Actor : TActorRange<AActor>(World))
+		{
+			if (!WatchedActors.Contains(Actor))
+			{
+				NewActors.Add(Actor);
+			}
+
+			RemainingActors.Remove(Actor);
 		}
-		RemainingActors.Remove(Actor);
 	}
 
 	// Fire events for any actors that are now missing, which have presumably been deleted
diff --git a/Engine/Plugins/VirtualProduction/RemoteControlWebInterface/Source/RemoteControlWebInterface/Private/RCWebInterfaceLibrary.cpp b/Engine/Plugins/VirtualProduction/RemoteControlWebInterface/Source/RemoteControlWebInterface/Private/RCWebInterfaceLibrary.cpp
--- a/Engine/Plugins/VirtualProduction/RemoteControlWebInterface/Source/RemoteControlWebInterface/Private/RCWebInterfaceLibrary.cpp
+++ b/Engine/Plugins/VirtualProduction/RemoteControlWebInterface/Source/RemoteControlWebInterface/Private/RCWebInterfaceLibrary.cpp
@@ -52,7 +52,7 @@ TMap<FString, AActor*> URCWebInterfaceBlueprintLibrary::FindMatchingActorsToRebi
 				}
 				else if (UClass* SupportedBindingClass = RemoteControlProperty->GetSupportedBindingClass())
 				{
-					UWorld* World = URemoteControlLevelDependantBinding::GetCurrentWorld();
+					UWorld* World = RCPreset->GetPresetWorld(false);
 					if (SupportedBindingClass->IsChildOf(AActor::StaticClass()))
 					{
 						TArray<AActor*> MatchingActors;
