Static analysis issue with Engine/Plugins/VirtualProduction/RemoteControl/Source/WebRemoteControl/Private/WebSocketMessageHandler.cpp

diff --git a/Engine/Plugins/VirtualProduction/RemoteControl/Source/WebRemoteControl/Private/WebSocketMessageHandler.cpp b/Engine/Plugins/VirtualProduction/RemoteControl/Source/WebRemoteControl/Private/WebSocketMessageHandler.cpp
--- a/Engine/Plugins/VirtualProduction/RemoteControl/Source/WebRemoteControl/Private/WebSocketMessageHandler.cpp
+++ b/Engine/Plugins/VirtualProduction/RemoteControl/Source/WebRemoteControl/Private/WebSocketMessageHandler.cpp
@@ -1532,9 +1532,14 @@ bool FWebSocketMessageHandler::WriteActorPropertyChangePayload(URemoteControlPre
 
 void FWebSocketMessageHandler::OnActorAdded(AActor* Actor)
 {
+	// Should never be an issue, but can't hurt?!
+	if (!Actor || !IsValid(Actor))
+	{
+		return;
+	}
+
 #if WITH_EDITOR
 	// To support non-level editor worlds, we now check world type instead of specifically checking for the level editor.
-	if (Actor)
 	{
 		bool bFoundWorld = false;
 
