Missing include added.

diff --git a/Engine/Plugins/VirtualProduction/RemoteControl/Source/WebRemoteControl/Private/WebSocketMessageHandler.cpp b/Engine/Plugins/VirtualProduction/RemoteControl/Source/WebRemoteControl/Private/WebSocketMessageHandler.cpp
--- a/Engine/Plugins/VirtualProduction/RemoteControl/Source/WebRemoteControl/Private/WebSocketMessageHandler.cpp
+++ b/Engine/Plugins/VirtualProduction/RemoteControl/Source/WebRemoteControl/Private/WebSocketMessageHandler.cpp
@@ -5,6 +5,7 @@
 #include "Algo/ForEach.h"
 #include "Misc/CoreDelegates.h"
 #include "Misc/Guid.h"
+#include "Engine/Engine.h"
 #include "EngineUtils.h"
 #include "IRemoteControlModule.h"
 #include "GameFramework/Actor.h"
