1. Remote control preset UI revamp

2. Add filtering / search bar logic

diff --git a/Engine/Plugins/VirtualProduction/RemoteControl/Source/RemoteControlUI/Private/RemoteControlUIModule.cpp b/Engine/Plugins/VirtualProduction/RemoteControl/Source/RemoteControlUI/Private/RemoteControlUIModule.cpp
--- a/Engine/Plugins/VirtualProduction/RemoteControl/Source/RemoteControlUI/Private/RemoteControlUIModule.cpp
+++ b/Engine/Plugins/VirtualProduction/RemoteControl/Source/RemoteControlUI/Private/RemoteControlUIModule.cpp
@@ -38,6 +38,7 @@
 
 const FName FRemoteControlUIModule::EntityDetailsTabName = "RemoteControl_EntityDetails";
 const FName FRemoteControlUIModule::RemoteControlPanelTabName = "RemoteControl_RemoteControlPanel";
+const FString IRemoteControlUIModule::SettingsIniSection = TEXT("RemoteControl");
 
 static const FName DetailsTabIdentifiers_LevelEditor[] = {
 	"LevelEditorSelectionDetails",
@@ -234,7 +235,8 @@ namespace RemoteControlUIModule
 			{
 				if (ObjectWorldOuter->HasAnyFlags(RF_Transient))
 				{
-					return true;
+					// Check if it's also a remote control preset so that we don't accidentally include some other objects.
+					return  Object->IsA<URemoteControlPreset>();
 				}
 			}
 		}
diff --git a/Engine/Plugins/VirtualProduction/RemoteControl/Source/RemoteControlUI/Private/UI/Filters/CustomRCFilterData.cpp b/Engine/Plugins/VirtualProduction/RemoteControl/Source/RemoteControlUI/Private/UI/Filters/CustomRCFilterData.cpp
new file mode 100644
--- /dev/null
+++ b/Engine/Plugins/VirtualProduction/RemoteControl/Source/RemoteControlUI/Private/UI/Filters/CustomRCFilterData.cpp
@@ -0,0 +1,217 @@
+// Copyright Epic Games, Inc. All Rights Reserved.
+
+#include "UI/Filters/CustomRCFilterData.h"
+
+#include "AssetTypeCategories.h"
+#include "AssetToolsModule.h"
+
+#include "Filters/FilterBase.h"
+
+#include "IAssetTypeActions.h"
+
+#include "Styling/AppStyle.h"
+#include "Styling/SlateIconFinder.h"
+
+#include "UObject/Field.h"
+#include "UObject/Object.h"
+
+FSlateIcon FCustomTypeFilterData::EmptyIcon;
+FSlateIcon FCustomTypeFilterData::EntityIcon = FSlateIcon(FAppStyle::GetAppStyleSetName(), "Kismet.VariableList.TypeIcon");
+FSlateIcon FCustomTypeFilterData::ArrayIcon = FSlateIcon(FAppStyle::GetAppStyleSetName(), "Kismet.VariableList.ArrayTypeIcon");
+
+// TODO : Request for combined map icon svg(s) if possible, for now ignore the key image and show only the value image.
+FSlateIcon FCustomTypeFilterData::MapIcon = FSlateIcon(FAppStyle::GetAppStyleSetName(), "Kismet.VariableList.MapValueTypeIcon");
+
+FSlateIcon FCustomTypeFilterData::SetIcon = FSlateIcon(FAppStyle::GetAppStyleSetName(), "Kismet.VariableList.SetTypeIcon");
+
+void FCustomTypeFilterData::AddCategory(TSharedPtr<FFilterCategory> InCategory)
+{
+	Categories.Add(InCategory);
+}
+
+void FCustomTypeFilterData::BuildBackendFilter(FRCFilter& OutFilter)
+{
+	if (AssetTypeActions.IsValid() && AssetTypeActions.Pin()->CanFilter())
+	{
+		OutFilter.AddTypeFilter(AssetTypeActions.Pin()->GetSupportedClass());
+	}
+	else if (Class) // If there is no AssetTypeAction for this filter, simply add the class to the FRCFilter
+	{
+		OutFilter.AddTypeFilter(Class);
+	}
+	else if (!CustomTypeName.IsNone()) // It must be a custom type structure.
+	{
+		OutFilter.AddTypeFilter(CustomTypeName);
+	}
+	else // It must be basic property.
+	{
+		OutFilter.AddTypeFilter(EntityType);
+	}
+}
+
+TSharedPtr<IAssetTypeActions> FCustomTypeFilterData::GetAssetTypeActions() const
+{
+	return AssetTypeActions.Pin();
+}
+
+TArray<TSharedPtr<FFilterCategory>> FCustomTypeFilterData::GetCategories() const
+{
+	return Categories;
+}
+
+UClass* FCustomTypeFilterData::GetClass() const
+{
+	if (TSharedPtr<IAssetTypeActions> AssetTypeActionsPin = AssetTypeActions.Pin())
+	{
+		return AssetTypeActionsPin->GetSupportedClass();
+	}
+
+	return Class;
+}
+
+FFieldClass* FCustomTypeFilterData::GetEntityType() const
+{
+	return EntityType;
+}
+
+EEntityTypeCategories::Type FCustomTypeFilterData::GetEntityTypeCategory() const
+{
+	return EntityTypeCategory;
+}
+
+FLinearColor FCustomTypeFilterData::GetColor() const
+{
+	if (TSharedPtr<IAssetTypeActions> AssetTypeActionsPin = AssetTypeActions.Pin())
+	{
+		return AssetTypeActionsPin->GetTypeColor();
+	}
+
+	if (EntityTypeCategory & EEntityTypeCategories::Containers)
+	{
+		return FLinearColor::White;
+	}
+
+	return Color;
+}
+
+FText FCustomTypeFilterData::GetName() const
+{
+	if (TSharedPtr<IAssetTypeActions> AssetTypeActionsPin = AssetTypeActions.Pin())
+	{
+		return AssetTypeActionsPin->GetName();
+	}
+
+	if (Class)
+	{
+		return Class->GetDisplayNameText();
+	}
+
+	if (!CustomTypeName.IsNone())
+	{
+		return FText::FromName(CustomTypeName);
+	}
+
+	return GetProperEntityDisplayName();
+}
+
+FString FCustomTypeFilterData::GetFilterName() const
+{
+	if (TSharedPtr<IAssetTypeActions> AssetTypeActionsPin = AssetTypeActions.Pin())
+	{
+		return AssetTypeActionsPin->GetFilterName().ToString();
+	}
+	
+	if (Class)
+	{
+		return Class->GetFName().ToString();
+	}
+
+	if (!CustomTypeName.IsNone())
+	{
+		return CustomTypeName.ToString();
+	}
+
+	return EntityType->GetFName().ToString();
+}
+
+FSlateIcon FCustomTypeFilterData::GetSlateIcon() const
+{
+	if (EntityTypeCategory & EEntityTypeCategories::Assets)
+	{
+		return FSlateIconFinder::FindIconForClass(GetClass());
+	}
+
+	return GetProperEntityIcon();
+}
+
+void FCustomTypeFilterData::ResetBackendFilter(FRCFilter& OutFilter)
+{
+	if (AssetTypeActions.IsValid())
+	{
+		if (AssetTypeActions.Pin()->CanFilter())
+		{
+			OutFilter.RemoveTypeFilter(AssetTypeActions.Pin()->GetSupportedClass());
+		}
+	}
+	else if (Class) // If there is no AssetTypeAction for this filter, simply add the class to the FRCFilter
+	{
+		OutFilter.RemoveTypeFilter(Class);
+	}
+	else if (!CustomTypeName.IsNone()) // It must be a custom type structure.
+	{
+		OutFilter.RemoveTypeFilter(CustomTypeName);
+	}
+	else // It must be basic property.
+	{
+		OutFilter.RemoveTypeFilter(EntityType);
+	}
+}
+
+FText FCustomTypeFilterData::GetProperEntityDisplayName() const
+{
+	FString EntityTypeName = EntityType->GetName();
+
+	EntityTypeName.RemoveFromEnd("Property");
+
+	if (EntityTypeName.Contains("Bool"))
+	{
+		EntityTypeName = "Boolean";
+	}
+	else if (EntityTypeName.Contains("Int"))
+	{
+		EntityTypeName = EntityTypeName.Replace(TEXT("Int"), TEXT("Integer"));
+	}
+	else if (EntityTypeName.Contains("Str"))
+	{
+		EntityTypeName = "String";
+	}
+
+	return FText::FromString(EntityTypeName);
+}
+
+FSlateIcon FCustomTypeFilterData::GetProperEntityIcon() const
+{
+	FSlateIcon ProperEntityIcon = EmptyIcon;
+
+	if (EntityTypeCategory & EEntityTypeCategories::Containers)
+	{
+		if (EntityType->IsChildOf(FArrayProperty::StaticClass()))
+		{
+			ProperEntityIcon = ArrayIcon;
+		}
+		else if (EntityType->IsChildOf(FSetProperty::StaticClass()))
+		{
+			ProperEntityIcon = SetIcon;
+		}
+		else if (EntityType->IsChildOf(FMapProperty::StaticClass()))
+		{
+			ProperEntityIcon = MapIcon;
+		}
+	}
+	else if (EntityType->IsChildOf(FProperty::StaticClass()))
+	{
+		ProperEntityIcon = EntityIcon;
+	}
+
+	return ProperEntityIcon;
+}
diff --git a/Engine/Plugins/VirtualProduction/RemoteControl/Source/RemoteControlUI/Private/UI/Filters/RCFilter.cpp b/Engine/Plugins/VirtualProduction/RemoteControl/Source/RemoteControlUI/Private/UI/Filters/RCFilter.cpp
new file mode 100644
--- /dev/null
+++ b/Engine/Plugins/VirtualProduction/RemoteControl/Source/RemoteControlUI/Private/UI/Filters/RCFilter.cpp
@@ -0,0 +1,98 @@
+// Copyright Epic Games, Inc. All Rights Reserved.
+
+#include "UI/Filters/RCFilter.h"
+
+#include "UI/SRCPanelExposedField.h"
+#include "UI/SRCPanelTreeNode.h"
+
+#include "UObject/Field.h"
+#include "UObject/UnrealType.h"
+
+void FRCFilter::AddTypeFilter(FFieldClass* InEntityTypeFilter)
+{
+	EntityTypeFilters.Add(InEntityTypeFilter);
+}
+
+void FRCFilter::AddTypeFilter(const FName InCustomTypeFilter)
+{
+	CustomTypeFilters.Add(InCustomTypeFilter);
+}
+
+void FRCFilter::AddTypeFilter(UClass* InAssetTypeFilter)
+{
+	AssetTypeFilters.Add(InAssetTypeFilter);
+}
+
+bool FRCFilter::DoesPassFilters(FEntityFilterType InEntityItem)
+{
+	// Do not process Field Child.
+	if (InEntityItem->GetRCType() == SRCPanelTreeNode::FieldChild)
+	{
+		return false;
+	}
+
+	// Return early if the RC type of entity is Actor.
+	if (InEntityItem->GetRCType() == SRCPanelTreeNode::Actor)
+	{
+		return CustomTypeFilters.Contains(RemoteControlTypes::NAME_RCActors);
+	}
+
+	if (TSharedPtr<const SRCPanelExposedField> ExposedField = StaticCastSharedRef<const SRCPanelExposedField>(InEntityItem))
+	{
+		// Return early if the field type of entity is Function.
+		if (ExposedField->GetFieldType() == EExposedFieldType::Function)	
+		{
+			return CustomTypeFilters.Contains(RemoteControlTypes::NAME_RCFunctions);
+		}
+
+		if (TSharedPtr<FRemoteControlProperty> RCProperty = StaticCastSharedPtr<FRemoteControlProperty>(ExposedField->GetRemoteControlField().Pin()))
+		{
+			if (FProperty* Property = RCProperty->GetProperty())
+			{
+				if (FObjectProperty* ObjectProperty = CastField<FObjectProperty>(Property))
+				{
+					return AssetTypeFilters.Contains(ObjectProperty->PropertyClass);
+				}
+
+				if (FStructProperty* StructProperty = CastField<FStructProperty>(Property))
+				{
+					// TODO : Implement filters for Custom Struct Classes (as of now this is tailored for Primary Types specifically).
+					return CustomTypeFilters.Contains(StructProperty->Struct->GetFName());
+				}
+
+				return EntityTypeFilters.Contains(Property->GetClass());
+			}
+		}
+	}
+
+	return false;
+}
+
+bool FRCFilter::HasAnyActiveFilters()
+{
+	return CustomTypeFilters.Num() || EntityTypeFilters.Num() || AssetTypeFilters.Num();
+}
+
+void FRCFilter::RemoveTypeFilter(FFieldClass* InEntityTypeFilter)
+{
+	if (EntityTypeFilters.Contains(InEntityTypeFilter))
+	{
+		EntityTypeFilters.Remove(InEntityTypeFilter);
+	}
+}
+
+void FRCFilter::RemoveTypeFilter(const FName& InCustomTypeFilter)
+{
+	if (CustomTypeFilters.Contains(InCustomTypeFilter))
+	{
+		CustomTypeFilters.Remove(InCustomTypeFilter);
+	}
+}
+
+void FRCFilter::RemoveTypeFilter(UClass* InAssetTypeFilter)
+{
+	if (AssetTypeFilters.Contains(InAssetTypeFilter))
+	{
+		AssetTypeFilters.Remove(InAssetTypeFilter);
+	}
+}
diff --git a/Engine/Plugins/VirtualProduction/RemoteControl/Source/RemoteControlUI/Private/UI/Filters/SRCFilterBar.h b/Engine/Plugins/VirtualProduction/RemoteControl/Source/RemoteControlUI/Private/UI/Filters/SRCFilterBar.h
new file mode 100644
--- /dev/null
+++ b/Engine/Plugins/VirtualProduction/RemoteControl/Source/RemoteControlUI/Private/UI/Filters/SRCFilterBar.h
@@ -0,0 +1,782 @@
+// Copyright Epic Games, Inc. All Rights Reserved.
+
+#pragma once
+
+#include "CoreMinimal.h"
+
+#include "AssetTools/Public/AssetTypeCategories.h"
+#include "AssetTools/Public/AssetToolsModule.h"
+
+#include "Filters/SBasicFilterBar.h"
+
+#include "IAssetTypeActions.h"
+
+#include "UI/Filters/CustomRCFilterData.h"
+#include "UI/Filters/RCFilter.h"
+
+#include "SRCFilterBar.generated.h"
+
+#define LOCTEXT_NAMESPACE "RCFilterBar"
+
+/** Delegate that subclasses can use to specify classes to not include in this filter
+ * Returning false for a class will prevent it from showing up in the add filter dropdown
+ */
+DECLARE_DELEGATE_RetVal_OneParam(bool, FOnFilterEntityType, FFieldClass* /* Entity Type */);
+
+/* Delegate used by SRCFilterBar to populate the add filter menu */
+DECLARE_DELEGATE_ThreeParams(FOnPopulateAddRCFilterMenu, UToolMenu*, TSharedPtr<FFilterCategory>, FOnFilterEntityType)
+
+/** ToolMenuContext that is used to create the Add Filter Menu */
+UCLASS()
+class REMOTECONTROLUI_API URCFilterBarContext : public UObject
+{
+	GENERATED_BODY()
+
+public:
+	TSharedPtr<FFilterCategory> MenuExpansion;
+	FOnPopulateAddRCFilterMenu PopulateFilterMenu;
+};
+
+/**
+ * A custom filter bar to handle Remote Control Filtering System.
+ */
+template<typename FilterType>
+class SRCFilterBar : public SBasicFilterBar<FilterType>
+{
+public:
+
+	using FOnFilterChanged = typename SBasicFilterBar<FilterType>::FOnFilterChanged;
+	using FOnExtendAddFilterMenu = typename SBasicFilterBar<FilterType>::FOnExtendAddFilterMenu;
+
+	SLATE_BEGIN_ARGS(SRCFilterBar)
+		: _UseDefaultEntityFilters(true)
+	{}
+
+		/** Delegate for when filters have changed */
+		SLATE_EVENT(SRCFilterBar<FilterType>::FOnFilterChanged, OnFilterChanged)
+
+		/** Delegate to extend the Add Filter dropdown */
+		SLATE_EVENT(SRCFilterBar<FilterType>::FOnExtendAddFilterMenu, OnExtendAddFilterMenu)
+
+		/** Initial List of Custom Filters that will be added to the AddFilter Menu */
+		SLATE_ARGUMENT(TArray<TSharedRef<FFilterBase<FilterType>>>, CustomFilters)
+
+		/** Initial List of Custom Class filters that will be added to the AddFilter Menu */
+		SLATE_ARGUMENT(TArray<TSharedRef<FCustomTypeFilterData>>, CustomTypeFilters)
+
+		/** Whether the filter bar should provide the default RC filters */
+		SLATE_ARGUMENT(bool, UseDefaultEntityFilters)
+
+	SLATE_END_ARGS()
+
+	/** Constructs this widget with InArgs */
+	void Construct(const FArguments& InArgs)
+	{
+		bUseDefaultEntityFilters = InArgs._UseDefaultEntityFilters;
+		CustomTypeFilters = InArgs._CustomTypeFilters;
+		
+		typename SBasicFilterBar<FilterType>::FArguments Args;
+		Args._OnFilterChanged = InArgs._OnFilterChanged;
+		Args._CustomFilters = InArgs._CustomFilters;
+		Args._OnExtendAddFilterMenu = InArgs._OnExtendAddFilterMenu;
+		
+		SBasicFilterBar<FilterType>::Construct(Args);
+		
+		CreateEntityTypeFilters();
+	}
+
+protected:
+
+	/**
+	 * A custom data model to represent each type of filters we have.
+	 */
+	class FEntityFilter
+	{
+	public:
+
+		FEntityFilter(TSharedRef<FCustomTypeFilterData> InCustomTypeFilter, FOnFilterChanged InOnFilterChanged)
+			: bEnabled(false)
+			, CustomTypeFilter(InCustomTypeFilter)
+			, OnFilterChanged(InOnFilterChanged)
+		{
+		}
+
+		/** Returns contribution of this filter to the combined filter */
+		void GetBackendFilter(FRCFilter& OutFilter) const
+		{
+			if (CustomTypeFilter)
+			{
+				if (bEnabled)
+				{
+					CustomTypeFilter->BuildBackendFilter(OutFilter);
+				}
+				else
+				{
+					CustomTypeFilter->ResetBackendFilter(OutFilter);
+				}
+			}
+		}
+
+		/** Gets the asset type actions associated with this filter */
+		const TSharedPtr<FCustomTypeFilterData>& GetCustomTypeFilterData() const
+		{
+			return CustomTypeFilter;
+		}
+
+		/** Returns the display name for this filter */
+		FText GetFilterDisplayName() const
+		{
+			if (CustomTypeFilter.IsValid())
+			{
+				return CustomTypeFilter->GetName();
+			}
+
+			return LOCTEXT("NoneText", "None");
+		}
+
+		/** Returns the actual name of this filter. */
+		FString GetFilterName() const
+		{
+			if (CustomTypeFilter.IsValid())
+			{
+				return CustomTypeFilter->GetFilterName();
+			}
+
+			return "None";
+		}
+
+		/** Returns true if this filter contributes to the combined filter */
+		bool IsEnabled() const
+		{
+			return bEnabled;
+		}
+
+		/** Sets whether or not this filter is applied to the combined filter */
+		void SetEnabled(bool InEnabled, bool InExecuteOnFilterChanged = true)
+		{
+			if (InEnabled != bEnabled)
+			{
+				bEnabled = InEnabled;
+
+				if (InExecuteOnFilterChanged)
+				{
+					OnFilterChanged.ExecuteIfBound();
+				}
+			}
+		}
+
+	protected:
+
+		/** Whether this filter is active or not. */
+		bool bEnabled;
+
+		/** View model for this filter describing the nature of this filter. */
+		TSharedPtr<FCustomTypeFilterData> CustomTypeFilter;
+
+		/** Invoked when the filter toggled */
+		FOnFilterChanged OnFilterChanged;
+	};
+
+public:
+
+	/** Use this function to get an FRCFilter that represents all the Asset & Entity Type filters that are currently active. */
+	FRCFilter GetCombinedBackendFilter() const
+	{
+		FRCFilter CombinedFilter;
+
+		// Add all selected filters
+		for (int32 FilterIdx = 0; FilterIdx < EntityFilters.Num(); ++FilterIdx)
+		{
+			const TSharedPtr<FEntityFilter> EntityFilter = EntityFilters[FilterIdx];
+
+			if (EntityFilter.IsValid())
+			{
+				EntityFilter->GetBackendFilter(CombinedFilter);
+			}
+		}
+
+		return CombinedFilter;
+	}
+
+	/** Remove all filters from the filter bar, while disabling any active ones */
+	virtual void RemoveAllFilters() override
+	{
+		EntityFilters.Empty();
+		
+		SBasicFilterBar<FilterType>::RemoveAllFilters();
+		
+		this->OnFilterChanged.ExecuteIfBound();
+	}
+
+protected:
+
+	/** Add an RC Filter to the list of filters, making it "Active" but not enabled */
+	TSharedRef<FEntityFilter> AddEntityFilter(const TSharedPtr<FCustomTypeFilterData>& CustomTypeFilter)
+	{
+		TSharedRef<FEntityFilter> NewFilter = MakeShared<FEntityFilter>(CustomTypeFilter.ToSharedRef(), this->OnFilterChanged);
+
+		// Add it to our list of just EntityFilters
+		EntityFilters.Add(NewFilter);
+
+		return NewFilter;
+	}
+	
+	/** Returns true if any filters are applied */
+	bool HasAnyFilters() const
+	{
+		return EntityFilters.Num() > 0;
+	}
+
+	/** Attempts to remove the given filter from list. */
+	void RemoveEntityFilter(const TSharedPtr<FCustomTypeFilterData>& CustomTypeFilter, bool ExecuteOnFilterChanged = true)
+	{
+		TSharedPtr<FEntityFilter> FilterToRemove;
+
+		for (const TSharedPtr<FEntityFilter>& EntityFilter : EntityFilters)
+		{
+			const TSharedPtr<FCustomTypeFilterData>& EntityFilterData = EntityFilter->GetCustomTypeFilterData();
+
+			if (EntityFilterData == CustomTypeFilter)
+			{
+				FilterToRemove = EntityFilter;
+				break;
+			}
+		}
+
+		if (FilterToRemove.IsValid())
+		{
+			EntityFilters.Remove(FilterToRemove.ToSharedRef());
+
+			if (ExecuteOnFilterChanged)
+			{
+				this->OnFilterChanged.ExecuteIfBound();
+			}
+		}
+	}
+
+	/** Called when reset filters option is pressed */
+	void OnResetFilters()
+	{
+		RemoveAllFilters();
+	}
+
+	/** Create the default set of EClassCastFlags Filters provided with the widget if requested */
+	void CreateEntityTypeFilters()
+	{
+		if (!bUseDefaultEntityFilters)
+		{
+			return;
+		}
+
+		TypeFilterCategories.Empty();
+
+		// Add the Basic property types category
+
+		// Entity Type Categories
+		TypeFilterCategories.Add(EEntityTypeCategories::Core
+			, MakeShareable(new FFilterCategory(LOCTEXT("TypeFilter_Core", "Core"), LOCTEXT("CoreFilterTooltip", "Filter in entities of core types.")))
+		);
+		TypeFilterCategories.Add(EEntityTypeCategories::Strings
+			, MakeShareable(new FFilterCategory(LOCTEXT("TypeFilter_Strings", "Strings"), LOCTEXT("StringsFilterTooltip", "Filter in entities of string types.")))
+		);
+		TypeFilterCategories.Add(EEntityTypeCategories::Structs
+			, MakeShareable(new FFilterCategory(LOCTEXT("TypeFilter_Structs", "Structures"), LOCTEXT("StructsFilterTooltip", "Filter in entities of structure types.")))
+		);
+		TypeFilterCategories.Add(EEntityTypeCategories::Objects
+			, MakeShareable(new FFilterCategory(LOCTEXT("TypeFilter_Objects", "Objects"), LOCTEXT("ObjectsFilterTooltip", "Filter in entities of objects types.")))
+		);
+		TypeFilterCategories.Add(EEntityTypeCategories::Containers
+			, MakeShareable(new FFilterCategory(LOCTEXT("TypeFilter_Containers", "Containers"), LOCTEXT("ContainersFilterTooltip", "Filter in entities of container types.")))
+		);
+		TypeFilterCategories.Add(EEntityTypeCategories::References
+			, MakeShareable(new FFilterCategory(LOCTEXT("TypeFilter_References", "References"), LOCTEXT("ReferencesFilterTooltip", "Filter in entities of reference types.")))
+		);
+		TypeFilterCategories.Add(EEntityTypeCategories::UserDefined
+			, MakeShareable(new FFilterCategory(LOCTEXT("TypeFilter_UserDefined", "Custom"), LOCTEXT("UserDefinedFilterTooltip", "Filter in entities of custom types.")))
+		);
+		TypeFilterCategories.Add(EEntityTypeCategories::Primary
+			, MakeShareable(new FFilterCategory(LOCTEXT("TypeFilter_Primary", "All"), LOCTEXT("PrimaryFilterTooltip", "Filter in entities of unique types.")))
+		);
+
+		// Core Type Filters.
+		PopulateFilterData(EEntityTypeCategories::Core, FBoolProperty::StaticClass());
+		PopulateFilterData(EEntityTypeCategories::Core, FByteProperty::StaticClass());
+		PopulateFilterData(EEntityTypeCategories::Core, FEnumProperty::StaticClass());
+		PopulateFilterData(EEntityTypeCategories::Core, FFloatProperty::StaticClass());
+		PopulateFilterData(EEntityTypeCategories::Core, FIntProperty::StaticClass());
+		PopulateFilterData(EEntityTypeCategories::Core, FInt64Property::StaticClass());
+		
+		// String Type Filters.
+		PopulateFilterData(EEntityTypeCategories::Strings, FNameProperty::StaticClass());
+		PopulateFilterData(EEntityTypeCategories::Strings, FStrProperty::StaticClass());
+		PopulateFilterData(EEntityTypeCategories::Strings, FTextProperty::StaticClass());
+		
+		// Primary Type Filters.
+		PopulateFilterData(EEntityTypeCategories::Primary, FStructProperty::StaticClass(), NAME_Vector);
+		PopulateFilterData(EEntityTypeCategories::Primary, FStructProperty::StaticClass(), NAME_Rotator);
+		PopulateFilterData(EEntityTypeCategories::Primary, FStructProperty::StaticClass(), NAME_Color);
+		PopulateFilterData(EEntityTypeCategories::Primary, FStructProperty::StaticClass(), NAME_Transform);
+		PopulateFilterData(EEntityTypeCategories::Primary, FProperty::StaticClass(), RemoteControlTypes::NAME_RCActors);
+		PopulateFilterData(EEntityTypeCategories::Primary, FProperty::StaticClass(), RemoteControlTypes::NAME_RCFunctions);
+		
+		// Container Type Filters.
+		PopulateFilterData(EEntityTypeCategories::Containers, FArrayProperty::StaticClass());
+		PopulateFilterData(EEntityTypeCategories::Containers, FMapProperty::StaticClass());
+		PopulateFilterData(EEntityTypeCategories::Containers, FSetProperty::StaticClass());
+		
+		// TODO : Add support to Custom Classes.
+		FAssetToolsModule& AssetToolsModule = FModuleManager::LoadModuleChecked<FAssetToolsModule>(TEXT("AssetTools"));
+
+		AssetFilterCategories.Empty();
+
+		// Add the Basic category
+		AssetFilterCategories.Add(EAssetTypeCategories::Basic, MakeShareable(new FFilterCategory(LOCTEXT("BasicFilter", "Basic"), LOCTEXT("BasicFilterTooltip", "Filter by basic assets."))));
+		
+		// Add the Material category
+		AssetFilterCategories.Add(EAssetTypeCategories::Materials, MakeShareable(new FFilterCategory(LOCTEXT("MaterialFilter", "Materials"), LOCTEXT("MaterialsFilterTooltip", "Filter by material assets."))));
+
+		// Get the browser type maps
+		TArray<TWeakPtr<IAssetTypeActions>> AssetTypeActionsList;
+		AssetToolsModule.Get().GetAssetTypeActionsList(AssetTypeActionsList);
+
+		// Sort the list
+		struct FCompareIAssetTypeActions
+		{
+			FORCEINLINE bool operator()(const TWeakPtr<IAssetTypeActions>& A, const TWeakPtr<IAssetTypeActions>& B) const
+			{
+				return A.Pin()->GetName().CompareTo(B.Pin()->GetName()) == -1;
+			}
+		};
+
+		AssetTypeActionsList.Sort(FCompareIAssetTypeActions());
+
+		const TSharedRef<FPathPermissionList>& AssetClassPermissionList = AssetToolsModule.Get().GetAssetClassPathPermissionList(EAssetClassAction::CreateAsset);
+
+		// For every asset type, convert it to an FCustomTypeFilterData and add it to the list
+		for (int32 ClassIdx = 0; ClassIdx < AssetTypeActionsList.Num(); ++ClassIdx)
+		{
+			const TWeakPtr<IAssetTypeActions>& WeakTypeActions = AssetTypeActionsList[ClassIdx];
+			if (WeakTypeActions.IsValid())
+			{
+				TSharedPtr<IAssetTypeActions> TypeActions = WeakTypeActions.Pin();
+				if (ensure(TypeActions.IsValid()) && TypeActions->CanFilter())
+				{
+					UClass* SupportedClass = TypeActions->GetSupportedClass();
+					if ((!SupportedClass || AssetClassPermissionList->PassesFilter(SupportedClass->GetClassPathName().ToString())) &&
+						(TypeActions->GetCategories() & (EAssetTypeCategories::Basic | EAssetTypeCategories::Materials)) &&
+						!AssetTypesToIgnore.Contains(TypeActions->GetName().ToString()))
+					{
+						// Convert the AssetTypeAction to an FCustomTypeFilterData and add it to our list
+						PopulateFilterData(TypeActions);
+					}
+				}
+			}
+		}
+
+		// Do a second pass through all the CustomTypeFilters with TypeFilterCategories to update their categories
+		UpdateEntityTypeCategories();
+	}
+
+	/** Normalizes the CustomTypeFilters to their respective categories. */
+	void UpdateEntityTypeCategories()
+	{
+		for (const TSharedRef<FCustomTypeFilterData>& CustomTypeFilter : CustomTypeFilters)
+		{
+			for (TMap<EEntityTypeCategories::Type, TSharedPtr<FFilterCategory>>::TIterator MenuIt = TypeFilterCategories.CreateIterator(); MenuIt; ++MenuIt)
+			{
+				if (MenuIt.Key() & CustomTypeFilter->GetEntityTypeCategory())
+				{
+					CustomTypeFilter->AddCategory(MenuIt.Value());
+				}
+			}
+
+			for (TMap<EAssetTypeCategories::Type, TSharedPtr<FFilterCategory>>::TIterator MenuIt = AssetFilterCategories.CreateIterator(); MenuIt; ++MenuIt)
+			{
+				if (TSharedPtr<IAssetTypeActions> AssetTypeActions = CustomTypeFilter->GetAssetTypeActions())
+				{
+					if (MenuIt.Key() & AssetTypeActions->GetCategories())
+					{
+						CustomTypeFilter->AddCategory(MenuIt.Value());
+					}
+				}
+			}
+		}
+	}
+	
+	/* Filter Dropdown Related Functionality */
+
+	/** Handler for when the add filter menu is populated by a category */
+	void CreateFiltersMenuCategory(FToolMenuSection& Section, const TArray<TSharedPtr<FCustomTypeFilterData>> CustomTypeFilterDatas) const
+	{
+		for (int32 EntityFilterIdx = 0; EntityFilterIdx < CustomTypeFilterDatas.Num(); ++EntityFilterIdx)
+		{
+			const TSharedPtr<FCustomTypeFilterData>& CustomTypeFilterData = CustomTypeFilterDatas[EntityFilterIdx];
+
+			const FText& LabelText = CustomTypeFilterData->GetName();
+
+			Section.AddMenuEntry(
+				NAME_None,
+				LabelText,
+				FText::Format(LOCTEXT("FilterByTooltipPrefix", "Filter in entities of type {0}."), LabelText),
+				CustomTypeFilterData->GetSlateIcon(),
+				FUIAction(
+					FExecuteAction::CreateSP(const_cast<SRCFilterBar<FilterType>*>(this), &SRCFilterBar<FilterType>::FilterByTypeClicked, CustomTypeFilterData),
+					FCanExecuteAction(),
+					FIsActionChecked::CreateSP(this, &SRCFilterBar<FilterType>::IsEntityTypeInUse, CustomTypeFilterData)),
+				EUserInterfaceActionType::ToggleButton
+			);
+		}
+	}
+
+	/** Handler for when the add filter menu is populated by a category */
+	void CreateFiltersMenuCategory(UToolMenu* InMenu, const TArray<TSharedPtr<FCustomTypeFilterData>> CustomTypeFilterDatas) const
+	{
+		CreateFiltersMenuCategory(InMenu->AddSection("Section"), CustomTypeFilterDatas);
+	}
+
+	/** Handler for when the add filter button was clicked */
+	virtual TSharedRef<SWidget> MakeAddFilterMenu() override
+	{
+		const FName FilterMenuName = "RCFilterBar.RCFilterMenu";
+		if (!UToolMenus::Get()->IsMenuRegistered(FilterMenuName))
+		{
+			UToolMenu* Menu = UToolMenus::Get()->RegisterMenu(FilterMenuName);
+			Menu->bShouldCloseWindowAfterMenuSelection = true;
+			Menu->bCloseSelfOnly = true;
+
+			Menu->AddDynamicSection(NAME_None, FNewToolMenuDelegate::CreateLambda([](UToolMenu* InMenu)
+				{
+					if (URCFilterBarContext* Context = InMenu->FindContext<URCFilterBarContext>())
+					{
+						Context->PopulateFilterMenu.ExecuteIfBound(InMenu, Context->MenuExpansion, FOnFilterEntityType());
+					}
+				})
+			);
+		}
+
+		URCFilterBarContext* FilterBarContext = NewObject<URCFilterBarContext>();
+		FilterBarContext->PopulateFilterMenu = FOnPopulateAddRCFilterMenu::CreateSP(this, &SRCFilterBar<FilterType>::PopulateAddFilterMenu);
+		
+		// Auto expand the Structures Category if it is needed
+		if (TSharedPtr<FFilterCategory>* PrimaryCategory = TypeFilterCategories.Find(EEntityTypeCategories::Primary))
+		{
+			FilterBarContext->MenuExpansion = *PrimaryCategory;
+		}
+
+		FToolMenuContext ToolMenuContext(FilterBarContext);
+
+		return UToolMenus::Get()->GenerateWidget(FilterMenuName, ToolMenuContext);
+	}
+
+	/** Handler to Populate the Add Filter Menu. Use OnFilterEntityType in Subclasses to add classes to the exclusion list */
+	virtual void PopulateAddFilterMenu(UToolMenu* Menu, TSharedPtr<FFilterCategory> MenuExpansion, FOnFilterEntityType OnFilterEntityType)
+	{
+		// A local struct to describe a category in the filter menu
+		struct FCategoryMenu
+		{
+			/** The Property Type that belong to this category */
+			TArray<TSharedPtr<FCustomTypeFilterData>> EntityTypes;
+
+			//Menu section
+			FName SectionExtensionHook;
+			FText SectionHeading;
+
+			FCategoryMenu(const FName& InSectionExtensionHook, const FText& InSectionHeading)
+				: SectionExtensionHook(InSectionExtensionHook)
+				, SectionHeading(InSectionHeading)
+			{}
+		};
+
+		// Create a map of Categories to Menus
+		TMap<TSharedPtr<FFilterCategory>, FCategoryMenu> CategoryToMenuMap;
+
+		// For every asset type, move it into all the categories it should appear in
+		for (const TSharedRef<FCustomTypeFilterData>& CustomTypeFilter : CustomTypeFilters)
+		{
+			bool bPassesExternalFilters = true;
+
+			// Run any external class filters we have
+			if (OnFilterEntityType.IsBound())
+			{
+				bPassesExternalFilters = OnFilterEntityType.Execute(CustomTypeFilter->GetEntityType());
+			}
+
+			if (bPassesExternalFilters)
+			{
+				/** Get all the categories this filter belongs to */
+				TArray<TSharedPtr<FFilterCategory>> Categories = CustomTypeFilter->GetCategories();
+
+				for (const TSharedPtr<FFilterCategory>& Category : Categories)
+				{
+					// If the category for this custom class already exists
+					if (FCategoryMenu* CategoryMenu = CategoryToMenuMap.Find(Category))
+					{
+						CategoryMenu->EntityTypes.Add(CustomTypeFilter);
+					}
+					// Otherwise create a new FCategoryMenu for the category and add it to the map
+					else
+					{
+						const FName ExtensionPoint = NAME_None;
+						const FText SectionHeading = FText::Format(LOCTEXT("WildcardFilterHeadingHeadingTooltip", "{0} Filters"), Category->Title);
+
+						FCategoryMenu NewCategoryMenu(ExtensionPoint, SectionHeading);
+						NewCategoryMenu.EntityTypes.Add(CustomTypeFilter);
+
+						CategoryToMenuMap.Add(Category, NewCategoryMenu);
+					}
+				}
+			}
+		}
+
+		// Remove any empty categories
+		for (auto MenuIt = CategoryToMenuMap.CreateIterator(); MenuIt; ++MenuIt)
+		{
+			if (MenuIt.Value().EntityTypes.Num() == 0)
+			{
+				CategoryToMenuMap.Remove(MenuIt.Key());
+			}
+		}
+
+		// Set the extension hook for the primary category, if it exists and we have any entities for it
+		if (TSharedPtr<FFilterCategory>* PrimaryCategory = TypeFilterCategories.Find(EEntityTypeCategories::Primary))
+		{
+			if (*PrimaryCategory == MenuExpansion)
+			{
+				if (FCategoryMenu* PrimaryMenu = CategoryToMenuMap.Find(MenuExpansion))
+				{
+					PrimaryMenu->SectionExtensionHook = "FilterBarFilterPrimaryEntity";
+				}
+			}
+		}
+
+		// Populate the common filter sections (Reset Filters etc)
+		{
+			FToolMenuSection& Section = Menu->AddSection("FilterBarResetFilters");
+			Section.AddMenuEntry(
+				"ResetFilters",
+				LOCTEXT("FilterListResetFilters", "Reset Filters"),
+				LOCTEXT("FilterListResetToolTip", "Resets current filter selection"),
+				FSlateIcon(FAppStyle::Get().GetStyleSetName(), "PropertyWindow.DiffersFromDefault"),
+				FUIAction(
+					FExecuteAction::CreateSP(this, &SRCFilterBar<FilterType>::OnResetFilters),
+					FCanExecuteAction::CreateLambda([this]() { return HasAnyFilters(); }))
+			);
+
+			this->OnExtendAddFilterMenu.ExecuteIfBound(Menu);
+		}
+
+		// If we want to expand a category
+		if (MenuExpansion)
+		{
+			// First add the expanded category, this appears as standard entries in the list (Note: intentionally not using FindChecked here as removing it from the map later would cause the ref to be garbage)
+			FCategoryMenu* ExpandedCategory = CategoryToMenuMap.Find(MenuExpansion);
+			if (ExpandedCategory)
+			{
+				FToolMenuSection& Section = Menu->AddSection(ExpandedCategory->SectionExtensionHook, LOCTEXT("QuickFiltersLabel", "Quick Filters"));
+				
+				// If we are doing a full menu (i.e expanding primary) we add a menu entry which toggles all other categories
+				if (MenuExpansion == TypeFilterCategories.FindChecked(EEntityTypeCategories::Primary))
+				{
+					Section.AddMenuEntry(
+						NAME_None,
+						MenuExpansion->Title,
+						MenuExpansion->Tooltip,
+						FSlateIcon(FAppStyle::Get().GetStyleSetName(), "Kismet.VariableList.TypeIcon"),
+						FUIAction(
+							FExecuteAction::CreateSP(this, &SRCFilterBar<FilterType>::FilterByTypeCategoryClicked, MenuExpansion, ExpandedCategory->EntityTypes),
+							FCanExecuteAction(),
+							FGetActionCheckState::CreateSP(this, &SRCFilterBar<FilterType>::IsEntityTypeCategoryChecked, MenuExpansion, ExpandedCategory->EntityTypes)),
+						EUserInterfaceActionType::ToggleButton
+					);
+				}
+
+				// Now populate with all the assets from the expanded category
+				SRCFilterBar<FilterType>::CreateFiltersMenuCategory(Section, ExpandedCategory->EntityTypes);
+				
+				// Remove the Expanded from the map now, as this is treated differently and is no longer needed.
+				CategoryToMenuMap.Remove(MenuExpansion);
+			}
+		}
+
+		TSharedPtr<FFilterCategory>* PrimaryCategory = TypeFilterCategories.Find(EEntityTypeCategories::Primary);
+
+		// We are in Full Menu Mode if there is no menu expansion, or the menu expansion is EEntityTypeCategories::Primary
+		bool bInFullMenuMode = !MenuExpansion || (PrimaryCategory && MenuExpansion == *PrimaryCategory);
+
+		// If we are in full menu mode, add all the other categories as submenus
+		if (bInFullMenuMode)
+		{
+			FToolMenuSection& Section = Menu->AddSection("TypeFilterBarFilterAdvancedEntity", LOCTEXT("AdvancedTypesMenuHeading", "Advanced Filters"));
+
+			// Sort by category name so that we add the submenus in alphabetical order
+			CategoryToMenuMap.KeySort([](const TSharedPtr<FFilterCategory>& A, const TSharedPtr<FFilterCategory>& B) {
+				return A->Title.CompareTo(B->Title) < 0;
+				});
+
+			// For all the remaining categories, add them as submenus
+			for (const TPair<TSharedPtr<FFilterCategory>, FCategoryMenu>& CategoryMenuPair : CategoryToMenuMap)
+			{
+				Section.AddSubMenu(
+					NAME_None,
+					CategoryMenuPair.Key->Title,
+					CategoryMenuPair.Key->Tooltip,
+					FNewToolMenuDelegate::CreateSP(this, &SRCFilterBar<FilterType>::CreateFiltersMenuCategory, CategoryMenuPair.Value.EntityTypes),
+					FUIAction(
+						FExecuteAction::CreateSP(this, &SRCFilterBar<FilterType>::FilterByTypeCategoryClicked, CategoryMenuPair.Key, CategoryMenuPair.Value.EntityTypes),
+						FCanExecuteAction(),
+						FGetActionCheckState::CreateSP(this, &SRCFilterBar<FilterType>::IsEntityTypeCategoryChecked, CategoryMenuPair.Key, CategoryMenuPair.Value.EntityTypes)),
+					EUserInterfaceActionType::ToggleButton
+				);
+			}
+		}
+
+		// Now add all non-asset filters
+		this->PopulateCustomFilters(Menu);
+	}
+
+	/* Asset Type filter related functionality */
+
+	/** Handler for when filter by type is selected */
+	void FilterByTypeClicked(TSharedPtr<FCustomTypeFilterData> CustomTypeFilterData)
+	{
+		if (CustomTypeFilterData.IsValid())
+		{
+			if (IsEntityTypeInUse(CustomTypeFilterData))
+			{
+				RemoveEntityFilter(CustomTypeFilterData);
+			}
+			else
+			{
+				TSharedRef<FEntityFilter> NewFilter = AddEntityFilter(CustomTypeFilterData);
+				NewFilter->SetEnabled(true);
+			}
+		}
+	}
+
+	/** Handler to determine the "checked" state of class filter in the filter dropdown */
+	bool IsEntityTypeInUse(TSharedPtr<FCustomTypeFilterData> EntityType) const
+	{
+		for (const TSharedPtr<FEntityFilter> EntityFilter : this->EntityFilters)
+		{
+			if (EntityFilter.IsValid() && EntityFilter->GetCustomTypeFilterData() == EntityType)
+			{
+				return true;
+			}
+		}
+
+		return false;
+	}
+
+	/** Handler for when filter by type category is selected */
+	void FilterByTypeCategoryClicked(TSharedPtr<FFilterCategory> TypeCategory, TArray<TSharedPtr<FCustomTypeFilterData>> EntityTypes)
+	{
+		bool bFullCategoryInUse = IsEntityTypeCategoryInUse(TypeCategory, EntityTypes);
+		bool ExecuteOnFilterChanged = false;
+
+		for (const TSharedPtr<FCustomTypeFilterData>& CustomEntityType : EntityTypes)
+		{
+			if (bFullCategoryInUse)
+			{
+				RemoveEntityFilter(CustomEntityType);
+				ExecuteOnFilterChanged = true;
+			}
+			else if (!IsEntityTypeInUse(CustomEntityType))
+			{
+				TSharedRef<FEntityFilter> NewFilter = AddEntityFilter(CustomEntityType);
+				NewFilter->SetEnabled(true, false);
+				ExecuteOnFilterChanged = true;
+			}
+		}
+
+		if (ExecuteOnFilterChanged)
+		{
+			this->OnFilterChanged.ExecuteIfBound();
+		}
+	}
+
+	/** Handler to determine the "checked" state of an type category in the filter dropdown */
+	ECheckBoxState IsEntityTypeCategoryChecked(TSharedPtr<FFilterCategory> TypeCategory, TArray<TSharedPtr<FCustomTypeFilterData>> EntityTypes) const
+	{
+		bool bIsAnyActionInUse = false;
+		bool bIsAnyActionNotInUse = false;
+
+		for (const TSharedPtr<FCustomTypeFilterData>& CustomEntityTypeFilter : EntityTypes)
+		{
+			if (IsEntityTypeInUse(CustomEntityTypeFilter))
+			{
+				bIsAnyActionInUse = true;
+			}
+			else
+			{
+				bIsAnyActionNotInUse = true;
+			}
+
+			if (bIsAnyActionInUse && bIsAnyActionNotInUse)
+			{
+				return ECheckBoxState::Undetermined;
+			}
+
+		}
+
+		if (bIsAnyActionInUse)
+		{
+			return ECheckBoxState::Checked;
+		}
+		else
+		{
+			return ECheckBoxState::Unchecked;
+		}
+	}
+
+	/** Function to check if a given type category is in use */
+	bool IsEntityTypeCategoryInUse(TSharedPtr<FFilterCategory> TypeCategory, TArray<TSharedPtr<FCustomTypeFilterData>> EntityTypes) const
+	{
+		ECheckBoxState AssetTypeCategoryCheckState = IsEntityTypeCategoryChecked(TypeCategory, EntityTypes);
+
+		if (AssetTypeCategoryCheckState == ECheckBoxState::Unchecked)
+		{
+			return false;
+		}
+
+		// An entity type category is in use if any of its type actions are in use (ECheckBoxState::Checked or ECheckBoxState::Undetermined)
+		return true;
+	}
+
+private:
+
+	/** Populates CustomTypeFilters from Entity Type and its Category. */
+	void PopulateFilterData(EEntityTypeCategories::Type InEntityTypeCategory, FFieldClass* InEntityType, const FName& InCustomTypeName = NAME_None)
+	{
+		TSharedRef<FCustomTypeFilterData> CustomTypeFilterData = MakeShared<FCustomTypeFilterData>(InEntityTypeCategory, InEntityType, InCustomTypeName);
+		CustomTypeFilters.Add(CustomTypeFilterData);
+	}
+	
+	/** Populates CustomTypeFilters from Asset Type Action. */
+	void PopulateFilterData(TWeakPtr<IAssetTypeActions> InAssetTypeAction)
+	{
+		TSharedRef<FCustomTypeFilterData> CustomTypeFilterData = MakeShared<FCustomTypeFilterData>(InAssetTypeAction);
+		CustomTypeFilters.Add(CustomTypeFilterData);
+	}
+
+protected:
+
+	/** A copy of all active Entity Filters in this->EntityFilters for convenient access. */
+	TArray<TSharedRef<FEntityFilter>> EntityFilters;
+
+	/** List of all custom entity filters. */
+	TArray<TSharedRef<FCustomTypeFilterData>> CustomTypeFilters;
+
+	/** A map of property type with its corresponding filter category. */
+	TMap<EEntityTypeCategories::Type, TSharedPtr<FFilterCategory>> TypeFilterCategories;
+	
+	/** A map of asset type with its corresponding filter category. */
+	TMap<EAssetTypeCategories::Type, TSharedPtr<FFilterCategory>> AssetFilterCategories;
+
+	/** A list of Asset Types to ignore. */
+	TSet<FString> AssetTypesToIgnore = { "C++ Class" };
+
+	/** Whether the filter bar provides the default RC Filters */
+	bool bUseDefaultEntityFilters;
+};
+
+#undef LOCTEXT_NAMESPACE
diff --git a/Engine/Plugins/VirtualProduction/RemoteControl/Source/RemoteControlUI/Private/UI/Filters/SRCPanelFilter.cpp b/Engine/Plugins/VirtualProduction/RemoteControl/Source/RemoteControlUI/Private/UI/Filters/SRCPanelFilter.cpp
new file mode 100644
--- /dev/null
+++ b/Engine/Plugins/VirtualProduction/RemoteControl/Source/RemoteControlUI/Private/UI/Filters/SRCPanelFilter.cpp
@@ -0,0 +1,102 @@
+// Copyright Epic Games, Inc. All Rights Reserved.
+
+#include "UI/Filters/SRCPanelFilter.h"
+
+#include "Misc/ConfigCacheIni.h"
+
+#include "SlateOptMacros.h"
+
+#define LOCTEXT_NAMESPACE "RCPanelFilter"
+
+BEGIN_SLATE_FUNCTION_BUILD_OPTIMIZATION
+
+void SRCPanelFilter::Construct(const FArguments& InArgs)
+{
+	this->OnFilterChanged = InArgs._OnFilterChanged;
+
+	SRCFilterBar<FEntityFilterType>::FArguments Args;
+
+	/** Explicitly setting this to true as it should ALWAYS be true for SRCPanelFilter */
+	Args._UseDefaultEntityFilters = true;
+	Args._OnFilterChanged = this->OnFilterChanged;
+
+	SRCFilterBar<FEntityFilterType>::Construct(Args);
+}
+
+END_SLATE_FUNCTION_BUILD_OPTIMIZATION
+
+void SRCPanelFilter::SaveSettings(const FString& IniFilename, const FString& IniSection, const FString& SettingsString) const
+{
+	// Holds the list of filters that are active but not enabled.
+	FString ActiveTypeFilterString;
+
+	// Holds the list of enabled filters in this session.
+	FString EnabledTypeFilterString;
+
+	for (const TSharedRef<FEntityFilter>& Filter : this->EntityFilters)
+	{
+		const FString FilterName = Filter->GetFilterName();
+
+		if (ActiveTypeFilterString.Len() > 0)
+		{
+			ActiveTypeFilterString += TEXT(",");
+		}
+
+		ActiveTypeFilterString += FilterName;
+
+		if (Filter->IsEnabled())
+		{
+			if (EnabledTypeFilterString.Len() > 0)
+			{
+				EnabledTypeFilterString += TEXT(",");
+			}
+
+			EnabledTypeFilterString += FilterName;
+		}
+	}
+
+	GConfig->SetString(*IniSection, *(SettingsString + TEXT(".ActiveTypeFilters")), *ActiveTypeFilterString, IniFilename);
+	GConfig->SetString(*IniSection, *(SettingsString + TEXT(".EnabledTypeFilters")), *EnabledTypeFilterString, IniFilename);
+}
+
+void SRCPanelFilter::LoadSettings(const FString& IniFilename, const FString& IniSection, const FString& SettingsString)
+{
+	// Holds all the type filters that were found in the ActiveTypeFilters
+	FString ActiveTypeFilterString;
+
+	// Holds all the type filters that were found in the EnabledTypeFilters
+	FString EnabledTypeFilterString;
+	
+	GConfig->GetString(*IniSection, *(SettingsString + TEXT(".ActiveTypeFilters")), ActiveTypeFilterString, IniFilename);
+	GConfig->GetString(*IniSection, *(SettingsString + TEXT(".EnabledTypeFilters")), EnabledTypeFilterString, IniFilename);
+
+	// Parse comma delimited strings into arrays
+	TArray<FString> TypeFilterNames;
+	TArray<FString> EnabledTypeFilterNames;
+	constexpr bool bCullEmpty = true;
+	ActiveTypeFilterString.ParseIntoArray(TypeFilterNames, TEXT(","), bCullEmpty);
+	EnabledTypeFilterString.ParseIntoArray(EnabledTypeFilterNames, TEXT(","), bCullEmpty);
+
+	for (const TSharedRef<FCustomTypeFilterData>& CustomTypeFilter : CustomTypeFilters)
+	{
+		if (!this->IsEntityTypeInUse(CustomTypeFilter))
+		{
+			const FString FilterName = CustomTypeFilter->GetFilterName();
+			
+			if (TypeFilterNames.Contains(FilterName))
+			{
+				TSharedRef<FEntityFilter> NewFilter = AddEntityFilter(CustomTypeFilter);
+
+				if (EnabledTypeFilterNames.Contains(FilterName))
+				{
+					// NOTE: Do not trigger execution OnFilterChanged at this moment as it will cause panel refresh altogether.
+					NewFilter->SetEnabled(true, false);
+				}
+			}
+		}
+	}
+
+	OnFilterChanged.ExecuteIfBound();
+}
+
+#undef LOCTEXT_NAMESPACE
diff --git a/Engine/Plugins/VirtualProduction/RemoteControl/Source/RemoteControlUI/Private/UI/Filters/SRCPanelFilter.h b/Engine/Plugins/VirtualProduction/RemoteControl/Source/RemoteControlUI/Private/UI/Filters/SRCPanelFilter.h
new file mode 100644
--- /dev/null
+++ b/Engine/Plugins/VirtualProduction/RemoteControl/Source/RemoteControlUI/Private/UI/Filters/SRCPanelFilter.h
@@ -0,0 +1,43 @@
+// Copyright Epic Games, Inc. All Rights Reserved.
+
+#pragma once
+
+#include "CoreMinimal.h"
+
+#include "SRCFilterBar.h"
+
+struct SRCPanelTreeNode;
+
+using FEntityFilterType = TSharedRef<const SRCPanelTreeNode>;
+
+/**
+ * A custom list of filters currently applied to an Remote Control Panel.
+ */
+class SRCPanelFilter : public SRCFilterBar<FEntityFilterType>
+{
+public:
+
+	using FOnFilterChanged = typename SRCFilterBar<FEntityFilterType>::FOnFilterChanged;
+
+	SLATE_BEGIN_ARGS(SRCPanelFilter)
+	{}
+
+		/** Delegate for when filters have changed */
+		SLATE_EVENT(FOnFilterChanged, OnFilterChanged)
+
+	SLATE_END_ARGS()
+
+	/** Constructs this widget with InArgs */
+	void Construct(const FArguments& InArgs);
+	
+	/** Saves any settings to config that should be persistent between editor sessions */
+	void SaveSettings(const FString& IniFilename, const FString& IniSection, const FString& SettingsString) const;
+
+	/** Loads any settings to config that should be persistent between editor sessions */
+	void LoadSettings(const FString& IniFilename, const FString& IniSection, const FString& SettingsString);
+
+private:
+
+	/** Delegate for when filters have changed */
+	FOnFilterChanged OnFilterChanged;
+};
diff --git a/Engine/Plugins/VirtualProduction/RemoteControl/Source/RemoteControlUI/Private/UI/SRCPanelExposedEntitiesList.cpp b/Engine/Plugins/VirtualProduction/RemoteControl/Source/RemoteControlUI/Private/UI/SRCPanelExposedEntitiesList.cpp
--- a/Engine/Plugins/VirtualProduction/RemoteControl/Source/RemoteControlUI/Private/UI/SRCPanelExposedEntitiesList.cpp
+++ b/Engine/Plugins/VirtualProduction/RemoteControl/Source/RemoteControlUI/Private/UI/SRCPanelExposedEntitiesList.cpp
@@ -47,6 +47,12 @@ void SRCPanelExposedEntitiesList::Construct(const FArguments& InArgs, URemoteCon
 	ColumnSizeData.RightColumnWidth = TAttribute<float>(this, &SRCPanelExposedEntitiesList::OnGetRightColumnWidth);
 	ColumnSizeData.OnWidthChanged = SSplitter::FOnSlotResized::CreateSP(this, &SRCPanelExposedEntitiesList::OnSetColumnWidth);
 
+	bFilterApplicationRequested = false;
+	bSearchRequested = false;
+	SearchedText = MakeShared<FText>();
+	FilterCount = 0;
+	SearchCount = 0;
+
 	ChildSlot
 	[
 		SNew(SVerticalBox)
@@ -187,6 +193,27 @@ SRCPanelExposedEntitiesList::~SRCPanelExposedEntitiesList()
 	UnregisterEvents();
 }
 
+void SRCPanelExposedEntitiesList::Tick(const FGeometry& AllottedGeometry, const double InCurrentTime, const float InDeltaTime)
+{
+	if (bSearchRequested)
+	{
+		UE_LOG(LogTemp, Warning, TEXT("SRCPanelExposedEntitiesList::Tick -> Search Requested (Total Requests : %d)."), ++SearchCount);
+
+		TryRefreshingSearch(*SearchedText);
+
+		bSearchRequested = false;
+	}
+	
+	if (bFilterApplicationRequested)
+	{
+		UE_LOG(LogTemp, Warning, TEXT("SRCPanelExposedEntitiesList::Tick -> Filter Requested (Total Requests : %d)."), ++FilterCount);
+
+		ApplyFilters();
+
+		bFilterApplicationRequested = false;
+	}
+}
+
 TSharedPtr<SRCPanelTreeNode> SRCPanelExposedEntitiesList::GetSelectedGroup() const
 {
 	TArray<TSharedPtr<SRCPanelTreeNode>> SelectedNodes;
@@ -234,6 +261,13 @@ void SRCPanelExposedEntitiesList::SetSelection(const TSharedPtr<SRCPanelTreeNode
 	}
 }
 
+void SRCPanelExposedEntitiesList::SetBackendFilter(const FRCFilter& InBackendFilter)
+{
+	BackendFilter = InBackendFilter;
+
+	bFilterApplicationRequested = true;
+}
+
 void SRCPanelExposedEntitiesList::OnObjectPropertyChange(UObject* InObject, FPropertyChangedEvent& InChangeEvent)
 {
 	EPropertyChangeType::Type TypesNeedingRefresh = EPropertyChangeType::ArrayAdd | EPropertyChangeType::ArrayClear | EPropertyChangeType::ArrayRemove | EPropertyChangeType::ValueSet;
@@ -301,6 +335,56 @@ void SRCPanelExposedEntitiesList::Refresh()
 	}
 }
 
+void SRCPanelExposedEntitiesList::TryRefreshingSearch(const FText& InSearchText, bool bApplyFilter)
+{
+	FieldEntities.Reset();
+
+	if (FieldWidgetMap.IsEmpty() || !Preset.IsValid())
+	{
+		FieldsListView->RequestListRefresh();
+
+		*SearchedText = FText::GetEmpty();
+
+		return;
+	}
+
+	*SearchedText = InSearchText;
+
+	for (TWeakPtr<FRemoteControlEntity> WeakEntity : Preset->GetExposedEntities())
+	{
+		if (const TSharedPtr<FRemoteControlEntity> Entity = WeakEntity.Pin())
+		{
+			if (TSharedPtr<SRCPanelTreeNode> SelectedGroup = GetSelectedGroup())
+			{
+				if (FRemoteControlPresetGroup* EntityGroup = Preset->Layout.FindGroupFromField(Entity->GetId()))
+				{
+					const FString& EntityLabel = Entity->GetLabel().ToString();
+
+					if (EntityLabel.Contains(*InSearchText.ToString()) &&
+						(SelectedGroup->GetRCId() == EntityGroup->Id ||
+							Preset->Layout.IsDefaultGroup(SelectedGroup->GetRCId())))
+					{
+						if (TSharedPtr<SRCPanelTreeNode>* FoundNode = FieldWidgetMap.Find(Entity->GetId()))
+						{
+							(*FoundNode)->SetHighlightText(InSearchText);
+
+							FieldEntities.Add(*FoundNode);
+						}
+					}
+				}
+			}
+		}
+	}
+
+	// Only apply filter if it is requested otherwise skip this.
+	if (BackendFilter.HasAnyActiveFilters() && bApplyFilter)
+	{
+		bFilterApplicationRequested = true;
+	}
+
+	FieldsListView->RequestListRefresh();
+}
+
 void SRCPanelExposedEntitiesList::GenerateListWidgets()
 {
 	FieldWidgetMap.Reset();
@@ -442,6 +526,11 @@ void SRCPanelExposedEntitiesList::OnSelectionChanged(TSharedPtr<SRCPanelTreeNode
 		if (FRemoteControlPresetGroup* RCGroup = Preset->Layout.GetGroup(Node->GetRCId()))
 		{
 			GenerateListWidgets(*RCGroup);
+			
+			if (BackendFilter.HasAnyActiveFilters())
+			{
+				bFilterApplicationRequested = true;
+			}
 		}
 	}
 	else
@@ -729,6 +818,8 @@ void SRCPanelExposedEntitiesList::OnGroupAdded(const FRemoteControlPresetGroup&
 	GroupsListView->SetSelection(FieldGroup, ESelectInfo::OnMouseClick);
 	GroupsListView->ScrollToBottom();
 	GroupsListView->RequestListRefresh();
+
+	RequestSearchOrFilter();
 }
 
 void SRCPanelExposedEntitiesList::OnGroupDeleted(FRemoteControlPresetGroup DeletedGroup)
@@ -748,6 +839,8 @@ void SRCPanelExposedEntitiesList::OnGroupDeleted(FRemoteControlPresetGroup Delet
 		FieldGroups.RemoveAt(Index);
 		GroupsListView->RequestListRefresh();
 	}
+
+	RequestSearchOrFilter();
 }
 
 void SRCPanelExposedEntitiesList::OnGroupOrderChanged(const TArray<FGuid>& GroupIds)
@@ -780,11 +873,15 @@ void SRCPanelExposedEntitiesList::OnGroupRenamed(const FGuid& GroupId, FName New
 void SRCPanelExposedEntitiesList::OnFieldAdded(const FGuid& GroupId, const FGuid& FieldId, int32 FieldPosition)
 {
 	OnEntityAdded(FieldId);
+
+	RequestSearchOrFilter();
 }
 
 void SRCPanelExposedEntitiesList::OnFieldDeleted(const FGuid& GroupId, const FGuid& FieldId, int32 FieldPosition)
 {
 	OnEntityRemoved(GroupId, FieldId);
+
+	RequestSearchOrFilter();
 }
 
 void SRCPanelExposedEntitiesList::OnFieldOrderChanged(const FGuid& GroupId, const TArray<FGuid>& Fields)
@@ -825,7 +922,6 @@ void SRCPanelExposedEntitiesList::OnEntitiesUpdated(URemoteControlPreset*, const
 		}
 	}
 	
-
 	GEditor->GetTimerManager()->SetTimerForNextTick(FTimerDelegate::CreateLambda([WeakListPtr = TWeakPtr<SRCPanelExposedEntitiesList>(StaticCastSharedRef<SRCPanelExposedEntitiesList>(AsShared()))]()
 	{
 		if (TSharedPtr<SRCPanelExposedEntitiesList> ListPtr = WeakListPtr.Pin())
@@ -836,6 +932,65 @@ void SRCPanelExposedEntitiesList::OnEntitiesUpdated(URemoteControlPreset*, const
 	}));
 
 	OnEntityListUpdatedDelegate.ExecuteIfBound();
+
+	RequestSearchOrFilter();
+}
+
+void SRCPanelExposedEntitiesList::ApplyFilters()
+{
+	if (BackendFilter.HasAnyActiveFilters())
+	{
+		// If we are not actively searching anything then include all the entities.
+		if (SearchedText.IsValid() && SearchedText->IsEmpty())
+		{
+			for (TWeakPtr<FRemoteControlEntity> WeakEntity : Preset->GetExposedEntities())
+			{
+				if (const TSharedPtr<FRemoteControlEntity> Entity = WeakEntity.Pin())
+				{
+					if (TSharedPtr<SRCPanelTreeNode> SelectedGroup = GetSelectedGroup())
+					{
+						if (FRemoteControlPresetGroup* EntityGroup = Preset->Layout.FindGroupFromField(Entity->GetId()))
+						{
+							if (SelectedGroup->GetRCId() == EntityGroup->Id ||
+								Preset->Layout.IsDefaultGroup(SelectedGroup->GetRCId()))
+							{
+								if (TSharedPtr<SRCPanelTreeNode>* FoundNode = FieldWidgetMap.Find(Entity->GetId()))
+								{
+									if (BackendFilter.DoesPassFilters(FoundNode->ToSharedRef()))
+									{
+										FieldEntities.AddUnique(*FoundNode);
+									}
+								}
+							}
+						}
+					}
+				}
+			}
+		}
+		else
+		{
+			// Caution : Avoid filter during this phase of search as it would cause an endless loop between search and filters.
+			constexpr bool bApplyFilter = false;
+
+			TryRefreshingSearch(*SearchedText, bApplyFilter);
+		}
+	
+		// Apply the filter (always operate on the active list of entities).
+		FieldEntities.RemoveAll([&](TSharedPtr<const SRCPanelTreeNode> InEntity) { return !BackendFilter.DoesPassFilters(InEntity.ToSharedRef()); });
+
+		FieldsListView->RequestListRefresh();
+	}
+	else // Do a one time refresh in case all filters are cleared.
+	{
+		Refresh();
+	}
+}
+
+void SRCPanelExposedEntitiesList::RequestSearchOrFilter()
+{
+	bSearchRequested = SearchedText.IsValid() && !SearchedText->IsEmpty();
+
+	bFilterApplicationRequested = BackendFilter.HasAnyActiveFilters();
 }
 
 bool FGroupDragEvent::IsDraggedFromSameGroup() const
diff --git a/Engine/Plugins/VirtualProduction/RemoteControl/Source/RemoteControlUI/Private/UI/SRCPanelExposedEntitiesList.h b/Engine/Plugins/VirtualProduction/RemoteControl/Source/RemoteControlUI/Private/UI/SRCPanelExposedEntitiesList.h
--- a/Engine/Plugins/VirtualProduction/RemoteControl/Source/RemoteControlUI/Private/UI/SRCPanelExposedEntitiesList.h
+++ b/Engine/Plugins/VirtualProduction/RemoteControl/Source/RemoteControlUI/Private/UI/SRCPanelExposedEntitiesList.h
@@ -8,6 +8,7 @@
 #include "Widgets/DeclarativeSyntaxSupport.h"
 #include "Widgets/SCompoundWidget.h"
 #include "Widgets/Views/STreeView.h"
+#include "UI/Filters/RCFilter.h"
 #include "UObject/StrongObjectPtr.h"
 
 struct FRCPanelGroup;
@@ -59,6 +60,10 @@ public:
 
 	~SRCPanelExposedEntitiesList();
 
+	//~ Begin SWidget Interface
+	virtual void Tick(const FGeometry& AllottedGeometry, const double InCurrentTime, const float InDeltaTime) override;
+	//~ End SWidget Interface
+
 	/** Get the currently selected group. */
 	TSharedPtr<SRCPanelTreeNode> GetSelectedGroup() const;
 	
@@ -68,9 +73,18 @@ public:
 	/** Set the currently selected group or exposed entity. */
 	void SetSelection(const TSharedPtr<SRCPanelTreeNode>& Node, const bool bForceMouseClick = false);
 
+	/** Notifies the entities list view that the filter-list filter has changed */
+	void SetBackendFilter(const FRCFilter& InBackendFilter);
+
 	/** Recreate everything in the panel. */
 	void Refresh();
 	
+	/** Tries to refresh the list according to the provided search text. */
+	void TryRefreshingSearch(const FText& InSearchText, bool bApplyFilter = true);
+
+	/** Notifies us that the search has ended. */
+	void ResetSearch() { *SearchedText = FText::GetEmpty(); RequestSearchOrFilter(); };
+
 	/** Returns delegate called on selection change. */
 	FOnSelectionChange& OnSelectionChange() { return OnSelectionChangeDelegate; }
 
@@ -131,6 +145,12 @@ private:
 	void OnFieldOrderChanged(const FGuid& GroupId, const TArray<FGuid>& Fields);
 	void OnEntitiesUpdated(URemoteControlPreset*, const TSet<FGuid>& UpdatedEntities);
 
+	/** Requests that the entities view refreshes all it's items. */
+	void ApplyFilters();
+
+	/** Notifies us that we need to do a search or filter in the next tick. */
+	void RequestSearchOrFilter();
+
 private:
 	/** Holds the Groups list view. */
 	TSharedPtr<SListView<TSharedPtr<SRCPanelTreeNode>>> GroupsListView;
@@ -158,4 +178,14 @@ private:
 	FSimpleDelegate OnEntityListUpdatedDelegate;
 	/** Holds the cache of widgets to be used by this list's entities. */
 	TWeakPtr<FRCPanelWidgetRegistry> WidgetRegistry;
+	/** Holds the cache of filters to be used by this list's entities. */
+	FRCFilter BackendFilter;
+	/** If true, the entity items will be refreshed next frame. */
+	bool bFilterApplicationRequested;
+	/** If true, the entity items will be refreshed next frame. */
+	bool bSearchRequested;
+	/** Holds the text that is being searched actively. */
+	TSharedPtr<FText> SearchedText;
+	int32 SearchCount;
+	int32 FilterCount;
 };
\ No newline at end of file
diff --git a/Engine/Plugins/VirtualProduction/RemoteControl/Source/RemoteControlUI/Private/UI/SRCPanelExposedEntity.cpp b/Engine/Plugins/VirtualProduction/RemoteControl/Source/RemoteControlUI/Private/UI/SRCPanelExposedEntity.cpp
--- a/Engine/Plugins/VirtualProduction/RemoteControl/Source/RemoteControlUI/Private/UI/SRCPanelExposedEntity.cpp
+++ b/Engine/Plugins/VirtualProduction/RemoteControl/Source/RemoteControlUI/Private/UI/SRCPanelExposedEntity.cpp
@@ -335,6 +335,7 @@ TSharedRef<SWidget> SRCPanelExposedEntity::CreateEntityWidget(TSharedPtr<SWidget
 			.OnTextCommitted_Raw(this, &SRCPanelExposedEntity::OnLabelCommitted)
 			.OnVerifyTextChanged_Raw(this, &SRCPanelExposedEntity::OnVerifyItemLabelChanged)
 			.IsReadOnly_Lambda([this]() { return !bEditMode.Get(); })
+			.HighlightText_Lambda([this]() { return HighlightText.Get().ToString().Len() > 3 ? HighlightText.Get() : FText::GetEmpty(); })
 		];
 
 	Args.RenameButton = SNullWidget::NullWidget;
diff --git a/Engine/Plugins/VirtualProduction/RemoteControl/Source/RemoteControlUI/Private/UI/SRCPanelExposedEntity.h b/Engine/Plugins/VirtualProduction/RemoteControl/Source/RemoteControlUI/Private/UI/SRCPanelExposedEntity.h
--- a/Engine/Plugins/VirtualProduction/RemoteControl/Source/RemoteControlUI/Private/UI/SRCPanelExposedEntity.h
+++ b/Engine/Plugins/VirtualProduction/RemoteControl/Source/RemoteControlUI/Private/UI/SRCPanelExposedEntity.h
@@ -28,6 +28,8 @@ struct SRCPanelExposedEntity : public SRCPanelTreeNode
 	virtual FGuid GetRCId() const override final { return EntityId; }
 	/** Make the group name's text box editable. */
 	virtual void EnterRenameMode() override;
+	/** Updates the highlight text to active search term. */
+	virtual void SetHighlightText(const FText& InHightlightText = FText::GetEmpty()) override { HighlightText = InHightlightText; }
 
 protected:
 	void Initialize(const FGuid& InEntityId, URemoteControlPreset* InPreset, const TAttribute<bool>& InbEditMode);
@@ -50,6 +52,8 @@ protected:
 	TAttribute<bool> bEditMode;
 	/** Display name of the entity. */
 	FName CachedLabel;
+	/** Text to be highlighted while searching. */
+	TAttribute<FText> HighlightText;
 
 private:
 	/** Handles changing the object this entity is bound to upon selecting an actor in the rebinding dropdown. */
diff --git a/Engine/Plugins/VirtualProduction/RemoteControl/Source/RemoteControlUI/Private/UI/SRCPanelExposedField.cpp b/Engine/Plugins/VirtualProduction/RemoteControl/Source/RemoteControlUI/Private/UI/SRCPanelExposedField.cpp
--- a/Engine/Plugins/VirtualProduction/RemoteControl/Source/RemoteControlUI/Private/UI/SRCPanelExposedField.cpp
+++ b/Engine/Plugins/VirtualProduction/RemoteControl/Source/RemoteControlUI/Private/UI/SRCPanelExposedField.cpp
@@ -59,7 +59,7 @@ namespace ExposedFieldUtils
 
 TSharedPtr<SRCPanelTreeNode> SRCPanelExposedField::MakeInstance(const FGenerateWidgetArgs& Args)
 {
-	return SNew(SRCPanelExposedField, StaticCastSharedPtr<FRemoteControlField>(Args.Entity), Args.ColumnSizeData, Args.WidgetRegistry).Preset(Args.Preset).EditMode(Args.bIsInEditMode);
+	return SNew(SRCPanelExposedField, StaticCastSharedPtr<FRemoteControlField>(Args.Entity), Args.ColumnSizeData, Args.WidgetRegistry).Preset(Args.Preset).EditMode(Args.bIsInEditMode).HighlightText(Args.HighlightText);
 }
 
 void SRCPanelExposedField::Construct(const FArguments& InArgs, TWeakPtr<FRemoteControlField> InField, FRCColumnSizeData InColumnSizeData, TWeakPtr<FRCPanelWidgetRegistry> InWidgetRegistry)
@@ -69,6 +69,8 @@ void SRCPanelExposedField::Construct(const FArguments& InArgs, TWeakPtr<FRemoteC
 	ColumnSizeData = MoveTemp(InColumnSizeData);
 	WidgetRegistry = MoveTemp(InWidgetRegistry);
 	
+	HighlightText = InArgs._HighlightText;
+
 	if (TSharedPtr<FRemoteControlField> FieldPtr = WeakField.Pin())
 	{
 		Initialize(FieldPtr->GetId(), InArgs._Preset.Get(), InArgs._EditMode);
diff --git a/Engine/Plugins/VirtualProduction/RemoteControl/Source/RemoteControlUI/Private/UI/SRCPanelExposedField.h b/Engine/Plugins/VirtualProduction/RemoteControl/Source/RemoteControlUI/Private/UI/SRCPanelExposedField.h
--- a/Engine/Plugins/VirtualProduction/RemoteControl/Source/RemoteControlUI/Private/UI/SRCPanelExposedField.h
+++ b/Engine/Plugins/VirtualProduction/RemoteControl/Source/RemoteControlUI/Private/UI/SRCPanelExposedField.h
@@ -28,9 +28,11 @@ struct SRCPanelExposedField : public SRCPanelExposedEntity
 {
 	SLATE_BEGIN_ARGS(SRCPanelExposedField)
 		: _EditMode(true)
+		, _HighlightText()
 		, _Preset(nullptr)
 	{}
 		SLATE_ATTRIBUTE(bool, EditMode)
+		SLATE_ATTRIBUTE(FText, HighlightText)
 		SLATE_ATTRIBUTE(URemoteControlPreset*, Preset)
 	SLATE_END_ARGS()
 
diff --git a/Engine/Plugins/VirtualProduction/RemoteControl/Source/RemoteControlUI/Private/UI/SRCPanelTreeNode.h b/Engine/Plugins/VirtualProduction/RemoteControl/Source/RemoteControlUI/Private/UI/SRCPanelTreeNode.h
--- a/Engine/Plugins/VirtualProduction/RemoteControl/Source/RemoteControlUI/Private/UI/SRCPanelTreeNode.h
+++ b/Engine/Plugins/VirtualProduction/RemoteControl/Source/RemoteControlUI/Private/UI/SRCPanelTreeNode.h
@@ -43,6 +43,8 @@ struct SRCPanelTreeNode : public SCompoundWidget
 	virtual void SetIsHovered(bool bIsBeingHovered) {}
 	/** Make the node name's text box editable. */
 	virtual void EnterRenameMode() {};
+	/** Updates the highlight text to active search term. */
+	virtual void SetHighlightText(const FText& InHightlightText = FText::GetEmpty()) {};
 
 	virtual FReply OnMouseButtonDoubleClick(const FGeometry& InMyGeometry, const FPointerEvent& InMouseEvent) override { if (InMouseEvent.GetEffectingButton() == EKeys::LeftMouseButton) { EnterRenameMode(); } return SCompoundWidget::OnMouseButtonDoubleClick(InMyGeometry, InMouseEvent); };
 
diff --git a/Engine/Plugins/VirtualProduction/RemoteControl/Source/RemoteControlUI/Private/UI/SRemoteControlPanel.cpp b/Engine/Plugins/VirtualProduction/RemoteControl/Source/RemoteControlUI/Private/UI/SRemoteControlPanel.cpp
--- a/Engine/Plugins/VirtualProduction/RemoteControl/Source/RemoteControlUI/Private/UI/SRemoteControlPanel.cpp
+++ b/Engine/Plugins/VirtualProduction/RemoteControl/Source/RemoteControlUI/Private/UI/SRemoteControlPanel.cpp
@@ -59,6 +59,7 @@
 #include "ToolMenus.h"
 #include "UI/BaseLogicUI/RCLogicModeBase.h"
 #include "UI/BaseLogicUI/SRCLogicPanelListBase.h"
+#include "UI/Filters/SRCPanelFilter.h"
 #include "Widgets/DeclarativeSyntaxSupport.h"
 #include "Widgets/Docking/SDockTab.h"
 #include "Widgets/Images/SLayeredImage.h"
@@ -298,6 +299,21 @@ void SRemoteControlPanel::Construct(const FArguments& InArgs, URemoteControlPres
 
 	UpdateRebindButtonVisibility();
 
+	// Setup search filter.
+	SearchTextFilter = MakeShared<TTextFilter<const SRCPanelTreeNode&>>(TTextFilter<const SRCPanelTreeNode&>::FItemToStringArray::CreateSP(this, &SRemoteControlPanel::PopulateSearchStrings));
+	SearchedText = MakeShared<FText>(FText::GetEmpty());
+
+	// Create the Filter Widget
+	FilterPtr = SNew(SRCPanelFilter)
+		.OnFilterChanged(this, &SRemoteControlPanel::OnFilterChanged)
+		.AddMetaData<FTagMetaData>(FTagMetaData(TEXT("RemoteControlFilters")));
+
+	// Create the Filter Combo Button
+	TSharedPtr<SWidget> FilterComboButton = SRCPanelFilter::MakeAddFilterButton(FilterPtr.ToSharedRef());
+
+	TSharedPtr<ISlateMetaData> FilterComboButtonMetaData = MakeShared<FTagMetaData>(TEXT("ContentBrowserFiltersCombo"));
+	FilterComboButton->AddMetadata(FilterComboButtonMetaData.ToSharedRef());
+
 	TSharedPtr<SHorizontalBox> TopExtensions;
 
 	EntityProtocolDetails = SNew(SBox);
@@ -359,6 +375,10 @@ void SRemoteControlPanel::Construct(const FArguments& InArgs, URemoteControlPres
 			.FillWidth(1.f)
 			[
 				SAssignNew(SearchBoxPtr, SSearchBox)
+				.HintText(LOCTEXT("SearchHint", "Search"))
+				.OnTextChanged(this, &SRemoteControlPanel::OnSearchTextChanged)
+				.OnTextCommitted(this, &SRemoteControlPanel::OnSearchTextCommitted)
+				.DelayChangeNotificationsWhileTyping(true)
 			]
 
 			// Filters
@@ -367,21 +387,9 @@ void SRemoteControlPanel::Construct(const FArguments& InArgs, URemoteControlPres
 			.AutoWidth()
 			.Padding(5.f, 3.f, 3.f, 3.f)
 			[
-				SNew(SComboButton)
-				.ComboButtonStyle(&FAppStyle::Get().GetWidgetStyle<FComboButtonStyle>("SimpleComboButtonWithIcon"))
-				.ForegroundColor(FSlateColor::UseStyle())
-				.ToolTipText(LOCTEXT("Filters_Tooltip", "Filter options for the RC Preset."))
-				.HasDownArrow(true)
-				.AddMetaData<FTagMetaData>(FTagMetaData(TEXT("RemoteControlFiltersCombo")))
-				.ContentPadding(FMargin(1.f, 0.f))
-				.ButtonContent()
-				[
-					SNew(SImage)
-					.Image(FAppStyle::Get().GetBrush("Icons.Filter"))
-					.ColorAndOpacity(FSlateColor::UseForeground())
-				]
+				FilterComboButton.ToSharedRef()
 			]
-			
+
 			// Settings
 			+ SHorizontalBox::Slot()
 			.VAlign(VAlign_Center)
@@ -750,10 +758,13 @@ void SRemoteControlPanel::Construct(const FArguments& InArgs, URemoteControlPres
 	RegisterEvents();
 	CacheLevelClasses();
 	Refresh();
+	LoadSettings(InPreset->GetPresetId());
 }
 
 SRemoteControlPanel::~SRemoteControlPanel()
 {
+	SaveSettings();
+
 	UnbindRemoteControlCommands();
 	UnregisterEvents();
 
@@ -1580,6 +1591,12 @@ void SRemoteControlPanel::UpdateEntityDetailsView(const TSharedPtr<SRCPanelTreeN
 			}
 		}
 	}
+
+	// Trigger search to list the search results specific to selected group.
+	if (SearchedText.IsValid() && !SearchedText->IsEmptyOrWhitespace() && SearchedText->ToString().Len() > 3)
+	{
+		OnSearchTextChanged(*SearchedText);
+	}
 }
 
 void SRemoteControlPanel::UpdateRebindButtonVisibility()
@@ -2024,4 +2041,79 @@ bool SRemoteControlPanel::CanRenameEntity() const
 	return false;
 }
 
+void SRemoteControlPanel::OnSearchTextChanged(const FText& InFilterText)
+{
+	SearchTextFilter->SetRawFilterText(InFilterText);
+	SearchBoxPtr->SetError(SearchTextFilter->GetFilterErrorText());
+	*SearchedText = InFilterText;
+
+	const int32 Length = InFilterText.ToString().Len();
+
+	check(EntityList.IsValid());
+
+	if (Length > 3)
+	{
+		EntityList->TryRefreshingSearch(InFilterText);
+	}
+	else if (Length == 3 || Length == 0) // Avoid unnecessary refresh if search text is below the threshold.
+	{
+		EntityList->ResetSearch();
+
+		EntityList->Refresh();
+	}
+}
+
+void SRemoteControlPanel::OnSearchTextCommitted(const FText& InFilterText, ETextCommit::Type InCommitType)
+{
+	if (InCommitType == ETextCommit::OnCleared || InFilterText.IsEmpty())
+	{
+		EntityList->ResetSearch();
+
+		EntityList->Refresh();
+
+		return;
+	}
+
+	OnSearchTextChanged(InFilterText);
+}
+
+void SRemoteControlPanel::PopulateSearchStrings(const SRCPanelTreeNode& Item, TArray<FString>& OutSearchStrings) const
+{
+	if (Preset.IsValid())
+	{
+		if (TSharedPtr<FRemoteControlEntity> Entity = Preset->GetExposedEntity<FRemoteControlEntity>(Item.GetRCId()).Pin())
+		{
+			OutSearchStrings.Add(Entity->GetLabel().ToString());
+		}
+	}
+}
+
+void SRemoteControlPanel::OnFilterChanged()
+{
+	check(FilterPtr.IsValid());
+
+	FRCFilter Filter = FilterPtr->GetCombinedBackendFilter();
+
+	EntityList->SetBackendFilter(Filter);
+}
+
+void SRemoteControlPanel::LoadSettings(const FGuid& InInstanceId)
+{
+	const FString SettingsString = InInstanceId.ToString();
+
+	// Load all our data using the settings string as a key in the user settings ini.
+	FilterPtr->LoadSettings(GEditorPerProjectIni, IRemoteControlUIModule::SettingsIniSection, SettingsString);
+}
+
+void SRemoteControlPanel::SaveSettings()
+{
+	if (Preset.IsValid() && FilterPtr.IsValid())
+	{
+		const FString SettingsString = Preset->GetPresetId().ToString();
+
+		// Save all our data using the settings string as a key in the user settings ini.
+		FilterPtr->SaveSettings(GEditorPerProjectIni, IRemoteControlUIModule::SettingsIniSection, SettingsString);
+	}
+}
+
 #undef LOCTEXT_NAMESPACE /*RemoteControlPanel*/
diff --git a/Engine/Plugins/VirtualProduction/RemoteControl/Source/RemoteControlUI/Private/UI/SRemoteControlPanel.h b/Engine/Plugins/VirtualProduction/RemoteControl/Source/RemoteControlUI/Private/UI/SRemoteControlPanel.h
--- a/Engine/Plugins/VirtualProduction/RemoteControl/Source/RemoteControlUI/Private/UI/SRemoteControlPanel.h
+++ b/Engine/Plugins/VirtualProduction/RemoteControl/Source/RemoteControlUI/Private/UI/SRemoteControlPanel.h
@@ -4,7 +4,9 @@
 #include "Engine/EngineTypes.h"
 #include "EditorUndoClient.h"
 #include "Engine/EngineTypes.h"
+#include "Filters/SFilterBar.h"
 #include "IDetailTreeNode.h"
+#include "Misc/TextFilter.h"
 #include "RemoteControlUIModule.h"
 #include "Styling/SlateTypes.h"
 #include "RemoteControlFieldPath.h"
@@ -30,6 +32,7 @@ struct SRCPanelTreeNode;
 class SRCPanelFunctionPicker;
 class SRemoteControlPanel;
 class SRCPanelExposedEntitiesList;
+class SRCPanelFilter;
 class SSearchBox;
 class STextBlock;
 class URemoteControlPreset;
@@ -40,7 +43,6 @@ class URCAction;
 class FRCControllerModel;
 class FRCBehaviourModel;
 
-
 DECLARE_DELEGATE_TwoParams(FOnEditModeChange, TSharedPtr<SRemoteControlPanel> /* Panel */, bool /* bEditModeChange */);
 
 /**
@@ -277,6 +279,20 @@ private:
 
 	/** Called to test if user is able to rename a group/exposed entity. */
 	bool CanRenameEntity() const;
+	
+	// Exposed Entities filtering. (Filters the Exposed Entities view)
+	void OnSearchTextChanged(const FText& InFilterText);
+	void OnSearchTextCommitted(const FText& InFilterText, ETextCommit::Type InCommitType);
+	void PopulateSearchStrings(const SRCPanelTreeNode& Item, TArray<FString>& OutSearchStrings) const;
+
+	/** Handler for when a filter in the filter list has changed */
+	void OnFilterChanged();
+
+	/** Loads settings from config based on the preset identifier. */
+	void LoadSettings(const FGuid& InInstanceId);
+	
+	/** Saves settings from config based on the preset identifier. */
+	void SaveSettings();
 
 private:
 	static const FName DefaultRemoteControlPanelToolBarName;
@@ -333,6 +349,12 @@ private:
 	TSharedPtr<SSearchBox> SearchBoxPtr;
 	/** Holds a shared pointer reference to the active entity that is selected. */
 	TSharedPtr<SRCPanelTreeNode> SelectedEntity;
+	/** Text filter for the search text. */
+	TSharedPtr<TTextFilter<const SRCPanelTreeNode&>> SearchTextFilter;
+	/** The filter list */
+	TSharedPtr<SRCPanelFilter> FilterPtr;
+	/** Actively serached term. */
+	TSharedPtr<FText> SearchedText;
 
 	// ~ Remote Control Logic Panels ~
 
diff --git a/Engine/Plugins/VirtualProduction/RemoteControl/Source/RemoteControlUI/Public/IRemoteControlUIModule.h b/Engine/Plugins/VirtualProduction/RemoteControl/Source/RemoteControlUI/Public/IRemoteControlUIModule.h
--- a/Engine/Plugins/VirtualProduction/RemoteControl/Source/RemoteControlUI/Public/IRemoteControlUIModule.h
+++ b/Engine/Plugins/VirtualProduction/RemoteControl/Source/RemoteControlUI/Public/IRemoteControlUIModule.h
@@ -34,6 +34,7 @@ struct FGenerateWidgetArgs
 	TSharedPtr<FRemoteControlEntity> Entity;
 	TAttribute<bool> bIsInEditMode;
 	TWeakPtr<FRCPanelWidgetRegistry> WidgetRegistry;
+	FText HighlightText;
 };
 
 /**
@@ -184,4 +185,7 @@ public:
 	 * Unregister a previously registered widget factory.
 	 */
 	virtual void UnregisterWidgetFactoryForType(UScriptStruct* RemoteControlEntityType) = 0;
+
+	/** The section of EditorPerProjectUserSettings in which to save filter settings */
+	static const FString SettingsIniSection;
 };
diff --git a/Engine/Plugins/VirtualProduction/RemoteControl/Source/RemoteControlUI/Public/UI/Filters/CustomRCFilterData.h b/Engine/Plugins/VirtualProduction/RemoteControl/Source/RemoteControlUI/Public/UI/Filters/CustomRCFilterData.h
new file mode 100644
--- /dev/null
+++ b/Engine/Plugins/VirtualProduction/RemoteControl/Source/RemoteControlUI/Public/UI/Filters/CustomRCFilterData.h
@@ -0,0 +1,147 @@
+// Copyright Epic Games, Inc. All Rights Reserved.
+
+#pragma once
+
+#include "CoreMinimal.h"
+#include "RCFilter.h"
+#include "UObject/UnrealType.h"
+
+class FFieldClass;
+class FFilterCategory;
+struct FSlateIcon;
+class IAssetTypeActions;
+class UClass;
+
+/**
+ * A custom view model to represent each type of filters we have.
+ */
+class FCustomTypeFilterData
+{
+public:
+	
+	FCustomTypeFilterData(EEntityTypeCategories::Type InEntityTypeCategory, FFieldClass* InEntityType, const FName& InCustomTypeName = NAME_None)
+		: Class(nullptr)
+		, EntityType(InEntityType)
+		, EntityTypeCategory(InEntityTypeCategory)
+		, Color(FLinearColor::White)
+		, CustomTypeName(InCustomTypeName)
+		, AssetTypeActions(nullptr)
+	{
+	};
+	
+	FCustomTypeFilterData(EEntityTypeCategories::Type InEntityTypeCategory, FFieldClass* InEntityType, TSharedPtr<FFilterCategory> InCategory, const FName& InCustomTypeName = NAME_None)
+		: Class(nullptr)
+		, EntityType(InEntityType)
+		, EntityTypeCategory(InEntityTypeCategory)
+		, Color(FLinearColor::White)
+		, CustomTypeName(InCustomTypeName)
+		, AssetTypeActions(nullptr)
+	{
+		Categories.Add(InCategory);
+	};
+
+	/* You can provide an IAssetTypeActions for your type to get all the information from that */
+	FCustomTypeFilterData(TWeakPtr<IAssetTypeActions> InAssetTypeAction)
+		: Class(nullptr)
+		, EntityType(FObjectProperty::StaticClass())
+		, EntityTypeCategory(EEntityTypeCategories::Assets)
+		, Color(FLinearColor::White)
+		, CustomTypeName(NAME_None)
+		, AssetTypeActions(InAssetTypeAction)
+	{
+	}
+
+	/* Or you can provide a UClass, a Category and a Color to identify the filter */
+	FCustomTypeFilterData(UClass* InClass, TSharedPtr<FFilterCategory> InCategory, FLinearColor InColor)
+		: Class(InClass)
+		, EntityType(FObjectProperty::StaticClass())
+		, EntityTypeCategory(EEntityTypeCategories::Assets)
+		, Color(InColor)
+		, CustomTypeName(NAME_None)
+		, AssetTypeActions(nullptr)
+	{
+		Categories.Add(InCategory);
+	}
+
+	/* Unlike normal filters, Type Filters are allowed to belong to multiple categories */
+	void AddCategory(TSharedPtr<FFilterCategory> InCategory);
+
+	/** Add this filter to the input Backend Filter */
+	void BuildBackendFilter(FRCFilter& OutFilter);
+
+	/* Get the AssetTypeActions associated with this filter if they exist */
+	TSharedPtr<IAssetTypeActions> GetAssetTypeActions() const;
+
+	/* Get the Catgories this filter belongs to */
+	TArray<TSharedPtr<FFilterCategory>> GetCategories() const;
+
+	/* Get the UClass associated with this filter */
+	UClass* GetClass() const;
+
+	/* Get the Property Type associated with this filter */
+	FFieldClass* GetEntityType() const;
+
+	/* Get the Entity Type Catgories this filter belongs to */
+	EEntityTypeCategories::Type GetEntityTypeCategory() const;
+
+	/* Get the color of this filter */
+	FLinearColor GetColor() const;
+
+	/** Get the logical name of this filter class */
+	FText GetName() const;
+
+	/** Get the display name of this filter class */
+	FString GetFilterName() const;
+
+	/** Retrieves the appropriate icon for the filter. */
+	FSlateIcon GetSlateIcon() const;
+
+	/** Remove this filter from the input Backend Filter */
+	void ResetBackendFilter(FRCFilter& OutFilter);
+
+protected:
+
+	/** Returns proper name of the entity type to be displayed. */
+	FText GetProperEntityDisplayName() const;
+	
+	/** Returns proper icon of the entity type to be displayed. */
+	FSlateIcon GetProperEntityIcon() const;
+
+private:
+
+	/** The Class the filter is associated with (if it does not have an AssetTypeAction) */
+	UClass* Class;
+
+	/** The Entity Type the filter is associated with. */
+	FFieldClass* EntityType;
+
+	/** The Entity Type the filter is associated with. */
+	EEntityTypeCategories::Type EntityTypeCategory;
+
+	/** The Categories the filter will show up under */
+	TArray<TSharedPtr<FFilterCategory>> Categories;
+
+	/** The Color of the filter. */
+	FLinearColor Color;
+
+	/** If property name could not be deducted from entity type then this will be used instead. */
+	FName CustomTypeName;
+
+	/** An optional AssetTypeAction to get information about the filter class */
+	TWeakPtr<IAssetTypeActions> AssetTypeActions;
+
+	/** Empty Slate Icon. */
+	static FSlateIcon EmptyIcon;
+	
+	/** Entity Slate Icon. */
+	static FSlateIcon EntityIcon;
+
+	/** Array Slate Icon. */
+	static FSlateIcon ArrayIcon;
+
+	/** Map Slate Icon. */
+	static FSlateIcon MapIcon;
+	
+	/** Set Slate Icon. */
+	static FSlateIcon SetIcon;
+};
diff --git a/Engine/Plugins/VirtualProduction/RemoteControl/Source/RemoteControlUI/Public/UI/Filters/RCFilter.h b/Engine/Plugins/VirtualProduction/RemoteControl/Source/RemoteControlUI/Public/UI/Filters/RCFilter.h
new file mode 100644
--- /dev/null
+++ b/Engine/Plugins/VirtualProduction/RemoteControl/Source/RemoteControlUI/Public/UI/Filters/RCFilter.h
@@ -0,0 +1,89 @@
+// Copyright Epic Games, Inc. All Rights Reserved.
+
+#pragma once
+
+#include "CoreMinimal.h"
+
+class FFieldClass;
+struct SRCPanelTreeNode;
+
+using FEntityFilterType = TSharedRef<const SRCPanelTreeNode>;
+
+namespace RemoteControlTypes
+{
+	static const FName NAME_RCActors = "Exposed Actors";
+
+	static const FName NAME_RCFunctions = "Exposed Functions";
+}
+
+/**
+ * The list of categories for Entity Types.
+ */
+namespace EEntityTypeCategories
+{
+	enum Type
+	{
+		None = 0,
+
+		Core = 1 << 0,
+
+		Strings = 1 << 1,
+
+		Structs = 1 << 2,
+
+		Containers = 1 << 3,
+
+		Objects = 1 << 4,
+
+		References = 1 << 5,
+
+		UserDefined = 1 << 6,
+
+		Assets = 1 << 7,
+
+		Primary = 1 << 8,
+	};
+}
+
+/**
+ * A data model to represent combined backend filter.
+ */
+class FRCFilter
+{
+public:
+
+	/** Adds an entity type to the filters list. */
+	void AddTypeFilter(FFieldClass* InEntityTypeFilter);
+
+	/** Adds a custom type to the filters list. */
+	void AddTypeFilter(const FName InCustomTypeFilter);
+
+	/** Adds an asset type to the filters list. */
+	void AddTypeFilter(UClass* InAssetTypeFilter);
+	
+	/** Returns true if the given item passes atleast one filter. */
+	bool DoesPassFilters(FEntityFilterType InEntityItem);
+	
+	/** Returns true if atleast one of the type filters are populated. */
+	bool HasAnyActiveFilters();
+
+	/** Removes an entity type from the filters list. */
+	void RemoveTypeFilter(FFieldClass* InEntityTypeFilter);
+
+	/** Removes a custom type from the filters list. */
+	void RemoveTypeFilter(const FName& InCustomTypeFilter);
+	
+	/** Removes an asset type from the filters list. */
+	void RemoveTypeFilter(UClass* InAssetTypeFilter);
+
+private:
+
+	/** The list of active entity types. */
+	TSet<FFieldClass*> EntityTypeFilters;
+	
+	/** The list of active asset types. */
+	TSet<UClass*> AssetTypeFilters;
+
+	/** The list of active custom types. */
+	TSet<FName> CustomTypeFilters;
+};
