Add Missing Reset Arrow Button for Preset Properties

diff --git a/Engine/Plugins/VirtualProduction/RemoteControl/Source/RemoteControlUI/Private/UI/SRCPanelExposedField.cpp b/Engine/Plugins/VirtualProduction/RemoteControl/Source/RemoteControlUI/Private/UI/SRCPanelExposedField.cpp
--- a/Engine/Plugins/VirtualProduction/RemoteControl/Source/RemoteControlUI/Private/UI/SRCPanelExposedField.cpp
+++ b/Engine/Plugins/VirtualProduction/RemoteControl/Source/RemoteControlUI/Private/UI/SRCPanelExposedField.cpp
@@ -4,21 +4,23 @@
 
 #include "Algo/Transform.h"
 #include "EditorFontGlyphs.h"
-#include "RCPanelWidgetRegistry.h"
 #include "IDetailTreeNode.h"
 #include "IRemoteControlUIModule.h"
 #include "Layout/Visibility.h"
+#include "Materials/MaterialInstanceDynamic.h"
+#include "RCPanelWidgetRegistry.h"
 #include "RemoteControlField.h"
 #include "RemoteControlPanelStyle.h"
 #include "RemoteControlPreset.h"
 #include "RemoteControlSettings.h"
-#include "Styling/AppStyle.h"
-#include "ScopedTransaction.h"
 #include "SRCPanelTreeNode.h"
+#include "SResetToDefaultPropertyEditor.h"
+#include "ScopedTransaction.h"
+#include "Styling/AppStyle.h"
 #include "Styling/SlateBrush.h"
 #include "UObject/Object.h"
-#include "Widgets/Layout/SBorder.h"
 #include "Widgets/Input/SButton.h"
+#include "Widgets/Layout/SBorder.h"
 #include "Widgets/Text/SInlineEditableTextBlock.h"
 #include "Widgets/Text/STextBlock.h"
 
@@ -165,7 +167,69 @@ TSharedRef<SWidget> SRCPanelExposedField::ConstructWidget()
 						ChildWidgets.Add(SNew(SRCPanelFieldChildNode, ChildNode, ColumnSizeData));
 					}
 
-					return MakeFieldWidget(ExposedFieldUtils::CreateNodeValueWidget(MoveTemp(Node)));
+					TSharedPtr<SHorizontalBox> FieldWidget = SNew(SHorizontalBox);
+					FieldWidget->AddSlot()
+					[
+						ExposedFieldUtils::CreateNodeValueWidget(MoveTemp(Node))
+					];
+
+					
+					TSharedPtr<IPropertyHandle> PropertyHandle = Node->CreatePropertyHandle();
+					if (PropertyHandle->GetParentHandle()->GetPropertyDisplayName().ToString() == FString("Transform")
+						|| PropertyHandle->GetOuterBaseClass() == UMaterialInstanceDynamic::StaticClass())
+					{
+						// Set up a Zeroed DefaultValue, in case an ExposedEntity doesn't have a native Default. Needed for certain ResetToDefault cases.
+						void* ValuePtr;
+						Node->CreatePropertyHandle()->GetValueData(ValuePtr);
+						DefaultValue = MakeShared<uint8>(Node->CreatePropertyHandle()->GetProperty()->GetSize());
+						Node->CreatePropertyHandle()->GetProperty()->CopyCompleteValue(DefaultValue.Get(), ValuePtr);
+						Node->CreatePropertyHandle()->GetProperty()->ClearValue(DefaultValue.Get());
+						
+						auto IsVisible = [this, Node]()
+						{
+							TSharedPtr<IPropertyHandle> PropertyHandle = Node->CreatePropertyHandle();
+							void* DataPtr;
+							PropertyHandle->GetValueData(DataPtr);
+							FProperty* NodeProperty = PropertyHandle->GetProperty();
+							
+							bool bVisible = !NodeProperty->Identical(DefaultValue.Get(), DataPtr);
+							return bVisible ? EVisibility::Visible : EVisibility::Hidden;
+						};
+						
+						FieldWidget->AddSlot()
+						.AutoWidth()
+						.VAlign(VAlign_Center)
+						.Padding(FMargin(1.f,2.f))
+						[
+							SNew(SButton)
+							.IsFocusable(false)
+							.ButtonStyle(FAppStyle::Get(), "SimpleButton")
+							.ContentPadding(0) 
+							.Visibility_Lambda(IsVisible)
+							.OnClicked_Lambda([Node]()
+							{
+								Node->CreatePropertyHandle()->ResetToDefault();
+								return FReply::Handled();
+							})
+							.Content()
+							[
+								SNew(SImage)
+								.Image(FAppStyle::GetBrush("PropertyWindow.DiffersFromDefault"))
+								.ColorAndOpacity(FSlateColor::UseForeground())
+							]
+						];
+					}
+					else
+					{
+						FieldWidget->AddSlot()
+						.AutoWidth()
+						.VAlign(VAlign_Center)
+						.Padding(FMargin(1.f,2.f))
+						[
+							SNew(SResetToDefaultPropertyEditor, PropertyHandle)
+						];
+					}
+					return MakeFieldWidget(FieldWidget.ToSharedRef());
 				}
 			}
 		}
diff --git a/Engine/Plugins/VirtualProduction/RemoteControl/Source/RemoteControlUI/Private/UI/SRCPanelExposedField.h b/Engine/Plugins/VirtualProduction/RemoteControl/Source/RemoteControlUI/Private/UI/SRCPanelExposedField.h
--- a/Engine/Plugins/VirtualProduction/RemoteControl/Source/RemoteControlUI/Private/UI/SRCPanelExposedField.h
+++ b/Engine/Plugins/VirtualProduction/RemoteControl/Source/RemoteControlUI/Private/UI/SRCPanelExposedField.h
@@ -6,6 +6,7 @@
 #include "RemoteControlField.h"
 #include "SRCPanelExposedEntity.h"
 #include "SlateFwd.h"
+#include "SResetToDefaultPropertyEditor.h"
 #include "Widgets/DeclarativeSyntaxSupport.h"
 #include "Widgets/SCompoundWidget.h"
 
@@ -89,6 +90,8 @@ private:
 	TArray<TSharedPtr<SRCPanelFieldChildNode>> ChildWidgets;
 	/** Holds the panel's cached widgets. */
 	TWeakPtr<FRCPanelWidgetRegistry> WidgetRegistry;
+	/** Holds the zeroed Default Value of the ExposedField */
+	TSharedPtr<uint8> DefaultValue;
 };
 
 
