[UE-156718] Enable / Disable Behaviours - API support.

API support:

1) bIsEnabled boolean flag for URCBehaviour objects

2) SetIsBehaviourEnabled(bool) function in FRCBehaviourModel UI model.

3) IsBehaviourEnabled() function in FRCBehaviourModel UI model.

Usage: Invoked SetIsBehaviourEnabled as required on a checkbox (or other appropriate widget) change event.

Note: Also includes a minor UI change for storing Behaviour Title Text widget (part of the current Behaviours list panel) as a variable

diff --git a/Engine/Plugins/VirtualProduction/RemoteControl/Source/RemoteControlLogic/Private/Controller/RCController.cpp b/Engine/Plugins/VirtualProduction/RemoteControl/Source/RemoteControlLogic/Private/Controller/RCController.cpp
--- a/Engine/Plugins/VirtualProduction/RemoteControl/Source/RemoteControlLogic/Private/Controller/RCController.cpp
+++ b/Engine/Plugins/VirtualProduction/RemoteControl/Source/RemoteControlLogic/Private/Controller/RCController.cpp
@@ -71,6 +71,9 @@ void URCController::ExecuteBehaviours()
 {
 	for (URCBehaviour* Behaviour : Behaviours)
 	{
-		Behaviour->Execute();
+		if (Behaviour->bIsEnabled)
+		{
+			Behaviour->Execute();
+		}
 	}
 }
diff --git a/Engine/Plugins/VirtualProduction/RemoteControl/Source/RemoteControlLogic/Public/Behaviour/RCBehaviour.h b/Engine/Plugins/VirtualProduction/RemoteControl/Source/RemoteControlLogic/Public/Behaviour/RCBehaviour.h
--- a/Engine/Plugins/VirtualProduction/RemoteControl/Source/RemoteControlLogic/Public/Behaviour/RCBehaviour.h
+++ b/Engine/Plugins/VirtualProduction/RemoteControl/Source/RemoteControlLogic/Public/Behaviour/RCBehaviour.h
@@ -100,4 +100,10 @@ private:
 	/** Cached behaviour node */
 	UPROPERTY(Instanced)
 	TObjectPtr<URCBehaviourNode> CachedBehaviourNode;
+
+public:
+	/** Whether this Behaviour is currently enabled. 
+	* If disabled, it will be not evaluated when the associated Controller changes */
+	UPROPERTY()
+	bool bIsEnabled = true;
 };
diff --git a/Engine/Plugins/VirtualProduction/RemoteControl/Source/RemoteControlUI/Private/UI/Behaviour/RCBehaviourModel.cpp b/Engine/Plugins/VirtualProduction/RemoteControl/Source/RemoteControlUI/Private/UI/Behaviour/RCBehaviourModel.cpp
--- a/Engine/Plugins/VirtualProduction/RemoteControl/Source/RemoteControlUI/Private/UI/Behaviour/RCBehaviourModel.cpp
+++ b/Engine/Plugins/VirtualProduction/RemoteControl/Source/RemoteControlUI/Private/UI/Behaviour/RCBehaviourModel.cpp
@@ -8,12 +8,18 @@
 #include "UI/RemoteControlPanelStyle.h"
 #include "UI/SRemoteControlPanel.h"
 #include "Widgets/SNullWidget.h"
+#include "Widgets/Text/STextBlock.h"
 
 #define LOCTEXT_NAMESPACE "FRCBehaviourModel"
 
 FRCBehaviourModel::FRCBehaviourModel(URCBehaviour* InBehaviour)
 	: BehaviourWeakPtr(InBehaviour)
 {
+	const FText BehaviorDisplayName = BehaviourWeakPtr->GetDisplayName().ToUpper();
+
+	SAssignNew(BehaviourTitleText, STextBlock)
+		.Text(BehaviorDisplayName)
+		.Font(FRemoteControlPanelStyle::Get()->GetFontStyle("RemoteControlPanel.Behaviours.Title"));
 }
 
 TSharedRef<SWidget> FRCBehaviourModel::GetWidget() const
@@ -23,8 +29,6 @@ TSharedRef<SWidget> FRCBehaviourModel::GetWidget() const
 		return SNullWidget::NullWidget;
 	}
 
-	const FText BehaviorDisplayName = BehaviourWeakPtr->GetDisplayName().ToUpper();
-
 	return SNew(SHorizontalBox)
 		.Clipping(EWidgetClipping::OnDemand)
 		// Behaviour name
@@ -33,8 +37,7 @@ TSharedRef<SWidget> FRCBehaviourModel::GetWidget() const
 		.AutoWidth()
 		.Padding(FMargin(8.f))
 		[
-			SNew(STextBlock).Text(BehaviorDisplayName)
-			.Font(FRemoteControlPanelStyle::Get()->GetFontStyle("RemoteControlPanel.Behaviours.Title"))
+			BehaviourTitleText.ToSharedRef()
 		];
 }
 
@@ -58,6 +61,26 @@ void FRCBehaviourModel::OnOverrideBlueprint() const
 	}
 }
 
+bool FRCBehaviourModel::IsBehaviourEnabled() const
+{
+	if (URCBehaviour* Behaviour = BehaviourWeakPtr.Get())
+	{
+		return Behaviour->bIsEnabled;
+	}
+
+	return false;
+}
+
+void FRCBehaviourModel::SetIsBehaviourEnabled(const bool bIsEnabled)
+{
+	if (URCBehaviour* Behaviour = BehaviourWeakPtr.Get())
+	{
+		Behaviour->bIsEnabled = bIsEnabled;
+
+		BehaviourTitleText->SetEnabled(bIsEnabled);
+	}
+}
+
 URCBehaviour* FRCBehaviourModel::GetBehaviour() const
 {
 	return BehaviourWeakPtr.Get();
diff --git a/Engine/Plugins/VirtualProduction/RemoteControl/Source/RemoteControlUI/Private/UI/Behaviour/RCBehaviourModel.h b/Engine/Plugins/VirtualProduction/RemoteControl/Source/RemoteControlUI/Private/UI/Behaviour/RCBehaviourModel.h
--- a/Engine/Plugins/VirtualProduction/RemoteControl/Source/RemoteControlUI/Private/UI/Behaviour/RCBehaviourModel.h
+++ b/Engine/Plugins/VirtualProduction/RemoteControl/Source/RemoteControlUI/Private/UI/Behaviour/RCBehaviourModel.h
@@ -5,6 +5,7 @@
 #include "UI/BaseLogicUI/RCLogicModeBase.h"
 
 class IPropertyRowGenerator;
+class STextBlock;
 class SWidget;
 class URCBehaviour;
 
@@ -31,9 +32,18 @@ public:
 	URCBehaviour* GetBehaviour() const;
 
 	/** Handling for user action to override this behaviour via new Blueprint class */
-	void OnOverrideBlueprint() const;	
+	void OnOverrideBlueprint() const;
+
+	/** Whether the underlying Behaviour is currently enabled */
+	bool IsBehaviourEnabled() const;
+
+	/** Set the Enabled state of our underlying Behaviour */
+	void SetIsBehaviourEnabled(const bool bIsEnabled);
 
 private:
 	/** The Behaviour (Data model) associated with us*/
 	TWeakObjectPtr<URCBehaviour> BehaviourWeakPtr;
+
+	/** Text block widget for representing the Behaviour's title */
+	TSharedPtr<STextBlock> BehaviourTitleText;
 };
\ No newline at end of file
