1. RC logic - Implement UI as designed in figma

diff --git a/Engine/Plugins/VirtualProduction/RemoteControl/Source/RemoteControl/Private/RemoteControlExposeRegistry.cpp b/Engine/Plugins/VirtualProduction/RemoteControl/Source/RemoteControl/Private/RemoteControlExposeRegistry.cpp
--- a/Engine/Plugins/VirtualProduction/RemoteControl/Source/RemoteControl/Private/RemoteControlExposeRegistry.cpp
+++ b/Engine/Plugins/VirtualProduction/RemoteControl/Source/RemoteControl/Private/RemoteControlExposeRegistry.cpp
@@ -89,6 +89,11 @@ const TSet<UScriptStruct*>& URemoteControlExposeRegistry::GetExposedEntityTypes(
 	return ExposedTypes;
 }
 
+const bool URemoteControlExposeRegistry::IsEmpty() const
+{
+    return ExposedEntities.IsEmpty();
+}
+
 TSharedPtr<FRemoteControlEntity> URemoteControlExposeRegistry::AddExposedEntity(FRemoteControlEntity&& EntityToExpose, UScriptStruct* EntityType)
 {
 	LabelToIdCache.Add(EntityToExpose.GetLabel(), EntityToExpose.GetId());
diff --git a/Engine/Plugins/VirtualProduction/RemoteControl/Source/RemoteControl/Private/RemoteControlExposeRegistry.h b/Engine/Plugins/VirtualProduction/RemoteControl/Source/RemoteControl/Private/RemoteControlExposeRegistry.h
--- a/Engine/Plugins/VirtualProduction/RemoteControl/Source/RemoteControl/Private/RemoteControlExposeRegistry.h
+++ b/Engine/Plugins/VirtualProduction/RemoteControl/Source/RemoteControl/Private/RemoteControlExposeRegistry.h
@@ -154,6 +154,9 @@ public:
 	 */
 	const TSet<UScriptStruct*>& GetExposedEntityTypes() const;
 
+	/** Returns true when Exposed Entities is empty. */
+	const bool IsEmpty() const;
+
 	/**
 	 * Add an entity to the set of exposed entities of its type.
 	 * @param EntityToExpose the entity to expose.
diff --git a/Engine/Plugins/VirtualProduction/RemoteControl/Source/RemoteControl/Private/RemoteControlPreset.cpp b/Engine/Plugins/VirtualProduction/RemoteControl/Source/RemoteControl/Private/RemoteControlPreset.cpp
--- a/Engine/Plugins/VirtualProduction/RemoteControl/Source/RemoteControl/Private/RemoteControlPreset.cpp
+++ b/Engine/Plugins/VirtualProduction/RemoteControl/Source/RemoteControl/Private/RemoteControlPreset.cpp
@@ -1281,6 +1281,11 @@ const TSet<UScriptStruct*>& URemoteControlPreset::GetExposedEntityTypes() const
 	return Registry->GetExposedEntityTypes();
 }
 
+const bool URemoteControlPreset::HasEntities() const
+{
+	return !Registry->IsEmpty();
+}
+
 FName URemoteControlPreset::RenameExposedEntity(const FGuid& ExposedEntityId, FName NewLabel)
 {
 	FName AssignedLabel = Registry->RenameExposedEntity(ExposedEntityId, NewLabel);
diff --git a/Engine/Plugins/VirtualProduction/RemoteControl/Source/RemoteControl/Public/RemoteControlPreset.h b/Engine/Plugins/VirtualProduction/RemoteControl/Source/RemoteControl/Public/RemoteControlPreset.h
--- a/Engine/Plugins/VirtualProduction/RemoteControl/Source/RemoteControl/Public/RemoteControlPreset.h
+++ b/Engine/Plugins/VirtualProduction/RemoteControl/Source/RemoteControl/Public/RemoteControlPreset.h
@@ -435,6 +435,9 @@ public:
 	/** Get all types of exposed entities currently exposed. (ie. FRemoteControlActor) */
 	const TSet<UScriptStruct*>& GetExposedEntityTypes() const;
 
+	/** Returns true when Exposed Entities is populated. */
+	const bool HasEntities() const;
+
 	/** Returns whether an entity is exposed on the preset. */
 	bool IsExposed(const FGuid& ExposedEntityId) const;
 
diff --git a/Engine/Plugins/VirtualProduction/RemoteControl/Source/RemoteControlProtocolWidgets/Private/RemoteControlProtocolWidgetsModule.cpp b/Engine/Plugins/VirtualProduction/RemoteControl/Source/RemoteControlProtocolWidgets/Private/RemoteControlProtocolWidgetsModule.cpp
--- a/Engine/Plugins/VirtualProduction/RemoteControl/Source/RemoteControlProtocolWidgets/Private/RemoteControlProtocolWidgetsModule.cpp
+++ b/Engine/Plugins/VirtualProduction/RemoteControl/Source/RemoteControlProtocolWidgets/Private/RemoteControlProtocolWidgetsModule.cpp
@@ -8,6 +8,7 @@
 #include "IRemoteControlModule.h"
 #include "RemoteControlPreset.h"
 #include "ViewModels/ProtocolEntityViewModel.h"
+#include "Widgets/SNullWidget.h"
 #include "Widgets/SRCProtocolBindingList.h"
 
 DEFINE_LOG_CATEGORY(LogRemoteControlProtocolWidgets);
@@ -18,10 +19,18 @@ TSharedRef<SWidget> FRemoteControlProtocolWidgetsModule::GenerateDetailsForEntit
 
 	ResetProtocolBindingList();
 
+	if (!InFieldId.IsValid())
+	{
+		return SNullWidget::NullWidget;
+	}
+
 	// Currently only supports Properties
-	if (InFieldId.IsValid() && InPreset->GetExposedEntityType(InFieldId) == FRemoteControlProperty::StaticStruct())
+	if (const UScriptStruct* PropertyStruct = InPreset->GetExposedEntityType(InFieldId))
 	{
-		return SAssignNew(RCProtocolBindingList, SRCProtocolBindingList, FProtocolEntityViewModel::Create(InPreset, InFieldId));
+		if (PropertyStruct->IsChildOf(FRemoteControlProperty::StaticStruct()))
+		{
+			return SAssignNew(RCProtocolBindingList, SRCProtocolBindingList, FProtocolEntityViewModel::Create(InPreset, InFieldId));
+		}
 	}
 
 	return SNullWidget::NullWidget;
diff --git a/Engine/Plugins/VirtualProduction/RemoteControl/Source/RemoteControlUI/Private/Commands/RemoteControlCommands.cpp b/Engine/Plugins/VirtualProduction/RemoteControl/Source/RemoteControlUI/Private/Commands/RemoteControlCommands.cpp
--- a/Engine/Plugins/VirtualProduction/RemoteControl/Source/RemoteControlUI/Private/Commands/RemoteControlCommands.cpp
+++ b/Engine/Plugins/VirtualProduction/RemoteControl/Source/RemoteControlUI/Private/Commands/RemoteControlCommands.cpp
@@ -28,14 +28,11 @@ void FRemoteControlCommands::RegisterCommands()
 	// Find Preset
 	UI_COMMAND(FindPresetInContentBrowser, "Browse to Preset", "Browses to the associated preset and selects it in the most recently used Content Browser (summoning one if necessary)", EUserInterfaceActionType::Button, FInputChord(EModifierKey::Control, EKeys::B));
 
-	// Toggle Expose Functions
-	UI_COMMAND(ToggleExposeFunctions, "Functions", "Expose a function from available list of options.", EUserInterfaceActionType::Button, FInputChord());
-
 	// Toggle Protocol Mappings
-	UI_COMMAND(ToggleProtocolMappings, "Mappings", "View list of protocols mapped to active selection.", EUserInterfaceActionType::Button, FInputChord());
+	UI_COMMAND(ToggleProtocolMappings, "Protocols", "View list of protocols mapped to active selection.", EUserInterfaceActionType::ToggleButton, FInputChord());
 
 	// Toggle Logic Editor
-	UI_COMMAND(ToggleLogicEditor, "Logic", "View the logic applied to active selection.", EUserInterfaceActionType::Button, FInputChord());
+	UI_COMMAND(ToggleLogicEditor, "Logic", "View the logic applied to active selection.", EUserInterfaceActionType::ToggleButton, FInputChord());
 
 	// Delete Entity
 	UI_COMMAND(DeleteEntity, "Delete", "Delete the selected  group/exposed entity from the list.", EUserInterfaceActionType::Button, FInputChord(EKeys::Delete));
diff --git a/Engine/Plugins/VirtualProduction/RemoteControl/Source/RemoteControlUI/Private/Commands/RemoteControlCommands.h b/Engine/Plugins/VirtualProduction/RemoteControl/Source/RemoteControlUI/Private/Commands/RemoteControlCommands.h
--- a/Engine/Plugins/VirtualProduction/RemoteControl/Source/RemoteControlUI/Private/Commands/RemoteControlCommands.h
+++ b/Engine/Plugins/VirtualProduction/RemoteControl/Source/RemoteControlUI/Private/Commands/RemoteControlCommands.h
@@ -36,11 +36,6 @@ public:
 	 */
 	TSharedPtr<FUICommandInfo> SavePreset;
 
-	/**
-	 * Holds the information about UI Command that brings up a panel that facilitates exposing functions using RC API.
-	 */
-	TSharedPtr<FUICommandInfo> ToggleExposeFunctions;
-	
 	/**
 	 * Holds the information about UI Command that brings up a panel which holds the active protocol mappings.
 	 */
diff --git a/Engine/Plugins/VirtualProduction/RemoteControl/Source/RemoteControlUI/Private/Styling/RemoteControlStyles.cpp b/Engine/Plugins/VirtualProduction/RemoteControl/Source/RemoteControlUI/Private/Styling/RemoteControlStyles.cpp
new file mode 100644
--- /dev/null
+++ b/Engine/Plugins/VirtualProduction/RemoteControl/Source/RemoteControlUI/Private/Styling/RemoteControlStyles.cpp
@@ -0,0 +1,33 @@
+// Copyright Epic Games, Inc. All Rights Reserved.
+
+#include "RemoteControlStyles.h"
+
+FRCPanelStyle::FRCPanelStyle()
+	: HeaderRowPadding(4.f, 2.f)
+	, HeaderTextStyle(FTextBlockStyle::GetDefault())
+	, IconSize(FVector2D(20.f))
+	, PanelPadding(0.f)
+	, PanelTextStyle(FTextBlockStyle::GetDefault())
+	, SectionHeaderTextStyle(FTextBlockStyle::GetDefault())
+{
+}
+
+const FName FRCPanelStyle::TypeName(TEXT("FRemoteControlStyle"));
+
+const FRCPanelStyle& FRCPanelStyle::GetDefault()
+{
+	static FRCPanelStyle Default;
+	return Default;
+}
+
+void FRCPanelStyle::GetResources(TArray<const FSlateBrush*>& OutBrushes) const
+{
+	HeaderRowStyle.GetResources(OutBrushes);
+	FlatButtonStyle.GetResources(OutBrushes);
+	TableRowStyle.GetResources(OutBrushes);
+	TableViewStyle.GetResources(OutBrushes);
+	ToggleButtonStyle.GetResources(OutBrushes);
+	OutBrushes.Add(&ContentAreaBrush);
+	OutBrushes.Add(&SectionHeaderBrush);
+}
+
diff --git a/Engine/Plugins/VirtualProduction/RemoteControl/Source/RemoteControlUI/Private/Styling/RemoteControlStyles.h b/Engine/Plugins/VirtualProduction/RemoteControl/Source/RemoteControlUI/Private/Styling/RemoteControlStyles.h
new file mode 100644
--- /dev/null
+++ b/Engine/Plugins/VirtualProduction/RemoteControl/Source/RemoteControlUI/Private/Styling/RemoteControlStyles.h
@@ -0,0 +1,100 @@
+// Copyright Epic Games, Inc. All Rights Reserved.
+
+#pragma once
+
+#include "CoreMinimal.h"
+
+#include "Layout/Margin.h"
+
+#include "Styling/SlateBrush.h"
+#include "Styling/SlateTypes.h"
+#include "Styling/SlateWidgetStyle.h"
+
+#include "RemoteControlStyles.generated.h"
+
+/**
+ * Represents the appearance of an RC Panel.
+ */
+USTRUCT()
+struct FRCPanelStyle : public FSlateWidgetStyle
+{
+	GENERATED_BODY()
+
+	FRCPanelStyle();
+
+	virtual ~FRCPanelStyle() {}
+
+	//~ Begin : FSlateWidgetStyle
+	virtual void GetResources(TArray<const FSlateBrush*>& OutBrushes) const override;
+
+	static const FName TypeName;
+	virtual const FName GetTypeName() const override { return TypeName; };
+	
+	static const FRCPanelStyle& GetDefault();
+	//~ End : FSlateWidgetStyle
+
+	/** Brush used by the RC Panels to draw the content associated with this panel. */
+	UPROPERTY(EditAnywhere, Category = Appearance)
+		FSlateBrush ContentAreaBrush;
+	FRCPanelStyle& SetContentAreaBrush(const FSlateBrush& InContentAreaBrush) { ContentAreaBrush = InContentAreaBrush; return *this; }
+
+	/** Style used for the add new button */
+	UPROPERTY(EditAnywhere, Category = Appearance)
+		FButtonStyle FlatButtonStyle;
+	FRCPanelStyle& SetFlatButtonStyle(const FButtonStyle& InFlatButtonStyle) { FlatButtonStyle = InFlatButtonStyle; return *this; }
+
+	/** Panel Header Row Style */
+	UPROPERTY(EditAnywhere, Category = Appearance)
+		FMargin HeaderRowPadding;
+	FRCPanelStyle& SetHeaderRowPadding(const FMargin& InHeaderRowPadding) { HeaderRowPadding = InHeaderRowPadding; return *this; }
+	
+	/** Panel Header Row Style */
+	UPROPERTY(EditAnywhere, Category = Appearance)
+		FHeaderRowStyle HeaderRowStyle;
+	FRCPanelStyle& SetHeaderRowStyle(const FHeaderRowStyle& InHeaderRowStyle) { HeaderRowStyle = InHeaderRowStyle; return *this; }
+	
+	/** Panel Header Text Style */
+	UPROPERTY(EditAnywhere, Category = Appearance)
+		FTextBlockStyle HeaderTextStyle;
+	FRCPanelStyle& SetHeaderTextStyle(const FTextBlockStyle& InHeaderTextStyle) { HeaderTextStyle = InHeaderTextStyle; return *this; }
+
+	/** Padding used around this panel */
+	UPROPERTY(EditAnywhere, Category = Appearance)
+		FVector2D IconSize;
+	FRCPanelStyle& SetIconSize(const FVector2D& InIconSize) { IconSize = InIconSize; return *this; }
+
+	/** Padding used around this panel */
+	UPROPERTY(EditAnywhere, Category = Appearance)
+		FMargin PanelPadding;
+	FRCPanelStyle& SetPanelPadding(const FMargin& InPanelPadding) { PanelPadding = InPanelPadding; return *this; }
+
+	/** Panel Text Style */
+	UPROPERTY(EditAnywhere, Category = Appearance)
+		FTextBlockStyle PanelTextStyle;
+	FRCPanelStyle& SetPanelTextStyle(const FTextBlockStyle& InPanelTextStyle) { PanelTextStyle = InPanelTextStyle; return *this; }
+
+	/** Panel Section Header Text Style */
+	UPROPERTY(EditAnywhere, Category = Appearance)
+		FTextBlockStyle SectionHeaderTextStyle;
+	FRCPanelStyle& SetSectionHeaderTextStyle(const FTextBlockStyle& InSectionHeaderTextStyle) { SectionHeaderTextStyle = InSectionHeaderTextStyle; return *this; }
+
+	/** Brush used by the RC Panels to draw the section header associated with this panel. */
+	UPROPERTY(EditAnywhere, Category = Appearance)
+		FSlateBrush SectionHeaderBrush;
+	FRCPanelStyle& SetSectionHeaderBrush(const FSlateBrush& InSectionHeaderBrush) { SectionHeaderBrush = InSectionHeaderBrush; return *this; }
+
+	/** Panel Table Row Style */
+	UPROPERTY(EditAnywhere, Category = Appearance)
+		FTableRowStyle TableRowStyle;
+	FRCPanelStyle& SetTableRowStyle(const FTableRowStyle& InTableRowStyle) { TableRowStyle = InTableRowStyle; return *this; }
+	
+	/** Panel Table View Style */
+	UPROPERTY(EditAnywhere, Category = Appearance)
+		FTableViewStyle TableViewStyle;
+	FRCPanelStyle& SetTableViewStyle(const FTableViewStyle& InTableViewStyle) { TableViewStyle = InTableViewStyle; return *this; }
+
+	/** Style used for the add new button */
+	UPROPERTY(EditAnywhere, Category = Appearance)
+		FCheckBoxStyle ToggleButtonStyle;
+	FRCPanelStyle& SetToggleButtonStyle(const FCheckBoxStyle& InToggleButtonStyle) { ToggleButtonStyle = InToggleButtonStyle; return *this; }
+};
diff --git a/Engine/Plugins/VirtualProduction/RemoteControl/Source/RemoteControlUI/Private/UI/Action/SRCActionPanel.cpp b/Engine/Plugins/VirtualProduction/RemoteControl/Source/RemoteControlUI/Private/UI/Action/SRCActionPanel.cpp
--- a/Engine/Plugins/VirtualProduction/RemoteControl/Source/RemoteControlUI/Private/UI/Action/SRCActionPanel.cpp
+++ b/Engine/Plugins/VirtualProduction/RemoteControl/Source/RemoteControlUI/Private/UI/Action/SRCActionPanel.cpp
@@ -2,29 +2,66 @@
 
 #include "SRCActionPanel.h"
 
+#include "RemoteControlField.h"
+#include "RemoteControlPreset.h"
+
+#include "Action/RCActionContainer.h"
+#include "Action/RCFunctionAction.h"
+#include "Action/RCPropertyAction.h"
+
+#include "Behaviour/RCBehaviour.h"
+#include "Controller/RCController.h"
+
+#include "EdGraphSchema_K2.h"
+#include "EdGraph/EdGraphPin.h"
+
+#include "Framework/MultiBox/MultiBoxBuilder.h"
+#include "Misc/MessageDialog.h"
+
 #include "SlateOptMacros.h"
-#include "UI/SRemoteControlPanel.h"
-#include "SRCActionPanelHeader.h"
 #include "SRCActionPanelList.h"
-#include "SRCActionPanelFooter.h"
+#include "SRCBehaviourDetails.h"
+#include "Styling/RemoteControlStyles.h"
 
 #include "UI/Behaviour/RCBehaviourModel.h"
+#include "UI/Panels/SRCDockPanel.h"
+#include "UI/RemoteControlPanelStyle.h"
+#include "UI/SRemoteControlPanel.h"
+
+#include "Widgets/Input/SButton.h"
+#include "Widgets/Input/SComboButton.h"
+#include "Widgets/Input/SCheckBox.h"
+#include "Widgets/Layout/SBox.h"
 
 #define LOCTEXT_NAMESPACE "SRCActionPanel"
 
+TSharedRef<SBox> SRCActionPanel::NoneSelectedWidget = SNew(SBox)
+			.Padding(0.f)
+			.VAlign(VAlign_Center)
+			.HAlign(HAlign_Center)
+			[
+				SNew(STextBlock)
+				.Text(LOCTEXT("NoneSelected", "Select a behavior to view its actions."))
+				.TextStyle(&FAppStyle::GetWidgetStyle<FTextBlockStyle>("NormalText"))
+				.Justification(ETextJustify::Center)
+			];
+
 BEGIN_SLATE_FUNCTION_BUILD_OPTIMIZATION
 
 void SRCActionPanel::Construct(const FArguments& InArgs, const TSharedRef<SRemoteControlPanel>& InPanel)
 {
 	SRCLogicPanelBase::Construct(SRCLogicPanelBase::FArguments(), InPanel);
 	
+	RCPanelStyle = &FRemoteControlPanelStyle::Get()->GetWidgetStyle<FRCPanelStyle>("RemoteControlPanel.MinorPanel");
+
 	WrappedBoxWidget = SNew(SBox);
 	UpdateWrappedWidget();
 	
 	ChildSlot
-	[
-		WrappedBoxWidget.ToSharedRef()
-	];
+		.Padding(RCPanelStyle->PanelPadding)
+		[
+			WrappedBoxWidget.ToSharedRef()
+		];
 
 	// Register delegates
 	InPanel->OnBehaviourSelectionChanged.AddSP(this, &SRCActionPanel::OnBehaviourSelectionChanged);
@@ -42,60 +79,284 @@ void SRCActionPanel::UpdateWrappedWidget(TSharedPtr<FRCBehaviourModel> InBehavio
 {
 	if (InBehaviourItem.IsValid())
 	{
-		WrappedBoxWidget->SetContent(
-		SNew(SVerticalBox)
-		// Header
-		+ SVerticalBox::Slot()
-		.Padding(5.0f, 5.0f)
-		.AutoHeight()
-		[
-			SNew(SRCActionPanelHeader, SharedThis(this), InBehaviourItem.ToSharedRef())
-		]
+		// Action Dock Panel
+		TSharedPtr<SRCMinorPanel> ActionDockPanel = SNew(SRCMinorPanel)
+			.HeaderLabel(LOCTEXT("ActionsLabel", "Actions"))
+			[
+				SAssignNew(ActionPanelList, SRCActionPanelList, SharedThis(this), InBehaviourItem)
+			];
 
-		// Border (separating Header and "If Value Is" panel)
-		+ SVerticalBox::Slot()
-		.AutoHeight()
-		[
-			SNew(SBorder)
-			.BorderImage(FAppStyle::Get().GetBrush("ToolPanel.GroupBorder"))
-			.BorderBackgroundColor(FColor::Black)
-			.Padding(FMargin(0.f, 1.f))
+		// Add New Action Button
+		const TSharedRef<SWidget> AddNewActionButton = SNew(SComboButton)
+			.AddMetaData<FTagMetaData>(FTagMetaData(TEXT("Add Action")))
+			.HAlign(HAlign_Center)
+			.VAlign(VAlign_Center)
+			.ButtonStyle(&RCPanelStyle->FlatButtonStyle)
+			.ForegroundColor(FSlateColor::UseForeground())
+			.CollapseMenuOnParentFocus(true)
+			.HasDownArrow(false)
+			.ContentPadding(FMargin(4.f, 2.f))
+			.ButtonContent()
 			[
-				SNullWidget::NullWidget
+				SNew(SBox)
+				.WidthOverride(RCPanelStyle->IconSize.X)
+				.HeightOverride(RCPanelStyle->IconSize.Y)
+				[
+					SNew(SImage)
+					.ColorAndOpacity(FSlateColor::UseForeground())
+					.Image(FAppStyle::GetBrush("Icons.PlusCircle"))
+				]
 			]
-		]
+			.MenuContent()
+			[
+				GetActionMenuContentWidget()
+			];
+		
+		// Add All Button
+		TSharedRef<SWidget> AddAllActionsButton = SNew(SButton)
+			.AddMetaData<FTagMetaData>(FTagMetaData(TEXT("Add All Actions")))
+			.HAlign(HAlign_Center)
+			.VAlign(VAlign_Center)
+			.ForegroundColor(FSlateColor::UseForeground())
+			.ButtonStyle(&RCPanelStyle->FlatButtonStyle)
+			.ToolTipText(LOCTEXT("AddAllToolTip", "Adds all the available actions."))
+			.OnClicked(this, &SRCActionPanel::OnAddAllFields)
+			.Visibility(this, &SRCActionPanel::HandleAddAllButtonVisibility)
+			[
+				SNew(SBox)
+				.WidthOverride(RCPanelStyle->IconSize.X)
+				.HeightOverride(RCPanelStyle->IconSize.Y)
+				[
+					SNew(SImage)
+					.ColorAndOpacity(FSlateColor::UseForeground())
+					.Image(FAppStyle::GetBrush("Icons.Duplicate"))
+				]
+			];
+		
+		// Empty All Button
+		TSharedRef<SWidget> EmptyAllButton = SNew(SButton)
+			.AddMetaData<FTagMetaData>(FTagMetaData(TEXT("Empty All Actions")))
+			.HAlign(HAlign_Center)
+			.VAlign(VAlign_Center)
+			.ForegroundColor(FSlateColor::UseForeground())
+			.ButtonStyle(&RCPanelStyle->FlatButtonStyle)
+			.ToolTipText(LOCTEXT("EmptyAllToolTip", "Deletes all the actions."))
+			.OnClicked(this, &SRCActionPanel::RequestDeleteAllItems)
+			.Visibility_Lambda([this]() { return ActionPanelList.IsValid() && !ActionPanelList->IsEmpty() ? EVisibility::Visible : EVisibility::Collapsed; })
+			[
+				SNew(SBox)
+				.WidthOverride(RCPanelStyle->IconSize.X)
+				.HeightOverride(RCPanelStyle->IconSize.Y)
+				[
+					SNew(SImage)
+					.ColorAndOpacity(FSlateColor::UseForeground())
+					.Image(FAppStyle::GetBrush("Icons.Delete"))
+				]
+			];
 
-		// Behaviour Specific Details Panel
-		+ SVerticalBox::Slot()
-		.Padding(20.0f, 5.0f)
-		.AutoHeight()
-		[
-			InBehaviourItem->GetBehaviourDetailsWidget()			
-		]
+		ActionDockPanel->AddHeaderToolbarItem(EToolbar::Left, AddNewActionButton);
+		ActionDockPanel->AddHeaderToolbarItem(EToolbar::Right, AddAllActionsButton);
+		ActionDockPanel->AddHeaderToolbarItem(EToolbar::Right, EmptyAllButton);
 
-		// Actions Mapping List
-		+ SVerticalBox::Slot()
-		.AutoHeight()
-		.Padding(FMargin(5.f))
-		[
-			SAssignNew(ActionPanelList, SRCActionPanelList, SharedThis(this), InBehaviourItem)
-		]
+		// Duplicate Dock Panel
+		TSharedPtr<SRCMinorPanel> DuplicateDockPanel = SNew(SRCMinorPanel)
+			.HeaderLabel(LOCTEXT("DuplicateLabel", "Duplicate"))
+			[
+				SNew(SRCBehaviourDetails, SharedThis(this), InBehaviourItem.ToSharedRef())
+			];
 
-		// Footer
-		+ SVerticalBox::Slot()
-		.Padding(5.0f, 5.0f)
-		.AutoHeight()
-		[
-			SNew(SRCActionPanelFooter, SharedThis(this), InBehaviourItem.ToSharedRef())
-		]
-	);
+		// Duplicate Behaviour Button
+		const TSharedRef<SWidget> DuplicateBehaviourButton = SNew(SButton)
+			.AddMetaData<FTagMetaData>(FTagMetaData(TEXT("Duplicate Behavior")))
+			.HAlign(HAlign_Center)
+			.VAlign(VAlign_Center)
+			.ButtonStyle(&RCPanelStyle->FlatButtonStyle)
+			.ForegroundColor(FSlateColor::UseForeground())
+			.ContentPadding(FMargin(4.f, 2.f))
+			.OnClicked(this, &SRCActionPanel::OnClickOverrideBlueprintButton)
+			[
+				SNew(SBox)
+				.WidthOverride(RCPanelStyle->IconSize.X)
+				.HeightOverride(RCPanelStyle->IconSize.Y)
+				[
+					SNew(SImage)
+					.ColorAndOpacity(FSlateColor::UseForeground())
+					.Image(FAppStyle::GetBrush("Icons.PlusCircle"))
+				]
+			];
+		
+		// Toggle Behaviour Button
+		// TODO : Replace this with the actual algorithm to disable or enable behaviour.
+		const bool bIsChecked = FMath::RandBool();
+		const TSharedRef<SWidget> ToggleBehaviourButton = SNew(SCheckBox)
+			.AddMetaData<FTagMetaData>(FTagMetaData(TEXT("Duplicate Behavior")))
+			.HAlign(HAlign_Center)
+			.Style(&RCPanelStyle->ToggleButtonStyle)
+			.ForegroundColor(FSlateColor::UseForeground())
+			.IsChecked_Lambda([bIsChecked]() { return bIsChecked ? ECheckBoxState::Checked : ECheckBoxState::Unchecked; });
+
+		DuplicateDockPanel->AddHeaderToolbarItem(EToolbar::Left, DuplicateBehaviourButton);
+		DuplicateDockPanel->AddHeaderToolbarItem(EToolbar::Right, ToggleBehaviourButton);
+
+		TSharedRef<SRCMajorPanel> ActionsPanel = SNew(SRCMajorPanel)
+			.EnableFooter(false)
+			.EnableHeader(false)
+			.ChildOrientation(Orient_Vertical);
+
+		ActionsPanel->AddPanel(DuplicateDockPanel.ToSharedRef(), 0.5f);
+		ActionsPanel->AddPanel(ActionDockPanel.ToSharedRef(), 0.5f);
+
+		WrappedBoxWidget->SetContent(ActionsPanel);
 	}
 	else
 	{
-		WrappedBoxWidget->SetContent(
-			SNew(STextBlock).Text(LOCTEXT("Unselected", "No Behaviour Selected"))
-		);
+		WrappedBoxWidget->SetContent(NoneSelectedWidget);
+	}
+}
+
+FReply SRCActionPanel::OnClickOverrideBlueprintButton()
+{
+	if (!SelectedBehaviourItemWeakPtr.IsValid())
+	{
+		return FReply::Unhandled();
+	}
+
+	if (TSharedPtr<FRCBehaviourModel> Behaviour = SelectedBehaviourItemWeakPtr.Pin())
+	{
+		Behaviour->OnOverrideBlueprint();
+	}
+
+	return FReply::Handled();
+}
+
+TSharedRef<SWidget> SRCActionPanel::GetActionMenuContentWidget()
+{
+	constexpr bool bShouldCloseWindowAfterMenuSelection = true;
+	FMenuBuilder MenuBuilder(bShouldCloseWindowAfterMenuSelection, nullptr);
+
+	if (!SelectedBehaviourItemWeakPtr.IsValid())
+	{
+		return MenuBuilder.MakeWidget();
 	}
+
+	// List of exposed entities
+	if (URemoteControlPreset* Preset = GetPreset())
+	{
+		const TArray<TWeakPtr<FRemoteControlField>>& RemoteControlFields = Preset->GetExposedEntities<FRemoteControlField>();
+		for (const TWeakPtr<FRemoteControlField>& RemoteControlFieldWeakPtr : RemoteControlFields)
+		{
+			if (const TSharedPtr<FRemoteControlField> RemoteControlField = RemoteControlFieldWeakPtr.Pin())
+			{
+				if (const URCBehaviour* Behaviour = SelectedBehaviourItemWeakPtr.Pin()->GetBehaviour())
+				{
+					// Skip if we already have an Action created for this exposed entity
+					if (Behaviour->ActionContainer->FindActionByFieldId(RemoteControlField->GetId()))
+					{
+						continue;
+					}
+
+					// Create menu entry
+					FUIAction Action(FExecuteAction::CreateSP(this, &SRCActionPanel::OnAddActionClicked, RemoteControlField));
+					MenuBuilder.AddMenuEntry(
+						FText::Format(LOCTEXT("AddAction", "{0}"), FText::FromName(RemoteControlField->GetLabel())),
+						FText::Format(LOCTEXT("AddActionTooltip", "Add {0}"), FText::FromName(RemoteControlField->GetLabel())),
+						FSlateIcon(),
+						MoveTemp(Action));
+				}
+			}
+		}
+	}
+
+	return MenuBuilder.MakeWidget();
+}
+
+void SRCActionPanel::OnAddActionClicked(TSharedPtr<FRemoteControlField> InRemoteControlField)
+{
+	if (!SelectedBehaviourItemWeakPtr.IsValid())
+	{
+		return;
+	}
+
+	if (const URCBehaviour* Behaviour = SelectedBehaviourItemWeakPtr.Pin()->GetBehaviour())
+	{
+		URCAction* NewAction = nullptr;
+
+		if (InRemoteControlField->FieldType == EExposedFieldType::Property)
+		{
+			NewAction = Behaviour->ActionContainer->AddAction(StaticCastSharedPtr<FRemoteControlProperty>(InRemoteControlField));
+		}
+		else if (InRemoteControlField->FieldType == EExposedFieldType::Function)
+		{
+			NewAction = Behaviour->ActionContainer->AddAction(StaticCastSharedPtr<FRemoteControlFunction>(InRemoteControlField));
+		}
+
+		if (const TSharedPtr<SRemoteControlPanel> RemoteControlPanel = GetRemoteControlPanel())
+		{
+			RemoteControlPanel->OnActionAdded.Broadcast(NewAction);
+		}
+	}
+}
+
+FReply SRCActionPanel::OnClickEmptyButton()
+{
+	if (const TSharedPtr<FRCBehaviourModel> BehaviourItem = SelectedBehaviourItemWeakPtr.Pin())
+	{
+		if (const URCBehaviour* Behaviour = BehaviourItem->GetBehaviour())
+		{
+			Behaviour->ActionContainer->EmptyActions();
+		}
+	}
+
+	if (const TSharedPtr<SRemoteControlPanel> RemoteControlPanel = GetRemoteControlPanel())
+	{
+		RemoteControlPanel->OnEmptyActions.Broadcast();
+	}
+
+	return FReply::Handled();
+}
+
+FReply SRCActionPanel::OnAddAllFields()
+{
+	if (!SelectedBehaviourItemWeakPtr.IsValid())
+	{
+		return FReply::Handled();
+	}
+
+	if (URemoteControlPreset* Preset = GetPreset())
+	{
+		if (const URCBehaviour* Behaviour = SelectedBehaviourItemWeakPtr.Pin()->GetBehaviour())
+		{
+			const TArray<TWeakPtr<FRemoteControlField>>& RemoteControlFields = Preset->GetExposedEntities<FRemoteControlField>();
+
+			// Enumerate the list of Exposed Entities and Functions available in this Preset for our use as Actions
+			for (const TWeakPtr<FRemoteControlField>& RemoteControlFieldWeakPtr : RemoteControlFields)
+			{
+				if (const TSharedPtr<FRemoteControlField> RemoteControlField = RemoteControlFieldWeakPtr.Pin())
+				{
+					URCAction* NewAction = nullptr;
+
+					// Property Action
+					if (RemoteControlField->FieldType == EExposedFieldType::Property)
+					{
+						NewAction = Behaviour->ActionContainer->AddAction(StaticCastSharedPtr<FRemoteControlProperty>(RemoteControlField));
+					}
+					// Function Action
+					else if (RemoteControlField->FieldType == EExposedFieldType::Function)
+					{
+						NewAction = Behaviour->ActionContainer->AddAction(StaticCastSharedPtr<FRemoteControlFunction>(RemoteControlField));
+					}
+
+					// Broadcast New Action
+					if (const TSharedPtr<SRemoteControlPanel> RemoteControlPanel = GetRemoteControlPanel())
+					{
+						RemoteControlPanel->OnActionAdded.Broadcast(NewAction);
+					}
+				}
+			}
+		}
+	}
+
+	return FReply::Handled();
 }
 
 bool SRCActionPanel::IsListFocused() const
@@ -108,4 +369,33 @@ void SRCActionPanel::DeleteSelectedPanelItem()
 	ActionPanelList->DeleteSelectedPanelItem();
 }
 
+FReply SRCActionPanel::RequestDeleteAllItems()
+{
+	if (!ActionPanelList.IsValid())
+	{
+		return FReply::Unhandled();
+	}
+
+	const FText WarningMessage = FText::Format(LOCTEXT("DeleteAllWarning", "You are about to delete '{0}' actions. This action might not be undone.\nAre you sure you want to proceed?"), ActionPanelList->Num());
+
+	EAppReturnType::Type UserResponse = FMessageDialog::Open(EAppMsgType::YesNo, WarningMessage);
+
+	if (UserResponse == EAppReturnType::Yes)
+	{
+		return OnClickEmptyButton();
+	}
+
+	return FReply::Handled();
+}
+
+EVisibility SRCActionPanel::HandleAddAllButtonVisibility() const
+{
+	if (URemoteControlPreset* Preset = GetPreset())
+	{
+		return Preset->HasEntities() ? EVisibility::Visible : EVisibility::Collapsed;
+	}
+
+	return EVisibility::Collapsed;
+}
+
 #undef LOCTEXT_NAMESPACE
\ No newline at end of file
diff --git a/Engine/Plugins/VirtualProduction/RemoteControl/Source/RemoteControlUI/Private/UI/Action/SRCActionPanel.h b/Engine/Plugins/VirtualProduction/RemoteControl/Source/RemoteControlUI/Private/UI/Action/SRCActionPanel.h
--- a/Engine/Plugins/VirtualProduction/RemoteControl/Source/RemoteControlUI/Private/UI/Action/SRCActionPanel.h
+++ b/Engine/Plugins/VirtualProduction/RemoteControl/Source/RemoteControlUI/Private/UI/Action/SRCActionPanel.h
@@ -1,12 +1,14 @@
 ﻿// Copyright Epic Games, Inc. All Rights Reserved.
 
 #pragma once
+
 #include "UI/BaseLogicUI/SRCLogicPanelBase.h"
-#include "Widgets/Layout/SBox.h"
 
+struct FRCPanelStyle;
+struct FRemoteControlField;
 class FRCBehaviourModel;
+class SBox;
 class SRemoteControlPanel;
-
 /*
 * ~ SRCActionPanel ~
 *
@@ -25,7 +27,22 @@ public:
 	/** Constructs this widget with InArgs */
 	void Construct(const FArguments& InArgs, const TSharedRef<SRemoteControlPanel>& InPanel);
 
+	/** Whether the Actions list widget currently has focus.*/
+	bool IsListFocused() const;
+
+	/** Delete Item UI command implementation for this panel */
+	virtual void DeleteSelectedPanelItem() override;
+
+protected:
+
+	/** Warns user before deleting all items in a panel. */
+	virtual FReply RequestDeleteAllItems() override;
+
 private:
+
+	/** Determines the visibility Add All Button. */
+	EVisibility HandleAddAllButtonVisibility() const;
+
 	/** 
 	* Behaviour selection change listener.
 	* Updates the list of actions from newly selected Behaviour
@@ -35,6 +52,26 @@ private:
 	/* Rebuilds the Action Panel for a newly selected Behaviour*/
 	void UpdateWrappedWidget(TSharedPtr<FRCBehaviourModel> InBehaviourItem = nullptr);
 
+	/** Handles click event for Open Behaviour Blueprint button*/
+	FReply OnClickOverrideBlueprintButton();
+
+	/**
+	 * Builds a menu containing the list of all possible Actions
+	 * These are derived from the list of Exposed entities of the Remote Control Preset associated with us.
+	 */
+	TSharedRef<SWidget> GetActionMenuContentWidget();
+
+	/** Handles click event for Add Action button*/
+	void OnAddActionClicked(TSharedPtr<FRemoteControlField> InRemoteControlField);
+
+	/** Handles click event for Empty button; clear all Actions from the panel*/
+	FReply OnClickEmptyButton();
+
+	/** Handles click event for Add All button; Adds all possible actions for the active Remote Control Preset*/
+	FReply OnAddAllFields();
+
+private:
+
 	/** The parent Behaviour that this Action panel is associated with */
 	TWeakPtr<FRCBehaviourModel> SelectedBehaviourItemWeakPtr;
 	
@@ -44,10 +81,9 @@ private:
 	/** Widget representing List of Actions */
 	TSharedPtr<class SRCActionPanelList> ActionPanelList;
 
-public:
-	/** Whether the Actions list widget currently has focus.*/
-	bool IsListFocused() const;
+	/** Helper widget for behavior details. */
+	static TSharedRef<SBox> NoneSelectedWidget;
 
-	/** Delete Item UI command implementation for this panel */
-	virtual void DeleteSelectedPanelItem() override;
+	/** Panel Style reference. */
+	const FRCPanelStyle* RCPanelStyle;
 };
diff --git a/Engine/Plugins/VirtualProduction/RemoteControl/Source/RemoteControlUI/Private/UI/Action/SRCActionPanelFooter.cpp b/Engine/Plugins/VirtualProduction/RemoteControl/Source/RemoteControlUI/Private/UI/Action/SRCActionPanelFooter.cpp
deleted file mode 100644
diff --git a/Engine/Plugins/VirtualProduction/RemoteControl/Source/RemoteControlUI/Private/UI/Action/SRCActionPanelFooter.h b/Engine/Plugins/VirtualProduction/RemoteControl/Source/RemoteControlUI/Private/UI/Action/SRCActionPanelFooter.h
deleted file mode 100644
diff --git a/Engine/Plugins/VirtualProduction/RemoteControl/Source/RemoteControlUI/Private/UI/Action/SRCActionPanelHeader.cpp b/Engine/Plugins/VirtualProduction/RemoteControl/Source/RemoteControlUI/Private/UI/Action/SRCActionPanelHeader.cpp
deleted file mode 100644
diff --git a/Engine/Plugins/VirtualProduction/RemoteControl/Source/RemoteControlUI/Private/UI/Action/SRCActionPanelHeader.h b/Engine/Plugins/VirtualProduction/RemoteControl/Source/RemoteControlUI/Private/UI/Action/SRCActionPanelHeader.h
deleted file mode 100644
diff --git a/Engine/Plugins/VirtualProduction/RemoteControl/Source/RemoteControlUI/Private/UI/Action/SRCActionPanelList.cpp b/Engine/Plugins/VirtualProduction/RemoteControl/Source/RemoteControlUI/Private/UI/Action/SRCActionPanelList.cpp
--- a/Engine/Plugins/VirtualProduction/RemoteControl/Source/RemoteControlUI/Private/UI/Action/SRCActionPanelList.cpp
+++ b/Engine/Plugins/VirtualProduction/RemoteControl/Source/RemoteControlUI/Private/UI/Action/SRCActionPanelList.cpp
@@ -9,6 +9,7 @@
 #include "RCActionModel.h"
 #include "SlateOptMacros.h"
 #include "SRCActionPanel.h"
+#include "Styling/RemoteControlStyles.h"
 #include "UI/Behaviour/RCBehaviourModel.h"
 #include "UI/RCUIHelpers.h"
 #include "UI/RemoteControlPanelStyle.h"
@@ -19,6 +20,13 @@
 
 #define LOCTEXT_NAMESPACE "SRCActionPanelList"
 
+namespace FRemoteControlActionColumns
+{
+	const FName DragDropHandle = TEXT("DragDropHandle");
+	const FName Description = TEXT("Description");
+	const FName Value = TEXT("Value");
+}
+
 BEGIN_SLATE_FUNCTION_BUILD_OPTIMIZATION
 
 void SRCActionPanelList::Construct(const FArguments& InArgs, const TSharedRef<SRCActionPanel> InActionPanel, TSharedPtr<FRCBehaviourModel> InBehaviourItem)
@@ -28,43 +36,35 @@ void SRCActionPanelList::Construct(const FArguments& InArgs, const TSharedRef<SR
 	ActionPanelWeakPtr = InActionPanel;
 	BehaviourItemWeakPtr = InBehaviourItem;
 	
+	RCPanelStyle = &FRemoteControlPanelStyle::Get()->GetWidgetStyle<FRCPanelStyle>("RemoteControlPanel.MinorPanel");
+
 	ListView = SNew(SListView<TSharedPtr<FRCActionModel>>)
 		.ListItemsSource( &ActionItems )
 		.OnGenerateRow(this, &SRCActionPanelList::OnGenerateWidgetForList )
-		.HeaderRow
-		(
+		.ListViewStyle(&RCPanelStyle->TableViewStyle)
+		.HeaderRow(
 			SNew(SHeaderRow)
-
-			+ SHeaderRow::Column("Description")
-			.DefaultLabel(LOCTEXT("Description", "Description"))
-
-			+ SHeaderRow::Column("Value")
-			.DefaultLabel(LOCTEXT("Value", "Value"))
+			.Style(&RCPanelStyle->HeaderRowStyle)
+
+			+ SHeaderRow::Column(FRemoteControlActionColumns::DragDropHandle)
+			.DefaultLabel(LOCTEXT("RCActionDragDropHandleColumnHeader", ""))
+			.FixedWidth(25.f)
+			.HeaderContentPadding(RCPanelStyle->HeaderRowPadding)
+
+			+ SHeaderRow::Column(FRemoteControlActionColumns::Description)
+			.DefaultLabel(LOCTEXT("RCActionDescColumnHeader", "Description"))
+			.FillWidth(0.5f)
+			.HeaderContentPadding(RCPanelStyle->HeaderRowPadding)
+
+			+ SHeaderRow::Column(FRemoteControlActionColumns::Value)
+			.DefaultLabel(LOCTEXT("RCActionValueColumnHeader", "Value"))
+			.FillWidth(0.5f)
+			.HeaderContentPadding(RCPanelStyle->HeaderRowPadding)
 		);
 	
 	ChildSlot
 	[
-		SNew(SVerticalBox)
-		+ SVerticalBox::Slot()
-		.HAlign(HAlign_Fill)
-		.AutoHeight()
-		.Padding(FMargin(0.f, 5.f))
-		[
-			SNew(SBorder)
-			.BorderImage(FAppStyle::Get().GetBrush("ToolPanel.GroupBorder"))
-			.BorderBackgroundColor(FSlateColor::UseSubduedForeground())
-			.Padding(FMargin(3.f))
-			[
-				SNew(STextBlock)
-				.Text(LOCTEXT("ActionsMapping", "Actions Mapping"))
-			]
-		]
-
-		+ SVerticalBox::Slot()
-		.AutoHeight()
-		[
-			ListView.ToSharedRef()
-		]
+		ListView.ToSharedRef()
 	];
 
 	// Add delegates
@@ -80,6 +80,16 @@ void SRCActionPanelList::Construct(const FArguments& InArgs, const TSharedRef<SR
 
 END_SLATE_FUNCTION_BUILD_OPTIMIZATION
 
+bool SRCActionPanelList::IsEmpty() const
+{
+	return ActionItems.IsEmpty();
+}
+
+int32 SRCActionPanelList::Num() const
+{
+	return ActionItems.Num();
+}
+
 void SRCActionPanelList::Reset()
 {
 	ActionItems.Empty();
@@ -108,9 +118,10 @@ void SRCActionPanelList::Reset()
 TSharedRef<ITableRow> SRCActionPanelList::OnGenerateWidgetForList(TSharedPtr<FRCActionModel> InItem, const TSharedRef<STableViewBase>& OwnerTable)
 {
 	return SNew(STableRow<TSharedPtr<FString>>, OwnerTable)
-	[
-		InItem->GetWidget()
-	];
+		.Style(&RCPanelStyle->TableRowStyle)
+		[
+			InItem->GetWidget()
+		];
 }
 
 void SRCActionPanelList::OnActionAdded(URCAction* InAction)
diff --git a/Engine/Plugins/VirtualProduction/RemoteControl/Source/RemoteControlUI/Private/UI/Action/SRCActionPanelList.h b/Engine/Plugins/VirtualProduction/RemoteControl/Source/RemoteControlUI/Private/UI/Action/SRCActionPanelList.h
--- a/Engine/Plugins/VirtualProduction/RemoteControl/Source/RemoteControlUI/Private/UI/Action/SRCActionPanelList.h
+++ b/Engine/Plugins/VirtualProduction/RemoteControl/Source/RemoteControlUI/Private/UI/Action/SRCActionPanelList.h
@@ -1,9 +1,10 @@
 ﻿// Copyright Epic Games, Inc. All Rights Reserved.
 
 #pragma once
-#include "UI/BaseLogicUI/SRCLogicPanelListBase.h"
 
+#include "UI/BaseLogicUI/SRCLogicPanelListBase.h"
 
+struct FRCPanelStyle;
 class FRCBehaviourModel;
 class FRCActionModel;
 class ITableRow;
@@ -31,7 +32,20 @@ public:
 	/** Constructs this widget with InArgs */
 	void Construct(const FArguments& InArgs, const TSharedRef<SRCActionPanel> InActionPanel, TSharedPtr<FRCBehaviourModel> InBehaviourItem);
 
+	/** Returns true if the underlying list is valid and empty. */
+	virtual bool IsEmpty() const override;
+
+	/** Returns number of items in the list. */
+	virtual int32 Num() const override;
+
+	/** Whether the Actions List View currently has focus.*/
+	bool IsListFocused() const;
+
+	/** Deletes currently selected items from the list view*/
+	void DeleteSelectedPanelItem();
+
 private:
+
 	/** OnGenerateRow delegate for the Actions List View*/
 	TSharedRef<ITableRow> OnGenerateWidgetForList( TSharedPtr<FRCActionModel> InItem, const TSharedRef<STableViewBase>& OwnerTable );
 
@@ -41,20 +55,6 @@ private:
 	/** Responds to the removal of all actions. Rests UI state*/
 	void OnEmptyActions();
 
-private:
-	/** The parent Action Panel widget*/
-	TWeakPtr<SRCActionPanel> ActionPanelWeakPtr;
-	
-	/** The Behaviour (UI model) associated with us*/
-	TWeakPtr<FRCBehaviourModel> BehaviourItemWeakPtr;
-
-	/** List of Actions (UI model) active in this widget */
-	TArray<TSharedPtr<FRCActionModel>> ActionItems;
-	
-	/** List View widget for representing our Actions List*/
-	TSharedPtr<SListView<TSharedPtr<FRCActionModel>>> ListView;
-
-private:
 	/** Refreshes the list from the latest state of the data model*/
 	virtual void Reset() override;
 
@@ -67,12 +67,21 @@ private:
 	/** Removes the given Action UI model item from the list of UI models*/
 	virtual int32 RemoveModel(const TSharedPtr<FRCLogicModeBase> InModel) override;
 
-public:
-	/** Whether the Actions List View currently has focus.*/
-	bool IsListFocused() const;
+private:
 
-	/** Deletes currently selected items from the list view*/
-	void DeleteSelectedPanelItem();
+	/** The parent Action Panel widget*/
+	TWeakPtr<SRCActionPanel> ActionPanelWeakPtr;
+
+	/** The Behaviour (UI model) associated with us*/
+	TWeakPtr<FRCBehaviourModel> BehaviourItemWeakPtr;
+
+	/** List of Actions (UI model) active in this widget */
+	TArray<TSharedPtr<FRCActionModel>> ActionItems;
+
+	/** List View widget for representing our Actions List*/
+	TSharedPtr<SListView<TSharedPtr<FRCActionModel>>> ListView;
 
+	/** Panel Style reference. */
+	const FRCPanelStyle* RCPanelStyle;
 };
 
diff --git a/Engine/Plugins/VirtualProduction/RemoteControl/Source/RemoteControlUI/Private/UI/Action/SRCBehaviourDetails.cpp b/Engine/Plugins/VirtualProduction/RemoteControl/Source/RemoteControlUI/Private/UI/Action/SRCBehaviourDetails.cpp
new file mode 100644
--- /dev/null
+++ b/Engine/Plugins/VirtualProduction/RemoteControl/Source/RemoteControlUI/Private/UI/Action/SRCBehaviourDetails.cpp
@@ -0,0 +1,138 @@
+// Copyright Epic Games, Inc. All Rights Reserved.
+
+#include "SRCBehaviourDetails.h"
+
+#include "Action/RCActionContainer.h"
+#include "Action/RCFunctionAction.h"
+#include "Action/RCPropertyAction.h"
+#include "Behaviour/RCBehaviour.h"
+#include "Controller/RCController.h"
+#include "RemoteControlPreset.h"
+#include "SlateOptMacros.h"
+#include "SRCActionPanel.h"
+#include "Styling/CoreStyle.h"
+#include "Styling/RemoteControlStyles.h"
+#include "UI/Behaviour/RCBehaviourModel.h"
+#include "UI/BaseLogicUI/RCLogicHelpers.h"
+#include "UI/RCUIHelpers.h"
+#include "UI/RemoteControlPanelStyle.h"
+#include "UI/SRemoteControlPanel.h"
+#include "Widgets/Input/SButton.h"
+#include "Widgets/Layout/SSpacer.h"
+#include "Widgets/Layout/SSeparator.h"
+
+#define LOCTEXT_NAMESPACE "SRCBehaviourDetails"
+
+BEGIN_SLATE_FUNCTION_BUILD_OPTIMIZATION
+
+void SRCBehaviourDetails::Construct(const FArguments& InArgs, TSharedRef<SRCActionPanel> InActionPanel, TSharedRef<FRCBehaviourModel> InBehaviourItem)
+{
+	RCPanelStyle = &FRemoteControlPanelStyle::Get()->GetWidgetStyle<FRCPanelStyle>("RemoteControlPanel.MinorPanel");
+
+	ActionPanelWeakPtr = InActionPanel;
+	BehaviourItemWeakPtr = InBehaviourItem;
+
+	URCBehaviour* Behaviour = BehaviourItemWeakPtr.Pin()->GetBehaviour();
+	const FText BehaviourDisplayName = Behaviour->GetDisplayName().ToUpper();
+	const FText BehaviourDescription = Behaviour->GetBehaviorDescription();
+
+	const FSlateFontInfo& FontBehaviorDesc = FRemoteControlPanelStyle::Get()->GetFontStyle("RemoteControlPanel.Behaviours.BehaviorDescription");
+
+	FLinearColor TypeColor;
+	FString TypeDisplayName;
+	if (URCController* Controller = Behaviour->ControllerWeakPtr.Get())
+	{
+		TypeColor = UE::RCUIHelpers::GetFieldClassTypeColor(Behaviour->ControllerWeakPtr->GetProperty());
+		TypeDisplayName = FName::NameToDisplayString(UE::RCUIHelpers::GetFieldClassDisplayName(Behaviour->ControllerWeakPtr->GetProperty()).ToString(), false);
+	}
+
+	TSharedRef<SWidget> BehaviourDetailsWidget = InBehaviourItem->GetBehaviourDetailsWidget();
+	
+	ChildSlot
+		.Padding(RCPanelStyle->PanelPadding)
+		[
+			SNew(SVerticalBox)
+
+			+ SVerticalBox::Slot()
+			.HAlign(HAlign_Fill)
+			.Padding(4.f)
+			.AutoHeight()
+			[
+				SNew(SHorizontalBox)
+
+				// Behaviour Name
+				+ SHorizontalBox::Slot()
+				.HAlign(HAlign_Left)
+				.VAlign(VAlign_Center)
+				.FillWidth(1.f)
+				.Padding(4.f)
+				[
+					SNew(STextBlock)
+					.ColorAndOpacity(FLinearColor::White)
+					.Text(BehaviourDisplayName)
+					.TextStyle(&RCPanelStyle->SectionHeaderTextStyle)
+				]
+
+				+ SHorizontalBox::Slot()
+				.VAlign(VAlign_Center)
+				.AutoWidth()
+				.Padding(4.f, 0.f)
+				[
+					SNew(SBorder)
+					.BorderImage(FAppStyle::Get().GetBrush("SettingsEditor.CheckoutWarningBorder"))
+					.Padding(FMargin(8.f, 4.f))
+					.BorderBackgroundColor(TypeColor)
+					[
+						SNew(STextBlock)
+						.ColorAndOpacity(FLinearColor::Black)
+						.Text(FText::FromString(TypeDisplayName))
+						.TextStyle(&RCPanelStyle->PanelTextStyle)
+					]
+				]
+			]
+
+			+ SVerticalBox::Slot()
+			.Padding(8.f, 4.f)
+			.VAlign(VAlign_Top)
+			.FillHeight(1.f)
+			[
+				SNew(STextBlock)
+				.ColorAndOpacity(FSlateColor::UseForeground())
+				.Text(BehaviourDescription)
+				.TextStyle(&RCPanelStyle->PanelTextStyle)
+				.AutoWrapText(true) // note: use of autowidth on this slot prevents wrapping from working!
+			]
+			
+			// Border (separating Header and "If Value Is" panel)
+			+ SVerticalBox::Slot()
+			.Padding(2.f, 4.f)
+			.AutoHeight()
+			[
+				SNew(SSeparator)
+				.SeparatorImage(FAppStyle::Get().GetBrush("Separator"))
+				.Thickness(2.f)
+				.Orientation(EOrientation::Orient_Horizontal)
+				.Visibility_Lambda([BehaviourDetailsWidget]() { return BehaviourDetailsWidget != SNullWidget::NullWidget ? EVisibility::Visible : EVisibility::Collapsed; })
+			]
+
+			// Behaviour Specific Details Panel
+			+ SVerticalBox::Slot()
+			.Padding(8.f, 4.f)
+			.AutoHeight()
+			[
+				BehaviourDetailsWidget
+			]
+
+			// Spacer to fill the gap.
+			+ SVerticalBox::Slot()
+			.Padding(0)
+			.FillHeight(1.f)
+			[
+				SNew(SSpacer)
+			]
+		];
+}
+
+END_SLATE_FUNCTION_BUILD_OPTIMIZATION
+
+#undef LOCTEXT_NAMESPACE
\ No newline at end of file
diff --git a/Engine/Plugins/VirtualProduction/RemoteControl/Source/RemoteControlUI/Private/UI/Action/SRCBehaviourDetails.h b/Engine/Plugins/VirtualProduction/RemoteControl/Source/RemoteControlUI/Private/UI/Action/SRCBehaviourDetails.h
new file mode 100644
--- /dev/null
+++ b/Engine/Plugins/VirtualProduction/RemoteControl/Source/RemoteControlUI/Private/UI/Action/SRCBehaviourDetails.h
@@ -0,0 +1,35 @@
+// Copyright Epic Games, Inc. All Rights Reserved.
+
+#pragma once
+#include "RemoteControlField.h"
+#include "Widgets/SCompoundWidget.h"
+
+struct FRCPanelStyle;
+class SRCActionPanel;
+class FRCBehaviourModel;
+
+/**
+ * A custom widget dedicated to show the details of the selected behaviour.
+ */
+class REMOTECONTROLUI_API SRCBehaviourDetails : public SCompoundWidget
+{
+public:
+	SLATE_BEGIN_ARGS(SRCBehaviourDetails)
+		{
+		}
+
+	SLATE_END_ARGS()
+
+	/** Constructs this widget with InArgs */
+	void Construct(const FArguments& InArgs, TSharedRef<SRCActionPanel> InActionPanel, TSharedRef<FRCBehaviourModel> InBehaviourItem);
+
+private:
+	/** The parent Action panel UI widget holding this Header*/
+	TWeakPtr<SRCActionPanel> ActionPanelWeakPtr;
+	
+	/** The active Behaviour (UI Model) associated with the Actions being shown*/
+	TWeakPtr<FRCBehaviourModel> BehaviourItemWeakPtr;
+
+	/** Panel Style reference. */
+	const FRCPanelStyle* RCPanelStyle;
+};
diff --git a/Engine/Plugins/VirtualProduction/RemoteControl/Source/RemoteControlUI/Private/UI/BaseLogicUI/SRCLogicPanelBase.h b/Engine/Plugins/VirtualProduction/RemoteControl/Source/RemoteControlUI/Private/UI/BaseLogicUI/SRCLogicPanelBase.h
--- a/Engine/Plugins/VirtualProduction/RemoteControl/Source/RemoteControlUI/Private/UI/BaseLogicUI/SRCLogicPanelBase.h
+++ b/Engine/Plugins/VirtualProduction/RemoteControl/Source/RemoteControlUI/Private/UI/BaseLogicUI/SRCLogicPanelBase.h
@@ -37,6 +37,11 @@ public:
 	/** "Delete Item" UI command implementation for panels*/
 	virtual void DeleteSelectedPanelItem() = 0;
 
+protected:
+
+	/** Warns user before deleting all items in a panel. */
+	virtual FReply RequestDeleteAllItems() = 0;
+
 protected:
 	/** The parent Remote Control Panel widget*/
 	TWeakPtr<SRemoteControlPanel> PanelWeakPtr;
diff --git a/Engine/Plugins/VirtualProduction/RemoteControl/Source/RemoteControlUI/Private/UI/BaseLogicUI/SRCLogicPanelHeaderBase.cpp b/Engine/Plugins/VirtualProduction/RemoteControl/Source/RemoteControlUI/Private/UI/BaseLogicUI/SRCLogicPanelHeaderBase.cpp
deleted file mode 100644
diff --git a/Engine/Plugins/VirtualProduction/RemoteControl/Source/RemoteControlUI/Private/UI/BaseLogicUI/SRCLogicPanelHeaderBase.h b/Engine/Plugins/VirtualProduction/RemoteControl/Source/RemoteControlUI/Private/UI/BaseLogicUI/SRCLogicPanelHeaderBase.h
deleted file mode 100644
diff --git a/Engine/Plugins/VirtualProduction/RemoteControl/Source/RemoteControlUI/Private/UI/BaseLogicUI/SRCLogicPanelListBase.h b/Engine/Plugins/VirtualProduction/RemoteControl/Source/RemoteControlUI/Private/UI/BaseLogicUI/SRCLogicPanelListBase.h
--- a/Engine/Plugins/VirtualProduction/RemoteControl/Source/RemoteControlUI/Private/UI/BaseLogicUI/SRCLogicPanelListBase.h
+++ b/Engine/Plugins/VirtualProduction/RemoteControl/Source/RemoteControlUI/Private/UI/BaseLogicUI/SRCLogicPanelListBase.h
@@ -26,9 +26,16 @@ public:
 	/** Constructs this widget with InArgs */
 	void Construct(const FArguments& InArgs);
 
+	/** Returns true if the underlying list is valid and empty. */
+	virtual bool IsEmpty() const = 0;
+	
+	/** Returns number of items in the list. */
+	virtual int32 Num() const = 0;
+
 private:
+
 	/** Refreshes the list from the latest state of the model*/
-	virtual void Reset()  = 0;
+	virtual void Reset() = 0;
 
 	/** Handles broadcasting of a successful remove item operation.
 	* This is handled uniquely by each type of child list widget*/
@@ -41,6 +48,7 @@ private:
 	virtual int32 RemoveModel(const TSharedPtr<FRCLogicModeBase> InModel) = 0;
 
 protected:
+
 	/** Helper function for handling common Delete Item functionality across all child panels (Actions/Behaviours/Controllers)
 	* Currently invoked from each Panel List child class with appropriate template class*/
 	template<class T>
diff --git a/Engine/Plugins/VirtualProduction/RemoteControl/Source/RemoteControlUI/Private/UI/Behaviour/SRCBehaviourPanel.cpp b/Engine/Plugins/VirtualProduction/RemoteControl/Source/RemoteControlUI/Private/UI/Behaviour/SRCBehaviourPanel.cpp
--- a/Engine/Plugins/VirtualProduction/RemoteControl/Source/RemoteControlUI/Private/UI/Behaviour/SRCBehaviourPanel.cpp
+++ b/Engine/Plugins/VirtualProduction/RemoteControl/Source/RemoteControlUI/Private/UI/Behaviour/SRCBehaviourPanel.cpp
@@ -1,27 +1,55 @@
 ﻿// Copyright Epic Games, Inc. All Rights Reserved.
 
 #include "SRCBehaviourPanel.h"
-#include "SRCBehaviourPanelHeader.h"
 #include "SRCBehaviourPanelList.h"
 
+#include "Behaviour/RCBehaviourBlueprintNode.h"
+#include "Behaviour/RCBehaviour.h"
+#include "Behaviour/RCBehaviourNode.h"
+#include "Controller/RCController.h"
+#include "Framework/MultiBox/MultiBoxBuilder.h"
+#include "Kismet2/KismetEditorUtilities.h"
+#include "Misc/MessageDialog.h"
+#include "RemoteControlPreset.h"
+
 #include "SlateOptMacros.h"
+#include "Styling/RemoteControlStyles.h"
+
+#include "UI/Panels/SRCDockPanel.h"
+#include "UI/RemoteControlPanelStyle.h"
 #include "UI/SRemoteControlPanel.h"
+#include "UI/Controller/RCControllerModel.h"
+#include "Widgets/Input/SComboButton.h"
 
 #define LOCTEXT_NAMESPACE "SRCBehaviourPanel"
 
+TSharedRef<SBox> SRCBehaviourPanel::NoneSelectedWidget = SNew(SBox)
+			.Padding(0.f)
+			.VAlign(VAlign_Center)
+			.HAlign(HAlign_Center)
+			[
+				SNew(STextBlock)
+				.Text(LOCTEXT("NoneSelected", "Select a controller to view its behaviors."))
+				.TextStyle(&FAppStyle::GetWidgetStyle<FTextBlockStyle>("NormalText"))
+				.Justification(ETextJustify::Center)
+			];
+
 BEGIN_SLATE_FUNCTION_BUILD_OPTIMIZATION
 
 void SRCBehaviourPanel::Construct(const FArguments& InArgs, const TSharedRef<SRemoteControlPanel>& InPanel)
 {
 	SRCLogicPanelBase::Construct(SRCLogicPanelBase::FArguments(), InPanel);
 
+	RCPanelStyle = &FRemoteControlPanelStyle::Get()->GetWidgetStyle<FRCPanelStyle>("RemoteControlPanel.MinorPanel");
+
 	WrappedBoxWidget = SNew(SBox);
 	UpdateWrappedWidget();
 	
 	ChildSlot
-	[
-		WrappedBoxWidget.ToSharedRef()
-	];
+		.Padding(RCPanelStyle->PanelPadding)
+		[
+			WrappedBoxWidget.ToSharedRef()
+		];
 
 	// Register delegates
 	InPanel->OnControllerSelectionChanged.AddSP(this, &SRCBehaviourPanel::OnControllerSelectionChanged);
@@ -31,38 +59,162 @@ END_SLATE_FUNCTION_BUILD_OPTIMIZATION
 
 void SRCBehaviourPanel::OnControllerSelectionChanged(TSharedPtr<FRCControllerModel> InControllerItem)
 {
-	UpdateWrappedWidget(InControllerItem);
 	SelectedControllerItemWeakPtr = InControllerItem;
+	UpdateWrappedWidget(InControllerItem);
 }
 
 void SRCBehaviourPanel::UpdateWrappedWidget(TSharedPtr<FRCControllerModel> InControllerItem)
 {
 	if (InControllerItem.IsValid())
 	{
-		WrappedBoxWidget->SetContent(
-			SNew(SVerticalBox)
-			// Header
-			+ SVerticalBox::Slot()
-			.Padding(5.0f, 5.0f)
-			.AutoHeight()
-			[
-				SNew(SRCBehaviourPanelHeader, SharedThis(this), InControllerItem)
-			]
-			// List
-			+ SVerticalBox::Slot()
+		// Behaviour Dock Panel
+		TSharedPtr<SRCMinorPanel> BehaviourDockPanel = SNew(SRCMinorPanel)
+			.HeaderLabel(LOCTEXT("BehavioursLabel", "Behavior"))
 			[
 				SAssignNew(BehaviourPanelList, SRCBehaviourPanelList, SharedThis(this), InControllerItem)
+			];
+
+		// Add New Behaviour Button
+		const TSharedRef<SWidget> AddNewBehaviourButton = SNew(SComboButton)
+			.AddMetaData<FTagMetaData>(FTagMetaData(TEXT("Add Behavior")))
+			.HAlign(HAlign_Center)
+			.VAlign(VAlign_Center)
+			.ButtonStyle(&RCPanelStyle->FlatButtonStyle)
+			.ForegroundColor(FSlateColor::UseForeground())
+			.CollapseMenuOnParentFocus(true)
+			.HasDownArrow(false)
+			.ContentPadding(FMargin(4.f, 2.f))
+			.ButtonContent()
+			[
+				SNew(SBox)
+				.WidthOverride(RCPanelStyle->IconSize.X)
+				.HeightOverride(RCPanelStyle->IconSize.Y)
+				[
+					SNew(SImage)
+					.ColorAndOpacity(FSlateColor::UseForeground())
+					.Image(FAppStyle::GetBrush("Icons.PlusCircle"))
+				]
 			]
-		);
+			.MenuContent()
+			[
+				GetBehaviourMenuContentWidget()
+			];
+
+		// Empty All Button
+		TSharedRef<SWidget> EmptyAllButton = SNew(SButton)
+			.AddMetaData<FTagMetaData>(FTagMetaData(TEXT("Empty Behaviors")))
+			.HAlign(HAlign_Center)
+			.VAlign(VAlign_Center)
+			.ForegroundColor(FSlateColor::UseForeground())
+			.ButtonStyle(&RCPanelStyle->FlatButtonStyle)
+			.ToolTipText(LOCTEXT("EmptyAllToolTip", "Deletes all the behaviors."))
+			.OnClicked(this, &SRCBehaviourPanel::RequestDeleteAllItems)
+			.Visibility_Lambda([this]() { return BehaviourPanelList.IsValid() && !BehaviourPanelList->IsEmpty() ? EVisibility::Visible : EVisibility::Collapsed; })
+			[
+				SNew(SBox)
+				.WidthOverride(RCPanelStyle->IconSize.X)
+				.HeightOverride(RCPanelStyle->IconSize.Y)
+				[
+					SNew(SImage)
+					.ColorAndOpacity(FSlateColor::UseForeground())
+					.Image(FAppStyle::GetBrush("Icons.Delete"))
+				]
+			];
+
+		BehaviourDockPanel->AddHeaderToolbarItem(EToolbar::Left, AddNewBehaviourButton);
+		BehaviourDockPanel->AddHeaderToolbarItem(EToolbar::Right, EmptyAllButton);
+
+		WrappedBoxWidget->SetContent(BehaviourDockPanel.ToSharedRef());
 	}
 	else
 	{
-		WrappedBoxWidget->SetContent(
-			SNew(STextBlock).Text(LOCTEXT("Unselected", "No Controller Selected"))
-		);
+		WrappedBoxWidget->SetContent(NoneSelectedWidget);
 	}
 }
 
+TSharedRef<SWidget> SRCBehaviourPanel::GetBehaviourMenuContentWidget()
+{
+	constexpr bool bShouldCloseWindowAfterMenuSelection = true;
+	FMenuBuilder MenuBuilder(bShouldCloseWindowAfterMenuSelection, nullptr);
+
+	if (const TSharedPtr<FRCControllerModel> ControllerItem = SelectedControllerItemWeakPtr.Pin())
+	{
+		if (URCController* Controller = Cast<URCController>(ControllerItem->GetVirtualProperty()))
+		{
+			for (TObjectIterator<UClass> It; It; ++It)
+			{
+				UClass* Class = *It;
+
+				if (Class->IsChildOf(URCBehaviourNode::StaticClass()))
+				{
+					const bool bIsClassInstantiatable = Class->HasAnyClassFlags(CLASS_Deprecated | CLASS_NewerVersionExists | CLASS_Abstract) || FKismetEditorUtilities::IsClassABlueprintSkeleton(Class);
+					if (bIsClassInstantiatable)
+					{
+						continue;
+					}
+
+					if (Class->IsInBlueprint())
+					{
+						if (Class->GetSuperClass() != URCBehaviourBlueprintNode::StaticClass())
+						{
+							continue;
+						}
+					}
+
+					URCBehaviour* Behaviour = Controller->CreateBehaviour(Class);
+					if (!Behaviour)
+					{
+						continue;
+					}
+
+					FUIAction Action(FExecuteAction::CreateSP(this, &SRCBehaviourPanel::OnAddBehaviourClicked, Class));
+					MenuBuilder.AddMenuEntry(
+						FText::Format(LOCTEXT("AddBehaviourNode", "{0}"), Behaviour->GetDisplayName()),
+						FText::Format(LOCTEXT("AddBehaviourNodeTooltip", "{0}"), Behaviour->GetBehaviorDescription()),
+						FSlateIcon(),
+						Action);
+				}
+			}
+		}
+	}
+
+	return MenuBuilder.MakeWidget();
+}
+
+void SRCBehaviourPanel::OnAddBehaviourClicked(UClass* InClass)
+{
+	if (const TSharedPtr<FRCControllerModel> ControllerItem = SelectedControllerItemWeakPtr.Pin())
+	{
+		if (URCController* Controller = Cast<URCController>(ControllerItem->GetVirtualProperty()))
+		{
+			URCBehaviour* NewBehaviour = Controller->AddBehaviour(InClass);
+
+			if (const TSharedPtr<SRemoteControlPanel> RemoteControlPanel = GetRemoteControlPanel())
+			{
+				RemoteControlPanel->OnBehaviourAdded.Broadcast(NewBehaviour);
+			}
+		}
+	}
+}
+
+FReply SRCBehaviourPanel::OnClickEmptyButton()
+{
+	if (const TSharedPtr<FRCControllerModel> ControllerItem = SelectedControllerItemWeakPtr.Pin())
+	{
+		if (URCController* Controller = Cast<URCController>(ControllerItem->GetVirtualProperty()))
+		{
+			Controller->EmptyBehaviours();
+		}
+	}
+
+	if (const TSharedPtr<SRemoteControlPanel> RemoteControlPanel = GetRemoteControlPanel())
+	{
+		RemoteControlPanel->OnEmptyBehaviours.Broadcast();
+	}
+
+	return FReply::Handled();
+}
+
 bool SRCBehaviourPanel::IsListFocused() const
 {
 	return BehaviourPanelList.IsValid() && BehaviourPanelList->IsListFocused();
@@ -73,4 +225,23 @@ void SRCBehaviourPanel::DeleteSelectedPanelItem()
 	BehaviourPanelList->DeleteSelectedPanelItem();
 }
 
+FReply SRCBehaviourPanel::RequestDeleteAllItems()
+{
+	if (!BehaviourPanelList.IsValid())
+	{
+		return FReply::Unhandled();
+	}
+
+	const FText WarningMessage = FText::Format(LOCTEXT("DeleteAllWarning", "You are about to delete '{0}' behaviors. This action might not be undone.\nAre you sure you want to proceed?"), BehaviourPanelList->Num());
+
+	EAppReturnType::Type UserResponse = FMessageDialog::Open(EAppMsgType::YesNo, WarningMessage);
+
+	if (UserResponse == EAppReturnType::Yes)
+	{
+		return OnClickEmptyButton();
+	}
+
+	return FReply::Handled();
+}
+
 #undef LOCTEXT_NAMESPACE
\ No newline at end of file
diff --git a/Engine/Plugins/VirtualProduction/RemoteControl/Source/RemoteControlUI/Private/UI/Behaviour/SRCBehaviourPanel.h b/Engine/Plugins/VirtualProduction/RemoteControl/Source/RemoteControlUI/Private/UI/Behaviour/SRCBehaviourPanel.h
--- a/Engine/Plugins/VirtualProduction/RemoteControl/Source/RemoteControlUI/Private/UI/Behaviour/SRCBehaviourPanel.h
+++ b/Engine/Plugins/VirtualProduction/RemoteControl/Source/RemoteControlUI/Private/UI/Behaviour/SRCBehaviourPanel.h
@@ -2,11 +2,13 @@
 
 #pragma once
 
+#include "Behaviour/RCBehaviourNode.h"
 #include "UI/BaseLogicUI/SRCLogicPanelBase.h"
 
+struct FRCPanelStyle;
 class FRCControllerModel;
 class SBox;
-
+class SRCBehaviourPanel;
 /*
 * ~ SRCBehaviourPanel ~
 *
@@ -25,7 +27,19 @@ public:
 	/** Constructs this widget with InArgs */
 	void Construct(const FArguments& InArgs, const TSharedRef<SRemoteControlPanel>& InPanel);
 
+	/** Whether the Behaviour list widget currently has focus. Used for Delete Item UI command */
+	bool IsListFocused() const;
+
+	/** Delete Item UI command implementation for this panel */
+	virtual void DeleteSelectedPanelItem() override;
+
+protected:
+
+	/** Warns user before deleting all items in a panel. */
+	virtual FReply RequestDeleteAllItems() override;
+
 private:
+	
 	/**
 	* Controller list selection change listener.
 	* Updates the list of behaviours from newly selected Controller
@@ -35,7 +49,17 @@ private:
 	/* Rebuilds the Behaviour Panel for a newly selected Controller*/
 	void UpdateWrappedWidget(TSharedPtr<FRCControllerModel> InControllerItem = nullptr);
 
+	/** Builds a menu containing the list of all possible Behaviours */
+	TSharedRef<SWidget> GetBehaviourMenuContentWidget();
+
+	/** Handles click event for Add Behaviour button*/
+	void OnAddBehaviourClicked(UClass* InClass);
+
+	/** Handles click event for "Empty" button; clears all Behaviours from the panel*/
+	FReply OnClickEmptyButton();
+
 private:
+
 	/** The parent Controller that this Behaviour panel is associated with */
 	TWeakPtr<FRCControllerModel> SelectedControllerItemWeakPtr = nullptr;
 	
@@ -45,11 +69,9 @@ private:
 	/** Widget representing List of Behaviours */
 	TSharedPtr<class SRCBehaviourPanelList> BehaviourPanelList;
 
-public:
-	/** Whether the Behaviour list widget currently has focus. Used for Delete Item UI command */
-	bool IsListFocused() const;
-
-	/** Delete Item UI command implementation for this panel */
-	virtual void DeleteSelectedPanelItem() override;
+	/** Helper widget for behavior details. */
+	static TSharedRef<SBox> NoneSelectedWidget;
 
+	/** Panel Style reference. */
+	const FRCPanelStyle* RCPanelStyle;
 };
diff --git a/Engine/Plugins/VirtualProduction/RemoteControl/Source/RemoteControlUI/Private/UI/Behaviour/SRCBehaviourPanelHeader.cpp b/Engine/Plugins/VirtualProduction/RemoteControl/Source/RemoteControlUI/Private/UI/Behaviour/SRCBehaviourPanelHeader.cpp
deleted file mode 100644
diff --git a/Engine/Plugins/VirtualProduction/RemoteControl/Source/RemoteControlUI/Private/UI/Behaviour/SRCBehaviourPanelHeader.h b/Engine/Plugins/VirtualProduction/RemoteControl/Source/RemoteControlUI/Private/UI/Behaviour/SRCBehaviourPanelHeader.h
deleted file mode 100644
diff --git a/Engine/Plugins/VirtualProduction/RemoteControl/Source/RemoteControlUI/Private/UI/Behaviour/SRCBehaviourPanelList.cpp b/Engine/Plugins/VirtualProduction/RemoteControl/Source/RemoteControlUI/Private/UI/Behaviour/SRCBehaviourPanelList.cpp
--- a/Engine/Plugins/VirtualProduction/RemoteControl/Source/RemoteControlUI/Private/UI/Behaviour/SRCBehaviourPanelList.cpp
+++ b/Engine/Plugins/VirtualProduction/RemoteControl/Source/RemoteControlUI/Private/UI/Behaviour/SRCBehaviourPanelList.cpp
@@ -7,9 +7,19 @@
 #include "RemoteControlPreset.h"
 #include "SlateOptMacros.h"
 #include "SRCBehaviourPanel.h"
+#include "Styling/RemoteControlStyles.h"
 #include "UI/Behaviour/Builtin/RCBehaviourIsEqualModel.h"
+#include "UI/RemoteControlPanelStyle.h"
 #include "UI/SRemoteControlPanel.h"
 #include "UI/Controller/RCControllerModel.h"
+#include "Widgets/Views/SHeaderRow.h"
+
+#define LOCTEXT_NAMESPACE "RemoteControlPanelBehavioursList"
+
+namespace FRemoteControlBehaviourColumns
+{
+	const FName Behaviours = TEXT("Behaviors");
+}
 
 BEGIN_SLATE_FUNCTION_BUILD_OPTIMIZATION
 
@@ -20,10 +30,22 @@ void SRCBehaviourPanelList::Construct(const FArguments& InArgs, TSharedRef<SRCBe
 	BehaviourPanelWeakPtr = InBehaviourPanel;
 	ControllerItemWeakPtr = InControllerItem;
 	
+	RCPanelStyle = &FRemoteControlPanelStyle::Get()->GetWidgetStyle<FRCPanelStyle>("RemoteControlPanel.MinorPanel");
+
 	ListView = SNew(SListView<TSharedPtr<FRCBehaviourModel>>)
-		.ListItemsSource( &BehaviourItems )
+		.ListItemsSource(&BehaviourItems)
 		.OnSelectionChanged(this, &SRCBehaviourPanelList::OnTreeSelectionChanged)
-		.OnGenerateRow(this, &SRCBehaviourPanelList::OnGenerateWidgetForList );
+		.OnGenerateRow(this, &SRCBehaviourPanelList::OnGenerateWidgetForList)
+		.ListViewStyle(&RCPanelStyle->TableViewStyle)
+		.HeaderRow(
+			SNew(SHeaderRow)
+			.Style(&RCPanelStyle->HeaderRowStyle)
+
+			+ SHeaderRow::Column(FRemoteControlBehaviourColumns::Behaviours)
+			.DefaultLabel(LOCTEXT("RCBehaviourColumnHeader", "Behaviors"))
+			.FillWidth(1.f)
+			.HeaderContentPadding(RCPanelStyle->HeaderRowPadding)
+		);
 	
 	ChildSlot
 	[
@@ -39,6 +61,16 @@ void SRCBehaviourPanelList::Construct(const FArguments& InArgs, TSharedRef<SRCBe
 	Reset();
 }
 
+bool SRCBehaviourPanelList::IsEmpty() const
+{
+	return BehaviourItems.IsEmpty();
+}
+
+int32 SRCBehaviourPanelList::Num() const
+{
+	return BehaviourItems.Num();
+}
+
 END_SLATE_FUNCTION_BUILD_OPTIMIZATION
 
 void SRCBehaviourPanelList::Reset()
@@ -70,24 +102,22 @@ TSharedRef<ITableRow> SRCBehaviourPanelList::OnGenerateWidgetForList(TSharedPtr<
 	const TSharedRef<STableViewBase>& OwnerTable)
 {
 	return SNew(STableRow<TSharedPtr<FString>>, OwnerTable)
-	[
-		InItem->GetWidget()
-	];
+		.Style(&RCPanelStyle->TableRowStyle)
+		[
+			InItem->GetWidget()
+		];
 }
 
 void SRCBehaviourPanelList::OnTreeSelectionChanged(TSharedPtr<FRCBehaviourModel> InItem, ESelectInfo::Type)
 {
 	if (const TSharedPtr<SRemoteControlPanel> RemoteControlPanel = BehaviourPanelWeakPtr.Pin()->GetRemoteControlPanel())
 	{
-		if (InItem.IsValid())
+		if (InItem != SelectedBehaviourItemWeakPtr.Pin())
 		{
-			if (InItem != SelectedBehaviourItemWeakPtr.Pin())
-			{
-				RemoteControlPanel->OnBehaviourSelectionChanged.Broadcast(InItem);
-			}
-		}
+			SelectedBehaviourItemWeakPtr = InItem;
 
-		SelectedBehaviourItemWeakPtr = InItem;
+			RemoteControlPanel->OnBehaviourSelectionChanged.Broadcast(InItem);
+		}
 
 		if (const TSharedPtr<FRCControllerModel> ControllerItem = ControllerItemWeakPtr.Pin())
 		{
@@ -157,3 +187,5 @@ void SRCBehaviourPanelList::DeleteSelectedPanelItem()
 {
 	DeleteItemFromLogicPanel<FRCBehaviourModel>(BehaviourItems, ListView->GetSelectedItems());
 }
+
+#undef LOCTEXT_NAMESPACE
diff --git a/Engine/Plugins/VirtualProduction/RemoteControl/Source/RemoteControlUI/Private/UI/Behaviour/SRCBehaviourPanelList.h b/Engine/Plugins/VirtualProduction/RemoteControl/Source/RemoteControlUI/Private/UI/Behaviour/SRCBehaviourPanelList.h
--- a/Engine/Plugins/VirtualProduction/RemoteControl/Source/RemoteControlUI/Private/UI/Behaviour/SRCBehaviourPanelList.h
+++ b/Engine/Plugins/VirtualProduction/RemoteControl/Source/RemoteControlUI/Private/UI/Behaviour/SRCBehaviourPanelList.h
@@ -4,6 +4,7 @@
 
 #include "UI/BaseLogicUI/SRCLogicPanelListBase.h"
 
+struct FRCPanelStyle;
 class URCBehaviour;
 class FRCControllerModel;
 class FRCBehaviourModel;
@@ -31,8 +32,21 @@ public:
 
 	/** Constructs this widget with InArgs */
 	void Construct(const FArguments& InArgs, const TSharedRef<SRCBehaviourPanel> InBehaviourPanel, TSharedPtr<FRCControllerModel> InControllerItem);
+	
+	/** Returns true if the underlying list is valid and empty. */
+	virtual bool IsEmpty() const override;
+	
+	/** Returns number of items in the list. */
+	virtual int32 Num() const override;
+
+	/** Whether the Behaviours List View currently has focus.*/
+	bool IsListFocused() const;
+
+	/** Deletes currently selected items from the list view*/
+	void DeleteSelectedPanelItem();
 
 private:
+
 	/** OnGenerateRow delegate for the Behaviours List View */
 	TSharedRef<ITableRow> OnGenerateWidgetForList( TSharedPtr<FRCBehaviourModel> InItem, const TSharedRef<STableViewBase>& OwnerTable );
 	
@@ -45,21 +59,6 @@ private:
 	/** Responds to the removal of all Behaviours. Rests UI state */
 	void OnEmptyBehaviours();
 
-	/** The parent Behaviour Panel widget */
-	TWeakPtr<SRCBehaviourPanel> BehaviourPanelWeakPtr;
-	
-	/** The Controller (UI model) associated with us */
-	TWeakPtr<FRCControllerModel> ControllerItemWeakPtr;
-	
-	/** The currently selected Behaviour item (UI model) */
-	TWeakPtr<FRCBehaviourModel> SelectedBehaviourItemWeakPtr;
-
-	/** List of Behaviours (UI models) active in this widget */
-	TArray<TSharedPtr<FRCBehaviourModel>> BehaviourItems;
-	
-	/** List View widget for representing our Behaivours List*/
-	TSharedPtr<SListView<TSharedPtr<FRCBehaviourModel>>> ListView;
-
 	/** Refreshes the list from the latest state of the model*/
 	virtual void Reset() override;
 
@@ -72,11 +71,24 @@ private:
 	/** Removes the given Behaviour UI model item from the list of UI models*/
 	virtual int32 RemoveModel(const TSharedPtr<FRCLogicModeBase> InModel) override;
 
-public:
-	/** Whether the Behaviours List View currently has focus.*/
-	bool IsListFocused() const;
+private:
 
-	/** Deletes currently selected items from the list view*/
-	void DeleteSelectedPanelItem();
+	/** The parent Behaviour Panel widget */
+	TWeakPtr<SRCBehaviourPanel> BehaviourPanelWeakPtr;
+
+	/** The Controller (UI model) associated with us */
+	TWeakPtr<FRCControllerModel> ControllerItemWeakPtr;
+
+	/** The currently selected Behaviour item (UI model) */
+	TWeakPtr<FRCBehaviourModel> SelectedBehaviourItemWeakPtr;
+
+	/** List of Behaviours (UI models) active in this widget */
+	TArray<TSharedPtr<FRCBehaviourModel>> BehaviourItems;
+
+	/** List View widget for representing our Behaivours List*/
+	TSharedPtr<SListView<TSharedPtr<FRCBehaviourModel>>> ListView;
+
+	/** Panel Style reference. */
+	const FRCPanelStyle* RCPanelStyle;
 };
 
diff --git a/Engine/Plugins/VirtualProduction/RemoteControl/Source/RemoteControlUI/Private/UI/Controller/SRCControllerPanel.cpp b/Engine/Plugins/VirtualProduction/RemoteControl/Source/RemoteControlUI/Private/UI/Controller/SRCControllerPanel.cpp
--- a/Engine/Plugins/VirtualProduction/RemoteControl/Source/RemoteControlUI/Private/UI/Controller/SRCControllerPanel.cpp
+++ b/Engine/Plugins/VirtualProduction/RemoteControl/Source/RemoteControlUI/Private/UI/Controller/SRCControllerPanel.cpp
@@ -2,42 +2,120 @@
 
 #include "SRCControllerPanel.h"
 
+#include "Controller/RCController.h"
+#include "Framework/MultiBox/MultiBoxBuilder.h"
+#include "Misc/MessageDialog.h"
 #include "RCControllerModel.h"
+#include "RCVirtualPropertyContainer.h"
+#include "RemoteControlLogicConfig.h"
+#include "RemoteControlPreset.h"
 #include "SlateOptMacros.h"
-#include "SRCControllerPanelHeader.h"
 #include "SRCControllerPanelList.h"
+#include "Styling/RemoteControlStyles.h"
 #include "UI/BaseLogicUI/RCLogicModeBase.h"
+#include "UI/Panels/SRCDockPanel.h"
+#include "UI/RemoteControlPanelStyle.h"
 #include "UI/SRemoteControlPanel.h"
 #include "Widgets/SBoxPanel.h"
+#include "Widgets/Input/SComboButton.h"
 #include "Widgets/Views/SListView.h"
 
+#define LOCTEXT_NAMESPACE "SRCControllerPanel"
+
 BEGIN_SLATE_FUNCTION_BUILD_OPTIMIZATION
 
 void SRCControllerPanel::Construct(const FArguments& InArgs, const TSharedRef<SRemoteControlPanel>& InPanel)
 {
 	SRCLogicPanelBase::Construct(SRCLogicPanelBase::FArguments(), InPanel);
 	
-	ChildSlot
-	[
-		SNew(SVerticalBox)
-		// Header
-		+ SVerticalBox::Slot()
-		.Padding(FMargin(5.f))
-		.AutoHeight()
-		[
-			SNew(SRCControllerPanelHeader, SharedThis(this))
-		]
-		// List
-		+ SVerticalBox::Slot()
-		.Padding(FMargin(5.f))
+	RCPanelStyle = &FRemoteControlPanelStyle::Get()->GetWidgetStyle<FRCPanelStyle>("RemoteControlPanel.MinorPanel");
+
+	// Controller Dock Panel
+	TSharedPtr<SRCMinorPanel> ControllerDockPanel = SNew(SRCMinorPanel)
+		.HeaderLabel(LOCTEXT("ControllersLabel", "Controller"))
 		[
 			SAssignNew(ControllerPanelList, SRCControllerPanelList, SharedThis(this))
+		];
+
+	// Add New Controller Button
+	const TSharedRef<SWidget> AddNewControllerButton = SNew(SComboButton)
+		.AddMetaData<FTagMetaData>(FTagMetaData(TEXT("Add Controller")))
+		.HAlign(HAlign_Center)
+		.VAlign(VAlign_Center)
+		.ButtonStyle(&RCPanelStyle->FlatButtonStyle)
+		.ForegroundColor(FSlateColor::UseForeground())
+		.CollapseMenuOnParentFocus(true)
+		.HasDownArrow(false)
+		.ContentPadding(FMargin(4.f, 2.f))
+		.ButtonContent()
+		[
+			SNew(SBox)
+			.WidthOverride(RCPanelStyle->IconSize.X)
+			.HeightOverride(RCPanelStyle->IconSize.Y)
+			[
+				SNew(SImage)
+				.ColorAndOpacity(FSlateColor::UseForeground())
+				.Image(FAppStyle::GetBrush("Icons.PlusCircle"))
+			]
 		]
-	];
+		.MenuContent()
+		[
+			GetControllerMenuContentWidget()
+		];
+
+	// Empty All Button
+	TSharedRef<SWidget> EmptyAllButton = SNew(SButton)
+		.AddMetaData<FTagMetaData>(FTagMetaData(TEXT("Empty Controllers")))
+		.HAlign(HAlign_Center)
+		.VAlign(VAlign_Center)
+		.ForegroundColor(FSlateColor::UseForeground())
+		.ButtonStyle(&RCPanelStyle->FlatButtonStyle)
+		.ToolTipText(LOCTEXT("EmptyAllToolTip", "Deletes all the controllers."))
+		.OnClicked(this, &SRCControllerPanel::RequestDeleteAllItems)
+		.Visibility_Lambda([this]() { return ControllerPanelList.IsValid() && !ControllerPanelList->IsEmpty() ? EVisibility::Visible : EVisibility::Collapsed; })
+		[
+			SNew(SBox)
+			.WidthOverride(RCPanelStyle->IconSize.X)
+			.HeightOverride(RCPanelStyle->IconSize.Y)
+			[
+				SNew(SImage)
+				.ColorAndOpacity(FSlateColor::UseForeground())
+				.Image(FAppStyle::GetBrush("Icons.Delete"))
+			]
+		];
+
+	ControllerDockPanel->AddHeaderToolbarItem(EToolbar::Left, AddNewControllerButton);
+	ControllerDockPanel->AddHeaderToolbarItem(EToolbar::Right, EmptyAllButton);
+
+	ChildSlot
+		.Padding(RCPanelStyle->PanelPadding)
+		[
+			ControllerDockPanel.ToSharedRef()
+		];
 }
 
 END_SLATE_FUNCTION_BUILD_OPTIMIZATION
 
+static UObject* GetBaseStructForType(FName StructType)
+{
+	if (StructType == NAME_Vector)
+	{
+		return TBaseStructure<FVector>::Get();
+	}
+	else if (StructType == NAME_Color)
+	{
+		return TBaseStructure<FColor>::Get();
+	}
+	else if (StructType == NAME_Rotator)
+	{
+		return TBaseStructure<FRotator>::Get();
+	}
+
+	ensureMsgf(false, TEXT("Found unsupported struct type %s in config."), *StructType.ToString());
+
+	return nullptr;
+}
+
 bool SRCControllerPanel::IsListFocused() const
 {
 	return ControllerPanelList->IsListFocused();
@@ -51,6 +129,99 @@ void SRCControllerPanel::DeleteSelectedPanelItem()
 	}
 }
 
+FReply SRCControllerPanel::RequestDeleteAllItems()
+{
+	if (!ControllerPanelList.IsValid())
+	{
+		return FReply::Unhandled();
+	}
+
+	const FText WarningMessage = FText::Format(LOCTEXT("DeleteAllWarning", "You are about to delete '{0}' controllers. This action might not be undone.\nAre you sure you want to proceed?"), ControllerPanelList->Num());
+
+	EAppReturnType::Type UserResponse = FMessageDialog::Open(EAppMsgType::YesNo, WarningMessage);
+
+	if (UserResponse == EAppReturnType::Yes)
+	{
+		return OnClickEmptyButton();
+	}
+
+	return FReply::Handled();
+}
+
+TSharedRef<SWidget> SRCControllerPanel::GetControllerMenuContentWidget() const
+{
+	constexpr  bool bShouldCloseWindowAfterMenuSelection = true;
+	FMenuBuilder MenuBuilder(bShouldCloseWindowAfterMenuSelection, nullptr);
+
+	TArray<TPair<EPropertyBagPropertyType, UObject*>> VirtualPropertyFieldClassNames;
+
+	// See Config file: BaseRemoteControl.ini
+	const URemoteControlLogicConfig* RCLogicConfig = GetDefault<URemoteControlLogicConfig>();
+	for (const EPropertyBagPropertyType ControllerType : RCLogicConfig->SupportedControllerTypes)
+	{
+		if (ControllerType == EPropertyBagPropertyType::Struct)
+		{
+			for (const FName StructType : RCLogicConfig->SupportedControllerStructTypes)
+			{
+				UObject* ValueTypeObject = GetBaseStructForType(StructType);
+				if (ensure(ValueTypeObject))
+				{
+					VirtualPropertyFieldClassNames.Add({ ControllerType, ValueTypeObject });
+				}
+			}
+		}
+		else
+		{
+			VirtualPropertyFieldClassNames.Add({ ControllerType, nullptr });
+		}
+	}
+
+	// Generate a menu from the list of supported Controllers
+	for (const TPair<EPropertyBagPropertyType, UObject*>& Pair : VirtualPropertyFieldClassNames)
+	{
+		// Display Name
+		const UEnum* Enum = FindObjectChecked<UEnum>(nullptr, TEXT("/Script/StructUtils.EPropertyBagPropertyType"));
+		FString PropertyTypeString = Enum->GetValueAsName(Pair.Key).ToString();
+		PropertyTypeString.RemoveFromStart("EPropertyBagPropertyType::");
+
+		if (PropertyTypeString.StartsWith(TEXT("Bool")))
+		{
+			PropertyTypeString = "Boolean";
+		}
+		
+		if (PropertyTypeString.StartsWith(TEXT("Int")))
+		{
+			PropertyTypeString = PropertyTypeString.Replace(TEXT("Int"), TEXT("Integer"));
+		}
+
+		const FName DefaultName = Pair.Value ? Pair.Value->GetFName() : *PropertyTypeString;
+
+		// Menu Item
+		FUIAction Action(FExecuteAction::CreateSP(this, &SRCControllerPanel::OnAddControllerClicked, Pair.Key, Pair.Value));
+		MenuBuilder.AddMenuEntry(
+			FText::Format(LOCTEXT("AddController", "{0}"), FText::FromName(DefaultName)),
+			FText::Format(LOCTEXT("AddControllerTooltip", "Add a controller of type {0}."), FText::FromName(DefaultName)),
+			FSlateIcon(),
+			Action);
+	}
+
+	return MenuBuilder.MakeWidget();
+}
+
+void SRCControllerPanel::OnAddControllerClicked(const EPropertyBagPropertyType InValueType, UObject* InValueTypeObject) const
+{
+	// Add to the asset
+	if (URemoteControlPreset* Preset = GetPreset())
+	{
+		URCVirtualPropertyInContainer* NewVirtualProperty = Preset->AddVirtualProperty(URCController::StaticClass(), InValueType, InValueTypeObject);
+
+		// Refresh list
+		const TSharedPtr<SRemoteControlPanel> RemoteControlPanel = GetRemoteControlPanel();
+		check(RemoteControlPanel);
+		RemoteControlPanel->OnControllerAdded.Broadcast(NewVirtualProperty->PropertyName);
+	}
+}
+
 void SRCControllerPanel::EnterRenameMode()
 {
 	if (ControllerPanelList)
@@ -58,3 +229,20 @@ void SRCControllerPanel::EnterRenameMode()
 		ControllerPanelList->EnterRenameMode();
 	}
 }
+
+FReply SRCControllerPanel::OnClickEmptyButton()
+{
+	if (URemoteControlPreset* Preset = GetPreset())
+	{
+		Preset->ResetVirtualProperties();
+	}
+
+	if (const TSharedPtr<SRemoteControlPanel> RemoteControlPanel = GetRemoteControlPanel())
+	{
+		RemoteControlPanel->OnEmptyControllers.Broadcast();
+	}
+
+	return FReply::Handled();
+}
+
+#undef LOCTEXT_NAMESPACE
diff --git a/Engine/Plugins/VirtualProduction/RemoteControl/Source/RemoteControlUI/Private/UI/Controller/SRCControllerPanel.h b/Engine/Plugins/VirtualProduction/RemoteControl/Source/RemoteControlUI/Private/UI/Controller/SRCControllerPanel.h
--- a/Engine/Plugins/VirtualProduction/RemoteControl/Source/RemoteControlUI/Private/UI/Controller/SRCControllerPanel.h
+++ b/Engine/Plugins/VirtualProduction/RemoteControl/Source/RemoteControlUI/Private/UI/Controller/SRCControllerPanel.h
@@ -2,8 +2,13 @@
 
 #pragma once
 
+#include "PropertyBag.h"
 #include "UI/BaseLogicUI/SRCLogicPanelBase.h"
 
+enum class EPropertyBagPropertyType : uint8;
+struct FRCPanelStyle;
+class SRCControllerPanel;
+
 /*
 * ~ SRCControllerPanel ~
 *
@@ -22,11 +27,6 @@ public:
 	/** Constructs this widget with InArgs */
 	void Construct(const FArguments& InArgs, const TSharedRef<SRemoteControlPanel>& InPanel);
 
-private:
-	/** Widget representing List of Controllers */
-	TSharedPtr<class SRCControllerPanelList> ControllerPanelList;
-
-public:
 	/** Whether the Controller list widget currently has focus. Used for Delete Item UI command */
 	bool IsListFocused() const;
 
@@ -34,4 +34,28 @@ public:
 	virtual void DeleteSelectedPanelItem() override;
 
 	void EnterRenameMode();
+
+protected:
+
+	/** Warns user before deleting all items in a panel. */
+	virtual FReply RequestDeleteAllItems() override;
+
+private:
+
+	/** Builds a menu containing the list of all possible Controllers*/
+	TSharedRef<SWidget> GetControllerMenuContentWidget() const;
+
+	/** Handles click event for Add Controller button*/
+	void OnAddControllerClicked(const EPropertyBagPropertyType InValueType, UObject* InValueTypeObject = nullptr) const;
+
+	/** Handles click event for Empty button; clears all controllers from the panel*/
+	FReply OnClickEmptyButton();
+
+private:
+
+	/** Widget representing List of Controllers */
+	TSharedPtr<class SRCControllerPanelList> ControllerPanelList;
+
+	/** Panel Style reference. */
+	const FRCPanelStyle* RCPanelStyle;
 };
diff --git a/Engine/Plugins/VirtualProduction/RemoteControl/Source/RemoteControlUI/Private/UI/Controller/SRCControllerPanelHeader.cpp b/Engine/Plugins/VirtualProduction/RemoteControl/Source/RemoteControlUI/Private/UI/Controller/SRCControllerPanelHeader.cpp
deleted file mode 100644
diff --git a/Engine/Plugins/VirtualProduction/RemoteControl/Source/RemoteControlUI/Private/UI/Controller/SRCControllerPanelHeader.h b/Engine/Plugins/VirtualProduction/RemoteControl/Source/RemoteControlUI/Private/UI/Controller/SRCControllerPanelHeader.h
deleted file mode 100644
diff --git a/Engine/Plugins/VirtualProduction/RemoteControl/Source/RemoteControlUI/Private/UI/Controller/SRCControllerPanelList.cpp b/Engine/Plugins/VirtualProduction/RemoteControl/Source/RemoteControlUI/Private/UI/Controller/SRCControllerPanelList.cpp
--- a/Engine/Plugins/VirtualProduction/RemoteControl/Source/RemoteControlUI/Private/UI/Controller/SRCControllerPanelList.cpp
+++ b/Engine/Plugins/VirtualProduction/RemoteControl/Source/RemoteControlUI/Private/UI/Controller/SRCControllerPanelList.cpp
@@ -10,8 +10,11 @@
 #include "RemoteControlPreset.h"
 #include "SlateOptMacros.h"
 #include "SRCControllerPanel.h"
+#include "Styling/RemoteControlStyles.h"
 #include "UI/BaseLogicUI/RCLogicModeBase.h"
+#include "UI/RemoteControlPanelStyle.h"
 #include "UI/SRemoteControlPanel.h"
+#include "Widgets/Views/SHeaderRow.h"
 
 #define LOCTEXT_NAMESPACE "SRCControllerPanelList"
 
@@ -48,6 +51,19 @@ namespace UE::RCControllerPanelList
 
 			return SNullWidget::NullWidget;
 		}
+
+	private:
+
+		//~ SWidget Interface
+		virtual FReply OnMouseButtonDoubleClick(const FGeometry& InMyGeometry, const FPointerEvent& InMouseEvent) override
+		{
+			if (ControllerItem.IsValid())
+			{
+				ControllerItem->EnterRenameMode();
+			}
+
+			return FSuperRowType::OnMouseButtonDoubleClick(InMyGeometry, InMouseEvent);
+		}
 	};
 }
 
@@ -59,22 +75,26 @@ void SRCControllerPanelList::Construct(const FArguments& InArgs, const TSharedRe
 	
 	ControllerPanelWeakPtr = InControllerPanel;
 	
+	RCPanelStyle = &FRemoteControlPanelStyle::Get()->GetWidgetStyle<FRCPanelStyle>("RemoteControlPanel.MinorPanel");
+
 	ListView = SNew(SListView<TSharedPtr<FRCControllerModel>>)
 		.ListItemsSource(&ControllerItems)
 		.OnSelectionChanged(this, &SRCControllerPanelList::OnTreeSelectionChanged)
 		.OnGenerateRow(this, &SRCControllerPanelList::OnGenerateWidgetForList)
 		.SelectionMode(ESelectionMode::Single) // Current setup supports only single selection (and related display) of a Controller in the list
-		.HeaderRow
-		(
+		.HeaderRow(
 			SNew(SHeaderRow)
+			.Style(&RCPanelStyle->HeaderRowStyle)
 
 			+ SHeaderRow::Column(UE::RCControllerPanelList::ControllerNameColumn)
-			.DefaultLabel(LOCTEXT("Controller Name Column Name", "Name"))
+			.DefaultLabel(LOCTEXT("ControllerNameColumnName", "Name"))
 			.FillWidth(0.25f)
+			.HeaderContentPadding(RCPanelStyle->HeaderRowPadding)
 
 			+ SHeaderRow::Column(UE::RCControllerPanelList::ControllerValueColumn)
-			.DefaultLabel(LOCTEXT("Controller Value Column Name", "Input"))
+			.DefaultLabel(LOCTEXT("ControllerValueColumnName", "Input"))
 			.FillWidth(0.75f)
+			.HeaderContentPadding(RCPanelStyle->HeaderRowPadding)
 		);
 
 	ChildSlot
@@ -101,6 +121,16 @@ void SRCControllerPanelList::Construct(const FArguments& InArgs, const TSharedRe
 
 END_SLATE_FUNCTION_BUILD_OPTIMIZATION
 
+bool SRCControllerPanelList::IsEmpty() const
+{
+	return ControllerItems.IsEmpty();
+}
+
+int32 SRCControllerPanelList::Num() const
+{
+	return ControllerItems.Num();
+}
+
 void SRCControllerPanelList::Reset()
 {
 	ControllerItems.Empty();
@@ -138,13 +168,11 @@ void SRCControllerPanelList::Reset()
 
 TSharedRef<ITableRow> SRCControllerPanelList::OnGenerateWidgetForList(TSharedPtr<FRCControllerModel> InItem, const TSharedRef<STableViewBase>& OwnerTable)
 {
-	typedef UE::RCControllerPanelList::SControllerItemListRow ControllerRowType;
+	typedef UE::RCControllerPanelList::SControllerItemListRow SControllerRowType;
 
-	const TSharedRef<ControllerRowType> NewRow =
-		SNew(ControllerRowType, OwnerTable, InItem.ToSharedRef())
+	return SNew(SControllerRowType, OwnerTable, InItem.ToSharedRef())
+		.Style(&RCPanelStyle->TableRowStyle)
 		.Padding(FMargin(3.f));
-
-	return NewRow;
 }
 
 void SRCControllerPanelList::OnTreeSelectionChanged(TSharedPtr<FRCControllerModel> InItem, ESelectInfo::Type)
@@ -153,16 +181,12 @@ void SRCControllerPanelList::OnTreeSelectionChanged(TSharedPtr<FRCControllerMode
 	{
 		if (const TSharedPtr<SRemoteControlPanel> RemoteControlPanel = ControllerPanel->GetRemoteControlPanel())
 		{
-			if (InItem.IsValid())
+			if (InItem != SelectedControllerItemWeakPtr.Pin())
 			{
-				if (InItem != SelectedControllerItemWeakPtr.Pin())
-				{
-					RemoteControlPanel->OnControllerSelectionChanged.Broadcast(InItem);
-					RemoteControlPanel->OnBehaviourSelectionChanged.Broadcast(InItem->GetSelectedBehaviourModel());
-				}
+				SelectedControllerItemWeakPtr = InItem;
+				RemoteControlPanel->OnControllerSelectionChanged.Broadcast(InItem);
+				RemoteControlPanel->OnBehaviourSelectionChanged.Broadcast(InItem.IsValid() ? InItem->GetSelectedBehaviourModel() : nullptr);
 			}
-
-			SelectedControllerItemWeakPtr = InItem;
 		}
 	}
 }
@@ -255,4 +279,4 @@ void SRCControllerPanelList::EnterRenameMode()
 	}
 }
 
-#undef LOCTEXT_NAMESPACE
\ No newline at end of file
+#undef LOCTEXT_NAMESPACE
diff --git a/Engine/Plugins/VirtualProduction/RemoteControl/Source/RemoteControlUI/Private/UI/Controller/SRCControllerPanelList.h b/Engine/Plugins/VirtualProduction/RemoteControl/Source/RemoteControlUI/Private/UI/Controller/SRCControllerPanelList.h
--- a/Engine/Plugins/VirtualProduction/RemoteControl/Source/RemoteControlUI/Private/UI/Controller/SRCControllerPanelList.h
+++ b/Engine/Plugins/VirtualProduction/RemoteControl/Source/RemoteControlUI/Private/UI/Controller/SRCControllerPanelList.h
@@ -3,6 +3,7 @@
 #pragma once
 #include "UI/BaseLogicUI/SRCLogicPanelListBase.h"
 
+struct FRCPanelStyle;
 class FRCControllerModel;
 class FRCLogicModeBase;
 class IPropertyRowGenerator;
@@ -32,7 +33,22 @@ public:
 	/** Constructs this widget with InArgs */
 	void Construct(const FArguments& InArgs, const TSharedRef<SRCControllerPanel> InControllerPanel);
 
+	/** Returns true if the underlying list is valid and empty. */
+	virtual bool IsEmpty() const override;
+
+	/** Returns number of items in the list. */
+	virtual int32 Num() const override;
+
+	/** Whether the Controllers List View currently has focus.*/
+	bool IsListFocused() const;
+
+	/** Deletes currently selected items from the list view*/
+	void DeleteSelectedPanelItem();
+
+	void EnterRenameMode();
+
 private:
+
 	/** OnGenerateRow delegate for the Actions List View */
 	TSharedRef<ITableRow> OnGenerateWidgetForList( TSharedPtr<FRCControllerModel> InItem, const TSharedRef<STableViewBase>& OwnerTable );
 	
@@ -77,16 +93,7 @@ private:
 	/** Removes the given Controller UI model item from the list of UI models*/
 	virtual int32 RemoveModel(const TSharedPtr<FRCLogicModeBase> InModel) override;
 
-public:
-	/** Whether the Controllers List View currently has focus.*/
-	bool IsListFocused() const;
-
-	/** Deletes currently selected items from the list view*/
-	void DeleteSelectedPanelItem();
-
-	void EnterRenameMode();
-
-private:
-	
+	/** Panel Style reference. */
+	const FRCPanelStyle* RCPanelStyle;
 };
 
diff --git a/Engine/Plugins/VirtualProduction/RemoteControl/Source/RemoteControlUI/Private/UI/Drawers/SRCPanelDrawer.cpp b/Engine/Plugins/VirtualProduction/RemoteControl/Source/RemoteControlUI/Private/UI/Drawers/SRCPanelDrawer.cpp
new file mode 100644
--- /dev/null
+++ b/Engine/Plugins/VirtualProduction/RemoteControl/Source/RemoteControlUI/Private/UI/Drawers/SRCPanelDrawer.cpp
@@ -0,0 +1,475 @@
+// Copyright Epic Games, Inc. All Rights Reserved.
+
+#include "UI/Drawers/SRCPanelDrawer.h"
+
+#include "Framework/Text/PlainTextLayoutMarshaller.h"
+
+#include "Styling/AppStyle.h"
+
+#include "Widgets/Colors/SComplexGradient.h"
+#include "Widgets/Docking/SDockTab.h"
+#include "Widgets/Images/SImage.h"
+#include "Widgets/Input/SButton.h"
+#include "Widgets/Layout/SBox.h"
+#include "Widgets/Layout/SBorder.h"
+#include "Widgets/SBoxPanel.h"
+#include "Widgets/SLeafWidget.h"
+#include "Widgets/Text/SlateTextBlockLayout.h"
+
+const FVector2D FDockingPanelConstants::MaxMinorPanelSize(144.f, 25.f);
+const FVector2D FDockingPanelConstants::MaxMajorPanelSize(210.f, 50.f);
+
+const FVector2D FDockingPanelConstants::GetMaxPanelSizeFor(ETabRole TabRole)
+{
+	return (TabRole == ETabRole::MajorTab)
+		? MaxMajorPanelSize
+		: MaxMinorPanelSize;
+}
+
+#define LOCTEXT_NAMESPACE "RCPanelDrawer"
+
+/**
+ * Vertical text block for use in the tab drawer button.
+ * Text is aligned to the top of the widget if it fits without clipping;
+ * otherwise it is ellipsized and fills the widget height.
+ */
+class SRCPanelDrawerTextBlock : public SLeafWidget
+{
+public:
+
+	enum class ERotation
+	{
+		Clockwise,
+		CounterClockwise,
+	};
+
+	SLATE_BEGIN_ARGS(SRCPanelDrawerTextBlock)
+		: _Text()
+		, _TextStyle(&FAppStyle::Get().GetWidgetStyle<FTextBlockStyle>("NormalText"))
+		, _Rotation(ERotation::Clockwise)
+		, _OverflowPolicy()
+	{}
+	
+		SLATE_ATTRIBUTE(FText, Text)
+
+		SLATE_STYLE_ARGUMENT(FTextBlockStyle, TextStyle)
+
+		SLATE_ATTRIBUTE(ERotation, Rotation)
+
+		SLATE_ARGUMENT(TOptional<ETextOverflowPolicy>, OverflowPolicy)
+
+	SLATE_END_ARGS()
+
+	virtual ~SRCPanelDrawerTextBlock() {}
+
+	void Construct(const FArguments& InArgs)
+	{
+		Text = InArgs._Text;
+		TextStyle = *InArgs._TextStyle;
+		Rotation = InArgs._Rotation;
+		TextLayoutCache = MakeUnique<FSlateTextBlockLayout>(
+			this, FTextBlockStyle::GetDefault(), TOptional<ETextShapingMethod>(), TOptional<ETextFlowDirection>(),
+			FCreateSlateTextLayout(), FPlainTextLayoutMarshaller::Create(), nullptr);
+		TextLayoutCache->SetTextOverflowPolicy(InArgs._OverflowPolicy.IsSet() ? InArgs._OverflowPolicy : TextStyle.OverflowPolicy);
+	}
+
+	virtual int32 OnPaint(const FPaintArgs& Args, const FGeometry& AllottedGeometry, const FSlateRect& MyCullingRect, FSlateWindowElementList& OutDrawElements, int32 LayerId, const FWidgetStyle& InWidgetStyle, bool bParentEnabled) const override
+	{
+		// We're going to figure out the bounds of the corresponding horizontal text, and then rotate it into a vertical orientation.
+		const FVector2D LocalSize = AllottedGeometry.GetLocalSize();
+		const FVector2D DesiredHorizontalTextSize = TextLayoutCache->GetDesiredSize();
+		const FVector2D ActualHorizontalTextSize(FMath::Min(DesiredHorizontalTextSize.X, LocalSize.Y), FMath::Min(DesiredHorizontalTextSize.Y, LocalSize.X));
+
+		// Now determine the center of the vertical text by rotating the dimensions of the horizontal text.
+		// The center should align it to the top of the widget.
+		const FVector2D VerticalTextSize(ActualHorizontalTextSize.Y, ActualHorizontalTextSize.X);
+		const FVector2D VerticalTextCenter = VerticalTextSize / 2.f;
+
+		// Now determine where the horizontal text should be positioned so that it is centered on the vertical text:
+		//      +-+
+		//      |v|
+		//      |e|
+		// [ horizontal ]
+		//      |r|
+		//      |t|
+		//      +-+
+		const FVector2D HorizontalTextPosition = VerticalTextCenter - ActualHorizontalTextSize / 2.f;
+
+		// Define the text's geometry using the horizontal bounds, then rotate it 90/-90 degrees into place to become vertical.
+		const FSlateRenderTransform RotationTransform(FSlateRenderTransform(FQuat2D(FMath::DegreesToRadians(Rotation.Get() == ERotation::Clockwise ? 90 : -90))));
+		const FGeometry TextGeometry = AllottedGeometry.MakeChild(ActualHorizontalTextSize, FSlateLayoutTransform(HorizontalTextPosition), RotationTransform, FVector2D(0.5f, 0.5f));
+
+		return TextLayoutCache->OnPaint(Args, TextGeometry, MyCullingRect, OutDrawElements, LayerId, InWidgetStyle, ShouldBeEnabled(bParentEnabled));
+	}
+
+	virtual FVector2D ComputeDesiredSize(float LayoutScaleMultiplier) const override
+	{
+		// The text's desired size reflects the horizontal/untransformed text.
+		// Switch the dimensions for vertical text.
+		const FVector2D DesiredHorizontalTextSize = TextLayoutCache->ComputeDesiredSize(
+			FSlateTextBlockLayout::FWidgetDesiredSizeArgs(
+				Text.Get(),
+				FText(),
+				0.f,
+				false,
+				ETextWrappingPolicy::DefaultWrapping,
+				ETextTransformPolicy::None,
+				FMargin(),
+				1.f,
+				ETextJustify::Left),
+			LayoutScaleMultiplier, TextStyle);
+		return FVector2D(DesiredHorizontalTextSize.Y, DesiredHorizontalTextSize.X);
+	}
+
+	void SetText(TAttribute<FText> InText)
+	{
+		Text = InText;
+	}
+
+	void SetRotation(TAttribute<ERotation> InRotation)
+	{
+		Rotation = InRotation;
+	}
+
+private:
+
+	TAttribute<FText> Text;
+	FTextBlockStyle TextStyle;
+	TAttribute<ERotation> Rotation;
+	TUniquePtr<FSlateTextBlockLayout> TextLayoutCache;
+};
+
+class SRCPanelDrawerButton : public SCompoundWidget
+{
+
+	SLATE_BEGIN_ARGS(SRCPanelDrawerButton)
+	{}
+		SLATE_EVENT(FOnRCPanelDrawerButtonPressed, OnDrawerButtonPressed)
+	
+	SLATE_END_ARGS()
+
+public:
+
+	void Construct(const FArguments& InArgs, TSharedRef<FRCPanelDrawerArgs> ForPanel)
+	{
+		const FVector2D Size = FDockingPanelConstants::GetMaxPanelSizeFor(ETabRole::PanelTab);
+
+		const FSlateRenderTransform Rotate90(FQuat2D(FMath::DegreesToRadians(90.f)));
+
+		DockTabStyle = &FAppStyle::Get().GetWidgetStyle<FDockTabStyle>("Docking.Tab");
+
+		OnDrawerButtonPressed = InArgs._OnDrawerButtonPressed;
+
+		Panel = ForPanel;
+
+		static FLinearColor ActiveBorderColor = FAppStyle::Get().GetSlateColor("Docking.Tab.ActiveTabIndicatorColor").GetSpecifiedColor();
+		static FLinearColor ActiveBorderColorTransparent = FLinearColor(ActiveBorderColor.R, ActiveBorderColor.G, ActiveBorderColor.B, 0.f);
+		static TArray<FLinearColor> GradientStops{ ActiveBorderColorTransparent, ActiveBorderColor, ActiveBorderColorTransparent };
+
+		ChildSlot
+			.Padding(0.f)
+			[
+				SNew(SBox)
+				.WidthOverride(Size.Y) // Swap desired dimensions for a vertical tab
+				.HeightOverride(Size.X)
+				.Clipping(EWidgetClipping::ClipToBounds)
+				[
+					SNew(SOverlay)
+					+ SOverlay::Slot()
+					[
+						SAssignNew(MainButton, SButton)
+						.ToolTipText(ForPanel->ToolTip)
+						.ContentPadding(FMargin(0.f, DockTabStyle->TabPadding.Top, 0.f, DockTabStyle->TabPadding.Bottom))
+						.OnPressed_Lambda([this]() { OnDrawerButtonPressed.ExecuteIfBound(Panel.ToSharedRef(), false); })
+						.ForegroundColor(FSlateColor::UseForeground())
+						[
+							SNew(SVerticalBox)
+							+ SVerticalBox::Slot()
+							.AutoHeight()
+							.VAlign(VAlign_Center)
+							.HAlign(HAlign_Center)
+							.Padding(0.f, 5.f, 0.f, 5.f)
+							[
+								SNew(SImage)
+								.ColorAndOpacity(FSlateColor::UseForeground())
+								.Image(ForPanel->Icon.GetIcon())
+								.DesiredSizeOverride(FVector2D(16, 16))
+								.RenderTransform(ForPanel->bRotateIconBy90 ? Rotate90 : FSlateRenderTransform())
+								.RenderTransformPivot(FVector2D(.5f, .5f))
+							]
+							+ SVerticalBox::Slot()
+							.Padding(0.f, 5.f, 0.f, 5.f)
+							.FillHeight(1.f)
+							.HAlign(HAlign_Center)
+							.VAlign(VAlign_Center)
+							[
+								SAssignNew(Label, SRCPanelDrawerTextBlock)
+								.TextStyle(&DockTabStyle->TabTextStyle)
+								.Text(ForPanel->Label)
+								.OverflowPolicy(ETextOverflowPolicy::Ellipsis)
+								.Clipping(EWidgetClipping::ClipToBounds)
+							]
+							+ SVerticalBox::Slot()
+							.AutoHeight()
+							.HAlign(HAlign_Center)
+							.Padding(0.f, 5.f, 0.f, 4.f)
+							[
+								SNew(SBorder)
+								.BorderImage(FAppStyle::Get().GetBrush("ToolPanel.GroupBorder"))
+								.Visibility(this, &SRCPanelDrawerButton::GetPanelStateVisibility)
+								.ToolTipText(this, &SRCPanelDrawerButton::GetPanelStateToolTipText)
+								.Padding(2.f)
+								.HAlign(HAlign_Center)
+								[
+									SNew(SImage)
+									.ColorAndOpacity(FSlateColor::UseForeground())
+									.Image(this, &SRCPanelDrawerButton::GetPanelStateImage)
+								]
+							]
+						]
+					]
+					+ SOverlay::Slot()
+					[
+						SAssignNew(OpenBorder, SBorder)
+						.Visibility(EVisibility::HitTestInvisible)
+					]
+					+ SOverlay::Slot()
+					.HAlign(HAlign_Left)
+					[
+						SAssignNew(ActiveIndicator, SComplexGradient)
+						.DesiredSizeOverride(FVector2D(1.f, 1.f))
+						.GradientColors(GradientStops)
+						.Orientation(EOrientation::Orient_Horizontal)
+						.Visibility(this, &SRCPanelDrawerButton::GetActivePanelIndicatorVisibility)
+					]
+				]
+			];
+
+		UpdateAppearance(nullptr);
+	}
+
+	void UpdateAppearance(const TSharedPtr<FRCPanelDrawerArgs> OpenedDrawer)
+	{
+		bool bShouldAppearOpened = OpenedDrawer.IsValid();
+
+		SRCPanelDrawerTextBlock::ERotation Rotation;
+
+		Rotation = bShouldAppearOpened ? SRCPanelDrawerTextBlock::ERotation::CounterClockwise : SRCPanelDrawerTextBlock::ERotation::Clockwise;
+
+		check(Label);
+		Label->SetRotation(Rotation);
+
+		if (OpenedDrawer == Panel)
+		{
+			// this button is the one with the tab that is actually opened so show the tab border
+			OpenBorder->SetVisibility(EVisibility::HitTestInvisible);
+			MainButton->SetButtonStyle(&FAppStyle::Get().GetWidgetStyle<FButtonStyle>("Docking.SidebarButton.Opened"));
+
+			OpenBorder->SetBorderImage(FAppStyle::Get().GetBrush("Docking.Sidebar.Border_SquareRight"));
+		}
+		else
+		{
+			OpenBorder->SetVisibility(EVisibility::Collapsed);
+			MainButton->SetButtonStyle(&FAppStyle::Get().GetWidgetStyle<FButtonStyle>("Docking.SidebarButton.Closed"));
+		}
+	}
+
+	virtual FSlateColor GetForegroundColor() const
+	{
+		if (ActiveIndicator->GetVisibility() != EVisibility::Collapsed)
+		{
+			return DockTabStyle->ActiveForegroundColor;
+		}
+		else if (IsHovered())
+		{
+			return DockTabStyle->HoveredForegroundColor;
+		}
+
+		return FSlateColor::UseStyle();
+	}
+
+private:
+
+	EVisibility GetActivePanelIndicatorVisibility() const
+	{
+		return Panel->IsActive() ? EVisibility::HitTestInvisible : EVisibility::Collapsed;
+	}
+	
+	EVisibility GetPanelStateVisibility() const
+	{
+		return (IsHovered() || Panel->IsActive()) ? EVisibility::Visible : EVisibility::Hidden;
+	}
+
+	FText GetPanelStateToolTipText() const
+	{
+		if (Panel->IsActive())
+		{
+			return LOCTEXT("CollapsePanelToolTip", "Collapse panel to the drawer.");
+		}
+		else
+		{
+			return LOCTEXT("DrawPanelToolTip", "Draw panel from the drawer.");
+		}
+	}
+
+	const FSlateBrush* GetPanelStateImage() const
+	{
+		if (Panel->IsActive())
+		{
+			return FAppStyle::Get().GetBrush("Icons.ChevronLeft");
+		}
+		else
+		{
+			return FAppStyle::Get().GetBrush("Icons.ChevronRight");
+		}
+	}
+
+	TSharedPtr<FRCPanelDrawerArgs> Panel;
+	TSharedPtr<SRCPanelDrawerTextBlock> Label;
+	TSharedPtr<SWidget> ActiveIndicator;
+	TSharedPtr<SBorder> OpenBorder;
+	TSharedPtr<SButton> MainButton;
+	FOnRCPanelDrawerButtonPressed OnDrawerButtonPressed;
+	const FDockTabStyle* DockTabStyle;
+};
+
+void SRCPanelDrawer::Construct(const FArguments& InArgs)
+{
+	ChildSlot
+		.Padding(0.f)
+		[
+			SNew(SBorder)
+			.Padding(0.f)
+			.BorderImage(FAppStyle::Get().GetBrush("Docking.Sidebar.Background"))
+			[
+				SAssignNew(PanelBox, SVerticalBox)
+			]
+		];
+}
+
+bool SRCPanelDrawer::IsRegistered(TSharedRef<FRCPanelDrawerArgs> InPanel) const
+{
+	return Panels.ContainsByPredicate(
+		[InPanel](TPair<TSharedRef<FRCPanelDrawerArgs>, TSharedRef<SRCPanelDrawerButton>> PanelPair)
+		{
+			return PanelPair.Key == InPanel;
+		});
+}
+
+void SRCPanelDrawer::RegisterPanel(TSharedRef<FRCPanelDrawerArgs> InPanel)
+{
+	// Only register panels with visible drawer.
+	if (!IsRegistered(InPanel) && InPanel->DrawerVisibility == EVisibility::Visible)
+	{
+		SetVisibility(EVisibility::SelfHitTestInvisible);
+
+		TSharedRef<SRCPanelDrawerButton> PanelButton = SNew(SRCPanelDrawerButton, InPanel)
+			.OnDrawerButtonPressed(this, &SRCPanelDrawer::TogglePanel);
+
+		PanelBox->AddSlot()
+			// Make the tabs evenly fill the sidebar until they reach the max size
+			.FillHeight(1.f)
+			.MaxHeight(FDockingPanelConstants::GetMaxPanelSizeFor(ETabRole::PanelTab).X)
+			.HAlign(HAlign_Left)
+			.Padding(0.f, 0.f, 0.f, 4.f)
+			[
+				PanelButton
+			];
+
+		Panels.Emplace(InPanel, PanelButton);
+
+		// Upon explicit request open the drawer automatically after it's added
+		if (InPanel->bDrawnByDefault)
+		{
+			InPanel->SetState(ERCPanelState::Opened);
+		}
+
+		if (InPanel->IsActive())
+		{
+			PanelButton->UpdateAppearance(InPanel);
+
+			OpenedDrawer = InPanel->GetPanelID();
+		}
+	}
+}
+
+bool SRCPanelDrawer::UnregisterPanel(TSharedRef<FRCPanelDrawerArgs> InPanelToRemove)
+{
+	int32 FoundIndex = Panels.IndexOfByPredicate(
+		[InPanelToRemove](TPair<TSharedRef<FRCPanelDrawerArgs>, TSharedRef<SRCPanelDrawerButton>> PanelPair)
+		{
+			return PanelPair.Key == InPanelToRemove;
+		});
+
+	if (FoundIndex != INDEX_NONE)
+	{
+		TPair<TSharedRef<FRCPanelDrawerArgs>, TSharedRef<SRCPanelDrawerButton>> PanelPair = Panels[FoundIndex];
+
+		Panels.RemoveAt(FoundIndex);
+		PanelBox->RemoveSlot(PanelPair.Value);
+
+		if (Panels.Num() == 0)
+		{
+			SetVisibility(EVisibility::Collapsed);
+		}
+	}
+
+	return FoundIndex != INDEX_NONE;
+}
+
+void SRCPanelDrawer::TogglePanel(TSharedRef<FRCPanelDrawerArgs> InPanel, const bool bRemoveDrawer)
+{
+	auto PanelFinderPredicate = [InPanel](const TPair<TSharedRef<FRCPanelDrawerArgs>, TSharedRef<SRCPanelDrawerButton>>& TabAndButton)
+	{
+		return TabAndButton.Key == InPanel;
+	};
+
+	if (bRemoveDrawer)
+	{
+		UnregisterPanel(InPanel);
+	}
+
+	if (OpenedDrawer != InPanel->GetPanelID())
+	{
+		OpenedDrawer = InPanel->GetPanelID();
+
+		UpdateAppearance();
+	}
+	else if (TPair<TSharedRef<FRCPanelDrawerArgs>, TSharedRef<SRCPanelDrawerButton>>* Panel = Panels.FindByPredicate(PanelFinderPredicate))
+	{
+		OpenedDrawer = ERCPanels::RCP_None;
+
+		Panel->Key->SetState(ERCPanelState::Collapsed);
+
+		Panel->Value->UpdateAppearance(nullptr);
+	}
+	
+	if (InPanel->DrawerVisibility == EVisibility::Visible)
+	{
+		// Register mode specific panels as they might not get registered initially based on the mode we are in.
+		RegisterPanel(InPanel);
+	}
+
+	OnRCPanelToggledHandle.ExecuteIfBound(OpenedDrawer);
+}
+
+void SRCPanelDrawer::UpdateAppearance()
+{
+	for (auto Panel = Panels.CreateIterator(); Panel; ++Panel)
+	{
+		const bool bShouldUpdateAppearance = Panel->Key->GetPanelID() == OpenedDrawer;
+
+		Panel->Key->SetState(bShouldUpdateAppearance ? ERCPanelState::Opened : ERCPanelState::Collapsed);
+
+		if (bShouldUpdateAppearance)
+		{
+			Panel->Value->UpdateAppearance(Panel->Key);
+		}
+		else
+		{
+			Panel->Value->UpdateAppearance(nullptr);
+		}
+	}
+}
+
+#undef LOCTEXT_NAMESPACE
diff --git a/Engine/Plugins/VirtualProduction/RemoteControl/Source/RemoteControlUI/Private/UI/Drawers/SRCPanelDrawer.h b/Engine/Plugins/VirtualProduction/RemoteControl/Source/RemoteControlUI/Private/UI/Drawers/SRCPanelDrawer.h
new file mode 100644
--- /dev/null
+++ b/Engine/Plugins/VirtualProduction/RemoteControl/Source/RemoteControlUI/Private/UI/Drawers/SRCPanelDrawer.h
@@ -0,0 +1,179 @@
+// Copyright Epic Games, Inc. All Rights Reserved.
+
+#pragma once
+
+#include "CoreMinimal.h"
+#include "Textures/SlateIcon.h"
+#include "Widgets/SCompoundWidget.h"
+
+enum ETabRole : uint8;
+class SRCPanelDrawerButton;
+class SVerticalBox;
+
+/** The variety of panels we have in the RC Panel. */
+enum class ERCPanels : uint8
+{
+	RCP_None,
+
+	RCP_Properties,
+	
+	RCP_EntityDetails,
+
+	RCP_Protocols,
+
+	RCP_OutputLog,
+
+	RCP_Count,
+};
+
+/** State of the RC Panels. */
+enum class ERCPanelState
+{
+	/** Panel is actively drawn. */
+	Opened,
+
+	/** Panel is collapsed. */
+	Collapsed,
+};
+
+struct FDockingPanelConstants
+{
+	static const FVector2D MaxMinorPanelSize;
+	static const FVector2D MaxMajorPanelSize;
+	static const FVector2D GetMaxPanelSizeFor(ETabRole TabRole);
+};
+
+/** Arguments used to create a RC Panel Drawer button. */
+struct FRCPanelDrawerArgs
+{
+	FRCPanelDrawerArgs(ERCPanels InPanel)
+		: bDrawnByDefault(false)
+		, bRotateIconBy90(false)
+		, Icon()
+		, Label(FText::FromName(NAME_None))
+		, ToolTip(FText::FromName(NAME_None))
+		, DrawerVisibility(EVisibility::Visible)
+		, PanelID(InPanel)
+		, PanelState(ERCPanelState::Collapsed)
+	{
+	}
+	
+	bool operator==(const FRCPanelDrawerArgs& OtherPanel)
+	{
+		return PanelID == OtherPanel.GetPanelID();
+	}
+	
+	bool operator==(const FRCPanelDrawerArgs& OtherPanel) const
+	{
+		return PanelID == OtherPanel.GetPanelID();
+	}
+
+	ERCPanels GetPanelID() const
+	{
+		return PanelID;
+	}
+
+	bool IsActive() const
+	{
+		return PanelState == ERCPanelState::Opened;
+	}
+
+	void SetState(ERCPanelState NewState)
+	{
+		if (PanelState != NewState)
+		{
+			PanelState = NewState;
+		}
+	}
+
+	/** Whether to draw this panel by default. */
+	bool bDrawnByDefault;
+	
+	/** Whether to rotate the icon of this panel button. */
+	bool bRotateIconBy90;
+
+	/** The icon to use for the drawer panel button. */
+	FSlateIcon Icon;
+
+	/** The label to use for the drawer panel button. */
+	FText Label;
+	
+	/** The tooltip to use for the drawer panel button. */
+	FText ToolTip;
+	
+	/** Visibility of the drawer button. */
+	EVisibility DrawerVisibility;
+
+private:
+
+	/** Unique Identifier to identify the panel. */
+	ERCPanels PanelID;
+
+	/** Active state of this panel. */
+	ERCPanelState PanelState;
+};
+
+/** Delegate to be called when a drawer button is clicked. */
+DECLARE_DELEGATE_TwoParams(FOnRCPanelDrawerButtonPressed, TSharedRef<FRCPanelDrawerArgs> /* Panel */, const bool /* bRemoveDrawer */)
+
+/** Delegate to be called when a drawer button is clicked. */
+DECLARE_DELEGATE_OneParam(FOnRCPanelToggled, ERCPanels /* PanelID */)
+
+/**
+ * A custom tab drawer tailored for RC Panels.
+ */
+class SRCPanelDrawer : public SCompoundWidget
+{
+public:
+	SLATE_BEGIN_ARGS(SRCPanelDrawer)
+	{}
+	SLATE_END_ARGS()
+
+	/** Constructs this widget with InArgs */
+	void Construct(const FArguments& InArgs);
+
+	/**
+	 * @return true if this sidebar contains the provided tab
+	 */
+	bool IsRegistered(TSharedRef<FRCPanelDrawerArgs> InPanel) const;
+
+	/**
+	 * Adds a panel to the drawer.
+	 */
+	void RegisterPanel(TSharedRef<FRCPanelDrawerArgs> InPanel);
+
+	/**
+	 * Removes a panel from the drawer.
+	 *
+	 * Note it is not sufficient to call this to clean up the panel completely
+	 * as it would just unregister the panel from this drawer.
+	 * 
+	 * @return true if the tab was found and removal was successful
+	 */
+	bool UnregisterPanel(TSharedRef<FRCPanelDrawerArgs> InPanelToRemove);
+
+	/** Delegate to be called when a drawer button is clicked. */
+	FOnRCPanelToggled& OnRCPanelToggled() { return OnRCPanelToggledHandle; }
+
+	/** Toggles the panel. */
+	void TogglePanel(TSharedRef<FRCPanelDrawerArgs> InPanel, const bool bRemoveDrawer = false);
+
+private:
+
+	/** Updates the appearance of the drawn panel. */
+	void UpdateAppearance();
+
+private:
+
+	/** Container holds all drawer buttons. */
+	TSharedPtr<SVerticalBox> PanelBox;
+
+	/** Array of all registered panels. */
+	TArray<TPair<TSharedRef<FRCPanelDrawerArgs>, TSharedRef<SRCPanelDrawerButton>>> Panels;
+
+	/** Delegate to be called when a drawer button is clicked. */
+	FOnRCPanelToggled OnRCPanelToggledHandle;
+
+	/** Generally speaking one drawer is only ever open at once. */
+	ERCPanels OpenedDrawer;
+};
diff --git a/Engine/Plugins/VirtualProduction/RemoteControl/Source/RemoteControlUI/Private/UI/Panels/SRCDockPanel.cpp b/Engine/Plugins/VirtualProduction/RemoteControl/Source/RemoteControlUI/Private/UI/Panels/SRCDockPanel.cpp
new file mode 100644
--- /dev/null
+++ b/Engine/Plugins/VirtualProduction/RemoteControl/Source/RemoteControlUI/Private/UI/Panels/SRCDockPanel.cpp
@@ -0,0 +1,368 @@
+// Copyright Epic Games, Inc. All Rights Reserved.
+
+#include "UI/Panels/SRCDockPanel.h"
+
+#include "SlateOptMacros.h"
+#include "Styling/RemoteControlStyles.h"
+
+#include "UI/RemoteControlPanelStyle.h"
+
+#include "Widgets/Layout/SBorder.h"
+#include "Widgets/Layout/SBox.h"
+#include "Widgets/Layout/SSpacer.h"
+#include "Widgets/Layout/SSplitter.h"
+
+#include "Widgets/SBoxPanel.h"
+#include "Widgets/SWidget.h"
+#include "Widgets/Text/STextBlock.h"
+
+SLATE_IMPLEMENT_WIDGET(SRCMajorPanel)
+
+void SRCMajorPanel::PrivateRegisterAttributes(FSlateAttributeInitializer& AttributeInitializer)
+{
+}
+
+BEGIN_SLATE_FUNCTION_BUILD_OPTIMIZATION
+
+void SRCMajorPanel::Construct(const SRCMajorPanel::FArguments& InArgs)
+{
+	RCPanelStyle = &FRemoteControlPanelStyle::Get()->GetWidgetStyle<FRCPanelStyle>("RemoteControlPanel.MinorPanel");
+
+	bIsFooterEnabled = InArgs._EnableFooter;
+	bIsHeaderEnabled = InArgs._EnableHeader;
+
+	ChildSlot
+	.Padding(RCPanelStyle->PanelPadding)
+	[
+		SNew(SBorder)
+		.BorderImage(&RCPanelStyle->ContentAreaBrush)
+		.Padding(RCPanelStyle->PanelPadding)
+		[	
+			SAssignNew(ContentPanel, SSplitter)
+			.Orientation(InArgs._Orientation)
+
+			// Content Panel
+			+SSplitter::Slot()
+			.Value(1.f)
+			[
+				SAssignNew(Children, SSplitter)
+				.Orientation(InArgs._ChildOrientation)
+			]
+		]
+	];
+
+	// Enable Header upon explicit request.
+	if (bIsHeaderEnabled.Get() && ContentPanel.IsValid())
+	{
+		// Header Panel
+		ContentPanel->AddSlot(0)
+			.Resizable(false)
+			.SizeRule(SSplitter::SizeToContent)
+			[
+				SNew(SBox)
+				.HAlign(HAlign_Center)
+				.VAlign(VAlign_Center)
+				.Padding(2.f, 4.f)
+				[
+					SNew(STextBlock)
+					.TextStyle(&RCPanelStyle->HeaderTextStyle)
+					.Text(InArgs._HeaderLabel)
+				]
+			];
+	}
+
+	// Enable Footer upon explicit request.
+	if (bIsFooterEnabled.Get() && ContentPanel.IsValid())
+	{
+		// Footer Panel
+		ContentPanel->AddSlot(bIsHeaderEnabled.Get() ? 2 : 1)
+			.Resizable(false)
+			.SizeRule(SSplitter::SizeToContent)
+			[
+				SNew(SHorizontalBox)
+
+				// Left Toolbar
+				+SHorizontalBox::Slot()
+				.AutoWidth()
+				.VAlign(VAlign_Center)
+				.Padding(4.f, 4.f, 2.f, 4.f)
+				[
+					SAssignNew(LeftToolbar, SHorizontalBox)
+				]
+
+				// Spacer.
+				+SHorizontalBox::Slot()
+				.FillWidth(1.f)
+				.VAlign(VAlign_Center)
+				[
+					SNew(SSpacer)
+				]
+				
+				// Right Toolbar
+				+SHorizontalBox::Slot()
+				.AutoWidth()
+				.Padding(2.f, 4.f, 4.f, 4.f)
+				.VAlign(VAlign_Center)
+				[
+					SAssignNew(RightToolbar, SHorizontalBox)
+				]
+			];
+	}
+}
+
+END_SLATE_FUNCTION_BUILD_OPTIMIZATION
+
+void SRCMajorPanel::AddPanel(TSharedRef<SWidget> InContent, const float InDesiredSize)
+{
+	if (Children.IsValid())
+	{
+		Children->AddSlot()
+			.Value(InDesiredSize)
+			[
+				InContent
+			];
+	}
+}
+
+const TSharedRef<SWidget>& SRCMajorPanel::GetContent() const
+{
+	return ChildSlot.GetWidget();
+}
+
+void SRCMajorPanel::ClearContent()
+{
+	ChildSlot.DetachWidget();
+}
+
+void SRCMajorPanel::AddFooterToolItem(EToolbar InToolbar, TSharedRef<SWidget> InWidget)
+{
+	switch (InToolbar)
+	{
+		case Left:
+			if (LeftToolbar.IsValid() && bIsFooterEnabled.Get())
+			{
+				LeftToolbar->AddSlot()
+					.Padding(5.f, 0.f)
+					.AutoWidth()
+					.VAlign(VAlign_Center)
+					[
+						InWidget
+					];
+			}
+			break;
+		case Right:
+			if (RightToolbar.IsValid() && bIsFooterEnabled.Get())
+			{
+				RightToolbar->AddSlot()
+					.Padding(5.f, 0.f)
+					.AutoWidth()
+					.VAlign(VAlign_Center)
+					[
+						InWidget
+					];
+			}
+			break;
+		default:
+			break;
+	}
+}
+
+SLATE_IMPLEMENT_WIDGET(SRCMinorPanel)
+
+void SRCMinorPanel::PrivateRegisterAttributes(FSlateAttributeInitializer& AttributeInitializer)
+{
+}
+
+BEGIN_SLATE_FUNCTION_BUILD_OPTIMIZATION
+
+void SRCMinorPanel::Construct(const SRCMinorPanel::FArguments& InArgs)
+{
+	RCPanelStyle = &FRemoteControlPanelStyle::Get()->GetWidgetStyle<FRCPanelStyle>("RemoteControlPanel.MinorPanel");
+
+	bIsFooterEnabled = InArgs._EnableFooter;
+	bIsHeaderEnabled = InArgs._EnableHeader;
+
+	ChildSlot
+	.Padding(RCPanelStyle->PanelPadding)
+	[
+		SNew(SBorder)
+		.BorderImage(&RCPanelStyle->ContentAreaBrush)
+		.Padding(RCPanelStyle->PanelPadding)
+		[	
+			SAssignNew(ContentPanel, SSplitter)
+			.Orientation(InArgs._Orientation)
+
+			// Content Panel
+			+SSplitter::Slot()
+			.Value(1.f)
+			[
+				InArgs._Content.Widget
+			]
+		]
+	];
+
+	// Enable Header upon explicit request.
+	if (bIsHeaderEnabled.Get() && ContentPanel.IsValid())
+	{
+		// Header Panel
+		ContentPanel->AddSlot(0)
+			.SizeRule(SSplitter::SizeToContent)
+			.Resizable(false)
+			[
+				SNew(SHorizontalBox)
+
+				// Left Toolbar
+				+SHorizontalBox::Slot()
+				.AutoWidth()
+				.VAlign(VAlign_Center)
+				.Padding(4.f, 4.f, 2.f, 4.f)
+				[
+					SAssignNew(LeftHeaderToolbar, SHorizontalBox)
+				]
+
+				// Header Label.
+				+SHorizontalBox::Slot()
+				.FillWidth(1.f)
+				.VAlign(VAlign_Center)
+				.Padding(2.f, 4.f)
+				[
+					SNew(STextBlock)
+					.TextStyle(&RCPanelStyle->HeaderTextStyle)
+					.Text(InArgs._HeaderLabel)
+				]
+				
+				// Right Toolbar
+				+SHorizontalBox::Slot()
+				.AutoWidth()
+				.Padding(2.f, 4.f, 4.f, 4.f)
+				.VAlign(VAlign_Center)
+				[
+					SAssignNew(RightHeaderToolbar, SHorizontalBox)
+				]
+			];
+	}
+	
+	// Enable Footer upon explicit request.
+	if (bIsFooterEnabled.Get() && ContentPanel.IsValid())
+	{
+		// Footer Panel
+		ContentPanel->AddSlot(bIsHeaderEnabled.Get() ? 2 : 1)
+			.SizeRule(SSplitter::SizeToContent)
+			.Resizable(false)
+			[
+				SNew(SHorizontalBox)
+
+				// Left Toolbar
+				+SHorizontalBox::Slot()
+				.AutoWidth()
+				.VAlign(VAlign_Center)
+				.Padding(4.f, 4.f, 2.f, 4.f)
+				[
+					SAssignNew(LeftFooterToolbar, SHorizontalBox)
+				]
+
+				// Spacer.
+				+SHorizontalBox::Slot()
+				.FillWidth(1.f)
+				.VAlign(VAlign_Center)
+				[
+					SNew(SSpacer)
+				]
+				
+				// Right Toolbar
+				+SHorizontalBox::Slot()
+				.AutoWidth()
+				.Padding(2.f, 4.f, 4.f, 4.f)
+				.VAlign(VAlign_Center)
+				[
+					SAssignNew(RightFooterToolbar, SHorizontalBox)
+				]
+			];
+	}
+}
+
+END_SLATE_FUNCTION_BUILD_OPTIMIZATION
+
+void SRCMinorPanel::SetContent(TSharedRef<SWidget> InContent)
+{
+	ChildSlot
+		[
+			InContent
+		];
+}
+
+const TSharedRef<SWidget>& SRCMinorPanel::GetContent() const
+{
+	return ChildSlot.GetWidget();
+}
+
+void SRCMinorPanel::ClearContent()
+{
+	ChildSlot.DetachWidget();
+}
+
+void SRCMinorPanel::AddFooterToolbarItem(EToolbar InToolbar, TSharedRef<SWidget> InWidget)
+{
+	switch (InToolbar)
+	{
+		case Left:
+			if (LeftFooterToolbar.IsValid() && bIsFooterEnabled.Get())
+			{
+				LeftFooterToolbar->AddSlot()
+					.Padding(5.f, 0.f)
+					.AutoWidth()
+					.VAlign(VAlign_Center)
+					[
+						InWidget
+					];
+			}
+			break;
+		case Right:
+			if (RightFooterToolbar.IsValid() && bIsFooterEnabled.Get())
+			{
+				RightFooterToolbar->AddSlot()
+					.Padding(5.f, 0.f)
+					.AutoWidth()
+					.VAlign(VAlign_Center)
+					[
+						InWidget
+					];
+			}
+			break;
+		default:
+			break;
+	}
+}
+
+void SRCMinorPanel::AddHeaderToolbarItem(EToolbar InToolbar, TSharedRef<SWidget> InWidget)
+{
+	switch (InToolbar)
+	{
+		case Left:
+			if (LeftHeaderToolbar.IsValid() && bIsHeaderEnabled.Get())
+			{
+				LeftHeaderToolbar->AddSlot()
+					.Padding(5.f, 0.f)
+					.AutoWidth()
+					.VAlign(VAlign_Center)
+					[
+						InWidget
+					];
+			}
+			break;
+		case Right:
+			if (RightHeaderToolbar.IsValid() && bIsHeaderEnabled.Get())
+			{
+				RightHeaderToolbar->AddSlot()
+					.Padding(5.f, 0.f)
+					.AutoWidth()
+					.VAlign(VAlign_Center)
+					[
+						InWidget
+					];
+			}
+			break;
+		default:
+			break;
+	}
+}
diff --git a/Engine/Plugins/VirtualProduction/RemoteControl/Source/RemoteControlUI/Private/UI/Panels/SRCDockPanel.h b/Engine/Plugins/VirtualProduction/RemoteControl/Source/RemoteControlUI/Private/UI/Panels/SRCDockPanel.h
new file mode 100644
--- /dev/null
+++ b/Engine/Plugins/VirtualProduction/RemoteControl/Source/RemoteControlUI/Private/UI/Panels/SRCDockPanel.h
@@ -0,0 +1,176 @@
+// Copyright Epic Games, Inc. All Rights Reserved.
+
+#pragma once
+
+#include "CoreMinimal.h"
+#include "Widgets/SCompoundWidget.h"
+
+struct FRCPanelStyle;
+class SHorizontalBox;
+class SSplitter;
+class SWidget;
+
+enum EToolbar
+{
+	/** Left toolbar. */
+	Left,
+
+	/** Right toolbar. */
+	Right
+};
+
+/**
+ * A custom slate widget representing a major panel on RC Panel.
+ */
+class SRCMajorPanel : public SCompoundWidget
+{
+	SLATE_DECLARE_WIDGET(SRCMajorPanel, SCompoundWidget)
+
+public:
+
+	SLATE_BEGIN_ARGS(SRCMajorPanel)
+		: _EnableFooter(true)
+		, _EnableHeader(true)
+		, _HeaderLabel(FText::GetEmpty())
+		, _Orientation(Orient_Vertical)
+		, _ChildOrientation(Orient_Horizontal)
+	{}
+		SLATE_ARGUMENT(bool, EnableFooter)
+
+		SLATE_ARGUMENT(bool, EnableHeader)
+			
+		SLATE_ARGUMENT(FText, HeaderLabel)
+
+		SLATE_ARGUMENT(EOrientation, Orientation)
+
+		SLATE_ARGUMENT(EOrientation, ChildOrientation)
+
+	SLATE_END_ARGS()
+
+	/** Constructs this widget with InArgs */
+	void Construct(const FArguments& InArgs);
+
+	/**
+	 * Adds the given panel to the children.
+	 *
+	 * @param InPanel The panel widget.
+	 */
+	virtual void AddPanel(TSharedRef<SWidget> InPanel, const float InDesiredSize);
+
+	/**
+	 * Gets the content of this panel
+	 *
+	 * @return The widget used as content for the panel
+	 */
+	const TSharedRef<SWidget>& GetContent() const;
+
+	/** Clears out the content for the panel */
+	void ClearContent();
+
+	/** Adds widgets to the specified footer toolbar (Left or Right). */
+	void AddFooterToolItem(EToolbar InToolbar, TSharedRef<SWidget> InWidget);
+	
+private:
+
+	/** Holds whether the footer panel is enabled or not. */
+	TAttribute<bool> bIsFooterEnabled;
+	
+	/** Holds whether the header panel is enabled or not. */
+	TAttribute<bool> bIsHeaderEnabled;
+
+	/** Actual toolbar widget located left to the footer. */
+	TSharedPtr<SHorizontalBox> LeftToolbar;
+
+	/** Actual toolbar widget located right to the footer. */
+	TSharedPtr<SHorizontalBox> RightToolbar;
+
+	/** Holds the entire content panel. */
+	TSharedPtr<SSplitter> ContentPanel;
+	
+	/** Holds the entire child panels. */
+	TSharedPtr<SSplitter> Children;
+
+	/** Panel Style reference. */
+	const FRCPanelStyle* RCPanelStyle;
+};
+
+/**
+ * A custom slate widget representing a minor panel on RC Panel.
+ */
+class SRCMinorPanel : public SCompoundWidget
+{
+	SLATE_DECLARE_WIDGET(SRCMinorPanel, SCompoundWidget)
+
+public:
+
+	SLATE_BEGIN_ARGS(SRCMinorPanel)
+		: _Content()
+		, _EnableFooter(false)
+		, _EnableHeader(true)
+		, _HeaderLabel(FText::GetEmpty())
+		, _Orientation(Orient_Vertical)
+	{}
+		SLATE_DEFAULT_SLOT(FArguments, Content)
+
+		SLATE_ARGUMENT(bool, EnableFooter)
+
+		SLATE_ARGUMENT(bool, EnableHeader)
+
+		SLATE_ARGUMENT(FText, HeaderLabel)
+
+		SLATE_ARGUMENT(EOrientation, Orientation)
+
+	SLATE_END_ARGS()
+
+	/** Constructs this widget with InArgs */
+	void Construct(const FArguments& InArgs);
+
+	/**
+	 * Sets the content for this panel
+	 *
+	 * @param	InContent	The widget to use as content for the panel
+	 */
+	virtual void SetContent(TSharedRef<SWidget> InContent );
+
+	/**
+	 * Gets the content for this panel
+	 *
+	 * @return The widget used as content for the panel
+	 */
+	const TSharedRef<SWidget>& GetContent() const;
+
+	/** Clears out the content for the panel */
+	void ClearContent();
+
+	/** Adds widgets to the specified footer toolbar (Left or Right). */
+	void AddFooterToolbarItem(EToolbar InToolbar, TSharedRef<SWidget> InWidget);
+	
+	/** Adds widgets to the specified header toolbar (Left or Right). */
+	void AddHeaderToolbarItem(EToolbar InToolbar, TSharedRef<SWidget> InWidget);
+	
+private:
+
+	/** Holds whether the footer panel is enabled or not. */
+	TAttribute<bool> bIsFooterEnabled;
+	
+	/** Holds whether the header panel is enabled or not. */
+	TAttribute<bool> bIsHeaderEnabled;
+
+	/** Actual toolbar widget located left to the footer. */
+	TSharedPtr<SHorizontalBox> LeftFooterToolbar;
+
+	/** Actual toolbar widget located right to the footer. */
+	TSharedPtr<SHorizontalBox> RightFooterToolbar;
+	
+	/** Actual toolbar widget located left to the header. */
+	TSharedPtr<SHorizontalBox> LeftHeaderToolbar;
+
+	/** Actual toolbar widget located right to the header. */
+	TSharedPtr<SHorizontalBox> RightHeaderToolbar;
+
+	/** Holds the entire content panel. */
+	TSharedPtr<SSplitter> ContentPanel;
+
+	/** Panel Style reference. */
+	const FRCPanelStyle* RCPanelStyle;
+};
diff --git a/Engine/Plugins/VirtualProduction/RemoteControl/Source/RemoteControlUI/Private/UI/RCUIHelpers.cpp b/Engine/Plugins/VirtualProduction/RemoteControl/Source/RemoteControlUI/Private/UI/RCUIHelpers.cpp
--- a/Engine/Plugins/VirtualProduction/RemoteControl/Source/RemoteControlUI/Private/UI/RCUIHelpers.cpp
+++ b/Engine/Plugins/VirtualProduction/RemoteControl/Source/RemoteControlUI/Private/UI/RCUIHelpers.cpp
@@ -25,8 +25,10 @@ FLinearColor UE::RCUIHelpers::GetFieldClassTypeColor(const FProperty* InProperty
 	return FLinearColor::White;
 }
 
-FText UE::RCUIHelpers::GetFieldClassDisplayName(const FProperty* InProperty)
+FName UE::RCUIHelpers::GetFieldClassDisplayName(const FProperty* InProperty)
 {
+	FName FieldClassDisplayName = NAME_None;
+
 	if (ensure(InProperty))
 	{
 		const UEdGraphSchema_K2* Schema = GetDefault<UEdGraphSchema_K2>();
@@ -36,12 +38,27 @@ FText UE::RCUIHelpers::GetFieldClassDisplayName(const FProperty* InProperty)
 		{
 			if(UObject* SubCategoryObject = PinType.PinSubCategoryObject.Get())
 			{
-				return FText::FromName(*SubCategoryObject->GetName());
+				FieldClassDisplayName = *SubCategoryObject->GetName();
 			}
 			
-			return FText::FromName(PinType.PinCategory);
+			FieldClassDisplayName  = PinType.PinCategory;
+		}
+	}
+
+	if (!FieldClassDisplayName.IsNone())
+	{
+		FString FieldClassNameStr = FieldClassDisplayName.ToString();
+
+		if (FieldClassNameStr.StartsWith("Bool"))
+		{
+			FieldClassDisplayName = "Boolean";
+		}
+
+		if (FieldClassNameStr.StartsWith("Int"))
+		{
+			FieldClassDisplayName = *FieldClassNameStr.Replace(TEXT("Int"), TEXT("Integer"));
 		}
 	}
 
-	return FText::GetEmpty();
+	return FieldClassDisplayName;
 }
\ No newline at end of file
diff --git a/Engine/Plugins/VirtualProduction/RemoteControl/Source/RemoteControlUI/Private/UI/RCUIHelpers.h b/Engine/Plugins/VirtualProduction/RemoteControl/Source/RemoteControlUI/Private/UI/RCUIHelpers.h
--- a/Engine/Plugins/VirtualProduction/RemoteControl/Source/RemoteControlUI/Private/UI/RCUIHelpers.h
+++ b/Engine/Plugins/VirtualProduction/RemoteControl/Source/RemoteControlUI/Private/UI/RCUIHelpers.h
@@ -9,13 +9,13 @@ class URCController;
 namespace UE::RCUIHelpers
 {
 	/** GetFieldClassTypeColor
-	* Fetches the editor color associated with a given Unreal Type (FProperty)
-	* Used to provide color coding in the Remote Control Logic Actions panel 
-	*/
+	 * Fetches the editor color associated with a given Unreal Type (FProperty)
+	 * Used to provide color coding in the Remote Control Logic Actions panel 
+	 */
 	FLinearColor GetFieldClassTypeColor(const FProperty* InProperty);
 
 	/** GetFieldClassDisplayName
-	* Fetches the display name associated with a given Unreal Type (FProperty)
-	*/
-	FText GetFieldClassDisplayName(const FProperty* InProperty);
+	 * Fetches the display name associated with a given Unreal Type (FProperty)
+	 */
+	FName GetFieldClassDisplayName(const FProperty* InProperty);
 }
\ No newline at end of file
diff --git a/Engine/Plugins/VirtualProduction/RemoteControl/Source/RemoteControlUI/Private/UI/RemoteControlPanelStyle.cpp b/Engine/Plugins/VirtualProduction/RemoteControl/Source/RemoteControlUI/Private/UI/RemoteControlPanelStyle.cpp
--- a/Engine/Plugins/VirtualProduction/RemoteControl/Source/RemoteControlUI/Private/UI/RemoteControlPanelStyle.cpp
+++ b/Engine/Plugins/VirtualProduction/RemoteControl/Source/RemoteControlUI/Private/UI/RemoteControlPanelStyle.cpp
@@ -2,9 +2,12 @@
 
 #include "UI/RemoteControlPanelStyle.h"
 #include "Styling/AppStyle.h"
+#include "Styling/CoreStyle.h"
+#include "Styling/RemoteControlStyles.h"
 #include "Styling/SlateStyleRegistry.h"
 #include "Styling/SlateTypes.h"
-#include "Styling/CoreStyle.h"
+#include "Styling/StarshipCoreStyle.h"
+#include "Styling/StyleColors.h"
 #include "Interfaces/IPluginManager.h"
 #include "SlateOptMacros.h"
 #include "Widgets/Input/SButton.h"
@@ -90,40 +93,10 @@ void FRemoteControlPanelStyle::Initialize()
 
 	StyleSet->Set("RemoteControlPanel.GroupRowSelected", new BOX_BRUSH("Common/GroupBorderLight", FMargin(4.0f / 16.0f), NewSelectionColor));
 	StyleSet->Set("RemoteControlPanel.TransparentBorder", new FSlateColorBrush(FColor::Transparent));
-
-	FTableRowStyle GroupRowStyle = AppStyle.GetWidgetStyle<FTableRowStyle>("TableView.Row");
-	GroupRowStyle.SetSelectorFocusedBrush(*FRemoteControlPanelStyle::Get()->GetBrush("RemoteControlPanel.GroupRowSelected"));
-
-	GroupRowStyle.SetEvenRowBackgroundBrush(*FRemoteControlPanelStyle::Get()->GetBrush("RemoteControlPanel.GroupBorder"));
-	GroupRowStyle.SetOddRowBackgroundBrush(*FRemoteControlPanelStyle::Get()->GetBrush("RemoteControlPanel.GroupBorder"));
-	GroupRowStyle.SetEvenRowBackgroundHoveredBrush(*FRemoteControlPanelStyle::Get()->GetBrush("RemoteControlPanel.GroupBorder"));
-	GroupRowStyle.SetOddRowBackgroundHoveredBrush(*FRemoteControlPanelStyle::Get()->GetBrush("RemoteControlPanel.GroupBorder"));
-	GroupRowStyle.SetActiveHoveredBrush(*FRemoteControlPanelStyle::Get()->GetBrush("RemoteControlPanel.GroupRowSelected"));
-	GroupRowStyle.SetInactiveHoveredBrush(*FRemoteControlPanelStyle::Get()->GetBrush("RemoteControlPanel.GroupRowSelected"));
-	GroupRowStyle.SetActiveBrush(*FRemoteControlPanelStyle::Get()->GetBrush("RemoteControlPanel.GroupRowSelected"));
-	GroupRowStyle.SetInactiveBrush(*FRemoteControlPanelStyle::Get()->GetBrush("RemoteControlPanel.GroupRowSelected"));
 	
-	StyleSet->Set("RemoteControlPanel.GroupRow", GroupRowStyle);
-	
-	// Checkbox for live/edit mode
-	
-	{
-		StyleSet->Set("Switch.ToggleOff", new IMAGE_PLUGIN_BRUSH("Icons/Switch_OFF", Icon28x14));
-		
-		StyleSet->Set("Switch.ToggleOn", new IMAGE_PLUGIN_BRUSH("Icons/Switch_ON", Icon28x14));
-
-		FCheckBoxStyle SwitchStyle = FCheckBoxStyle()
-			.SetForegroundColor(FLinearColor::White)
-			.SetUncheckedImage(*StyleSet->GetBrush("Switch.ToggleOff"))
-			.SetUncheckedHoveredImage(*StyleSet->GetBrush("Switch.ToggleOff"))
-			.SetUncheckedPressedImage(*StyleSet->GetBrush("Switch.ToggleOff"))
-			.SetCheckedImage(*StyleSet->GetBrush("Switch.ToggleOn"))
-			.SetCheckedHoveredImage(*StyleSet->GetBrush("Switch.ToggleOn"))
-			.SetCheckedPressedImage(*StyleSet->GetBrush("Switch.ToggleOn"))
-			.SetPadding(FMargin(1, 1, 1, 1));
-
-		StyleSet->Set("RemoteControlPanel.Switch", SwitchStyle);
-	} 
+	// Toggle Button Style
+	StyleSet->Set("Switch.ToggleOff", new IMAGE_PLUGIN_BRUSH("Icons/Switch_OFF", Icon28x14));
+	StyleSet->Set("Switch.ToggleOn", new IMAGE_PLUGIN_BRUSH("Icons/Switch_ON", Icon28x14));
 
 	// Remote Control Logic UI
 	StyleSet->Set("RemoteControlPanel.Behaviours.Title", FCoreStyle::GetDefaultFontStyle("Regular", 12));
@@ -131,6 +104,9 @@ void FRemoteControlPanelStyle::Initialize()
 	StyleSet->Set("RemoteControlPanel.Behaviours.CustomBlueprint", new IMAGE_BRUSH_SVG("Starship/MainToolbar/blueprints", Icon20x20));
 	StyleSet->Set("RemoteControlPanel.Actions.ValuePanelHeader", FCoreStyle::GetDefaultFontStyle("Bold", 12));
 
+	// Remote Control Panel UI
+	SetupPanelStyles(StyleSet.ToSharedRef());
+
 	FSlateStyleRegistry::RegisterSlateStyle(*StyleSet.Get());
 }
 
@@ -161,6 +137,77 @@ FString FRemoteControlPanelStyle::InContent(const FString& RelativePath, const A
 	return (ContentDir / RelativePath) + Extension;
 }
 
+void FRemoteControlPanelStyle::SetupPanelStyles(TSharedRef<FSlateStyleSet> InStyle)
+{
+	const ISlateStyle& AppStyle = FAppStyle::Get();
+
+	// Flat Button Style
+	FButtonStyle FlatButtonStyle = AppStyle.GetWidgetStyle<FButtonStyle>("DetailsView.NameAreaButton");
+	FlatButtonStyle.Normal = *AppStyle.GetBrush("NoBorder");
+	FlatButtonStyle.Disabled = *AppStyle.GetBrush("NoBorder");
+
+	// Toggle Button Style
+	FCheckBoxStyle ToggleButtonStyle = FCheckBoxStyle()
+		.SetForegroundColor(FLinearColor::White)
+		.SetUncheckedImage(*StyleSet->GetBrush("Switch.ToggleOff"))
+		.SetUncheckedHoveredImage(*StyleSet->GetBrush("Switch.ToggleOff"))
+		.SetUncheckedPressedImage(*StyleSet->GetBrush("Switch.ToggleOff"))
+		.SetCheckedImage(*StyleSet->GetBrush("Switch.ToggleOn"))
+		.SetCheckedHoveredImage(*StyleSet->GetBrush("Switch.ToggleOn"))
+		.SetCheckedPressedImage(*StyleSet->GetBrush("Switch.ToggleOn"))
+		.SetPadding(FMargin(1, 1, 1, 1));
+
+	// Table View Style
+	FTableViewStyle TableViewStyle = AppStyle.GetWidgetStyle<FTableViewStyle>("SimpleListView");
+	TableViewStyle.BackgroundBrush = FSlateColorBrush(FStyleColors::Background);
+
+	// Header Row Style
+	FHeaderRowStyle HeaderRowStyle = AppStyle.GetWidgetStyle<FHeaderRowStyle>("PropertyTable.HeaderRow");
+	HeaderRowStyle.BackgroundBrush = FSlateColorBrush(FStyleColors::Background);
+	HeaderRowStyle.HorizontalSeparatorBrush = FSlateColorBrush(FStyleColors::Background);
+	HeaderRowStyle.HorizontalSeparatorThickness = 2.f;
+
+	// Table Row Style
+	FTableRowStyle TableRowStyle = AppStyle.GetWidgetStyle<FTableRowStyle>("DetailsView.TreeView.TableRow");
+	TableRowStyle.SetSelectorFocusedBrush(FSlateColorBrush(FStyleColors::Highlight));
+
+	TableRowStyle.SetEvenRowBackgroundBrush(FSlateColorBrush(FStyleColors::Panel));
+	TableRowStyle.SetOddRowBackgroundBrush(FSlateColorBrush(FStyleColors::Panel));
+	TableRowStyle.SetEvenRowBackgroundHoveredBrush(FSlateColorBrush(FStyleColors::Background));
+	TableRowStyle.SetOddRowBackgroundHoveredBrush(FSlateColorBrush(FStyleColors::Background));
+	TableRowStyle.SetActiveHoveredBrush(FSlateColorBrush(FStyleColors::Select));
+	TableRowStyle.SetInactiveHoveredBrush(FSlateColorBrush(FStyleColors::SelectInactive));
+	TableRowStyle.SetActiveBrush(FSlateColorBrush(FStyleColors::Select));
+	TableRowStyle.SetInactiveBrush(FSlateColorBrush(FStyleColors::SelectInactive));
+
+	StyleSet->Set("RemoteControlPanel.GroupRow", TableRowStyle);
+
+	// Text Styles
+	const FStyleFonts& StyleFonts = FStyleFonts::Get();
+
+	const FTextBlockStyle& NormalText = AppStyle.GetWidgetStyle<FTextBlockStyle>("NormalText");
+
+	const FTextBlockStyle HeaderText = FTextBlockStyle(NormalText)
+		.SetFont(StyleFonts.Large);
+
+	FRCPanelStyle RCMinorPanelStyle = FRCPanelStyle()
+		.SetFlatButtonStyle(FlatButtonStyle)
+		.SetToggleButtonStyle(ToggleButtonStyle)
+
+		.SetContentAreaBrush(FSlateColorBrush(FStyleColors::Panel))
+		.SetSectionHeaderBrush(FSlateColorBrush(FStyleColors::Foldout))
+
+		.SetHeaderRowStyle(HeaderRowStyle)
+		.SetTableRowStyle(TableRowStyle)
+		.SetTableViewStyle(TableViewStyle)
+
+		.SetHeaderTextStyle(HeaderText)
+		.SetPanelTextStyle(NormalText)
+		.SetSectionHeaderTextStyle(HeaderText);
+
+	StyleSet->Set("RemoteControlPanel.MinorPanel", RCMinorPanelStyle);
+}
+
 #undef IMAGE_PLUGIN_BRUSH
 #undef IMAGE_BRUSH
 #undef IMAGE_BRUSH_SVG
diff --git a/Engine/Plugins/VirtualProduction/RemoteControl/Source/RemoteControlUI/Private/UI/RemoteControlPanelStyle.h b/Engine/Plugins/VirtualProduction/RemoteControl/Source/RemoteControlUI/Private/UI/RemoteControlPanelStyle.h
--- a/Engine/Plugins/VirtualProduction/RemoteControl/Source/RemoteControlUI/Private/UI/RemoteControlPanelStyle.h
+++ b/Engine/Plugins/VirtualProduction/RemoteControl/Source/RemoteControlUI/Private/UI/RemoteControlPanelStyle.h
@@ -32,6 +32,8 @@ private:
 	/** Get a path to the module's resources folder. */
 	static FString InContent(const FString& RelativePath, const ANSICHAR* Extension);
 
+	static void SetupPanelStyles(TSharedRef<FSlateStyleSet> InStyle);
+
 private:
 	/** Holds the style set. */
 	static TSharedPtr<FSlateStyleSet> StyleSet;
diff --git a/Engine/Plugins/VirtualProduction/RemoteControl/Source/RemoteControlUI/Private/UI/SRCPanelDragHandle.h b/Engine/Plugins/VirtualProduction/RemoteControl/Source/RemoteControlUI/Private/UI/SRCPanelDragHandle.h
--- a/Engine/Plugins/VirtualProduction/RemoteControl/Source/RemoteControlUI/Private/UI/SRCPanelDragHandle.h
+++ b/Engine/Plugins/VirtualProduction/RemoteControl/Source/RemoteControlUI/Private/UI/SRCPanelDragHandle.h
@@ -34,8 +34,9 @@ public:
 			[
 				SNew(SHorizontalBox)
 				+ SHorizontalBox::Slot()
-				.Padding(5.0f, 0.0f)
+				.Padding(5.f, 0.f)
 				.VAlign(VAlign_Center)
+				.AutoWidth()
 				[
 					SNew(SImage)
 					.Image(FCoreStyle::Get().GetBrush("VerticalBoxDragIndicatorShort"))
diff --git a/Engine/Plugins/VirtualProduction/RemoteControl/Source/RemoteControlUI/Private/UI/SRCPanelExposedEntitiesList.cpp b/Engine/Plugins/VirtualProduction/RemoteControl/Source/RemoteControlUI/Private/UI/SRCPanelExposedEntitiesList.cpp
--- a/Engine/Plugins/VirtualProduction/RemoteControl/Source/RemoteControlUI/Private/UI/SRCPanelExposedEntitiesList.cpp
+++ b/Engine/Plugins/VirtualProduction/RemoteControl/Source/RemoteControlUI/Private/UI/SRCPanelExposedEntitiesList.cpp
@@ -13,28 +13,131 @@
 #include "Editor/EditorEngine.h"
 #include "Misc/Attribute.h"
 #include "Misc/Guid.h"
+#include "Misc/MessageDialog.h"
 #include "RemoteControlPanelStyle.h"
 #include "RemoteControlPreset.h"
 #include "RemoteControlUIModule.h"
 #include "ScopedTransaction.h"
 #include "SRCPanelFieldGroup.h"
 #include "SRCPanelExposedField.h"
+#include "Styling/RemoteControlStyles.h"
+#include "UI/Panels/SRCDockPanel.h"
 #include "UObject/Object.h"
 #include "Widgets/Input/SButton.h"
 #include "Widgets/Layout/SSpacer.h"
+#include "Widgets/Layout/SSeparator.h"
+#include "Widgets/Views/SListView.h"
 #include "Widgets/Views/STableRow.h"
 #include "Widgets/Views/STreeView.h"
 
 #define LOCTEXT_NAMESPACE "RemoteControlPanelEntitiesList"
 
-namespace FRemoteControlPresetColumns
+class SEntityRow : public SMultiColumnTableRow<TSharedPtr<SRCPanelTreeNode>>
 {
-	const FName DragDropHandle = TEXT("DragDropHandle");
-	const FName IsBound = TEXT("IsBound");
-	const FName Description = TEXT("Description");
-	const FName Value = TEXT("Value");
-	const FName Groups = TEXT("Groups");
-}
+public:
+
+	SLATE_BEGIN_ARGS(SEntityRow)
+	{}
+	
+		SLATE_ARGUMENT(TSharedPtr<SRCPanelTreeNode>, Entity)
+
+		SLATE_ATTRIBUTE(FMargin, Padding)
+		SLATE_STYLE_ARGUMENT(FTableRowStyle, Style)
+
+		// Low level DragAndDrop
+		SLATE_EVENT(FOnDragDetected, OnDragDetected)
+		SLATE_EVENT(FOnTableRowDragEnter, OnDragEnter)
+		SLATE_EVENT(FOnTableRowDragLeave, OnDragLeave)
+		SLATE_EVENT(FOnTableRowDrop, OnDrop)
+
+	SLATE_END_ARGS()
+
+	void Construct(const FArguments& InArgs, const TSharedRef<STableViewBase>& OwnerTableView)
+	{
+		Entity = InArgs._Entity;
+
+		FSuperRowType::FArguments SuperArgs = FSuperRowType::FArguments();
+
+		SuperArgs.OnDragDetected(InArgs._OnDragDetected);
+		SuperArgs.OnDragEnter(InArgs._OnDragEnter);
+		SuperArgs.OnDragLeave(InArgs._OnDragLeave);
+		SuperArgs.OnDrop(InArgs._OnDrop);
+
+		SuperArgs.ExpanderStyleSet(&FCoreStyle::Get());
+		SuperArgs.Padding(InArgs._Padding);
+		SuperArgs.ShowWires(false);
+		SuperArgs.Style(InArgs._Style);
+
+		FSuperRowType::Construct(SuperArgs, OwnerTableView);
+	}
+
+	TSharedRef<SWidget> GenerateWidgetForColumn(const FName& InColumnName) override
+	{
+		if (Entity.IsValid())
+		{
+			if (Entity->HasChildren() && InColumnName == FRemoteControlPresetColumns::Description)
+			{
+				// -- Row is for TreeView --
+				SHorizontalBox::FSlot* InnerContentSlotNativePtr = nullptr;
+
+				TSharedRef<SWidget> TreeNode = SNew(SHorizontalBox)
+					.Clipping(EWidgetClipping::OnDemand)
+					
+					+ SHorizontalBox::Slot()
+					.AutoWidth()
+					.HAlign(HAlign_Right)
+					.VAlign(VAlign_Fill)
+					[
+						SAssignNew(ExpanderArrowWidget, SExpanderArrow, SharedThis(this))
+						.StyleSet(ExpanderStyleSet)
+						.ShouldDrawWires(false)
+					]
+
+					+ SHorizontalBox::Slot()
+					.FillWidth(1)
+					.Expose(InnerContentSlotNativePtr)
+					.Padding(8.f, 2.f)
+					[
+						Entity->GetWidget(InColumnName)
+					];
+
+				InnerContentSlot = InnerContentSlotNativePtr;
+
+				return TreeNode;
+
+			}
+			else
+			{
+				return SNew(SBox)
+					.Padding(FMargin(4.f, 2.f))
+					.VAlign(VAlign_Center)
+					[
+						Entity->GetWidget(InColumnName)
+					];
+			}
+		}
+
+		return SNullWidget::NullWidget;
+	}
+
+private:
+
+	//~ SWidget Interface
+	virtual FReply OnMouseButtonDoubleClick(const FGeometry& InMyGeometry, const FPointerEvent& InMouseEvent) override
+	{
+		if (Entity.IsValid())
+		{
+			Entity->EnterRenameMode();
+		}
+
+		return FSuperRowType::OnMouseButtonDoubleClick(InMyGeometry, InMouseEvent);
+	}
+
+private:
+
+	/** Cached reference of Entity */
+	TSharedPtr<SRCPanelTreeNode> Entity;
+};
 
 void SRCPanelExposedEntitiesList::Construct(const FArguments& InArgs, URemoteControlPreset* InPreset, TWeakPtr<FRCPanelWidgetRegistry> InWidgetRegistry)
 {
@@ -50,143 +153,181 @@ void SRCPanelExposedEntitiesList::Construct(const FArguments& InArgs, URemoteCon
 	bFilterApplicationRequested = false;
 	bSearchRequested = false;
 	SearchedText = MakeShared<FText>();
-	FilterCount = 0;
-	SearchCount = 0;
 
-	ChildSlot
-	[
-		SNew(SVerticalBox)
-		
-		+SVerticalBox::Slot()
-		.FillHeight(1.f)
+	RCPanelStyle = &FRemoteControlPanelStyle::Get()->GetWidgetStyle<FRCPanelStyle>("RemoteControlPanel.MinorPanel");
+
+	// Major Panel
+	TSharedPtr<SRCMajorPanel> ExposePanel = SNew(SRCMajorPanel)
+		.HeaderLabel(LOCTEXT("ExposePanelHeader", "EXPOSE"))
+		.EnableFooter(false);
+
+	// Groups List
+	SAssignNew(GroupsListView, SListView<TSharedPtr<SRCPanelTreeNode>>)
+		.ItemHeight(24.f)
+		.OnGenerateRow(this, &SRCPanelExposedEntitiesList::OnGenerateRow)
+		.OnSelectionChanged(this, &SRCPanelExposedEntitiesList::OnSelectionChanged)
+		.SelectionMode(ESelectionMode::Single)
+		.ListItemsSource(reinterpret_cast<TArray<TSharedPtr<SRCPanelTreeNode>>*>(&FieldGroups))
+		.OnContextMenuOpening(this, &SRCPanelExposedEntitiesList::OnContextMenuOpening, SRCPanelTreeNode::Group)
+		.ListViewStyle(&RCPanelStyle->TableViewStyle)
+		.ClearSelectionOnClick(true);
+
+	// Group Dock Panel
+	TSharedPtr<SRCMinorPanel> GroupDockPanel = SNew(SRCMinorPanel)
+		.HeaderLabel(LOCTEXT("GroupsLabel", "Group"))
+		.EnableFooter(false)
 		[
-			SNew(SHorizontalBox)
-			
-			+SHorizontalBox::Slot()
-			.Padding(2.f, 4.f)
-			.FillWidth(0.25f)
+			GroupsListView.ToSharedRef()
+		];
+
+	// Add New Group Button
+	TSharedPtr<SWidget> NewGroupButton = SNew(SButton)
+		.AddMetaData<FTagMetaData>(FTagMetaData(TEXT("Add Group")))
+		.IsEnabled_Lambda([this]() { return bIsInEditMode.Get(); })
+		.HAlign(HAlign_Center)
+		.VAlign(VAlign_Center)
+		.ForegroundColor(FSlateColor::UseForeground())
+		.ButtonStyle(&RCPanelStyle->FlatButtonStyle)
+		.ToolTipText(LOCTEXT("NewGroupToolTip", "Create new group."))
+		.OnClicked(this, &SRCPanelExposedEntitiesList::OnCreateGroup)
+		[
+			SNew(SBox)
+			.WidthOverride(RCPanelStyle->IconSize.X)
+			.HeightOverride(RCPanelStyle->IconSize.Y)
 			[
-				SNew(SBorder)
-				.BorderImage(FAppStyle::Get().GetBrush("ToolPanel.DarkGroupBorder"))
-				[
-					SAssignNew(GroupsListView, SListView<TSharedPtr<SRCPanelTreeNode>>)
-					.ItemHeight(24.f)
-					.OnGenerateRow(this, &SRCPanelExposedEntitiesList::OnGenerateRow)
-					.OnSelectionChanged(this, &SRCPanelExposedEntitiesList::OnSelectionChanged)
-					.SelectionMode(ESelectionMode::Single)
-					.ListItemsSource(reinterpret_cast<TArray<TSharedPtr<SRCPanelTreeNode>>*>(&FieldGroups))
-					.OnContextMenuOpening(this, &SRCPanelExposedEntitiesList::OnContextMenuOpening, SRCPanelTreeNode::Group)
-					.ClearSelectionOnClick(true)
-					.HeaderRow(
-						SNew(SHeaderRow)
-
-						+ SHeaderRow::Column(FRemoteControlPresetColumns::Groups)
-						.DefaultLabel(LOCTEXT("RCPresetGroupsColumnHeader", "Groups"))
-						.FillWidth(1.f)
-						.ToolTipText(LOCTEXT("NewGroupToolTip", "Create new group."))
-						[
-							SNew(SHorizontalBox)
-				
-							+SHorizontalBox::Slot()
-							.Padding(3.f)
-							.VAlign(VAlign_Center)
-							.AutoWidth()
-							[
-								SNew(STextBlock)
-								.ColorAndOpacity(FColor::White)
-								.Font(FCoreStyle::GetDefaultFontStyle("Bold", 12.f))
-								.Text(LOCTEXT("GroupsLabel", "Groups"))
-							]
-	
-							+SHorizontalBox::Slot()
-							.Padding(3.f)
-							.VAlign(VAlign_Center)
-							.FillWidth(1.f)
-							[
-								SNew(SSpacer)
-							]
-	
-							+ SHorizontalBox::Slot()
-							.Padding(3.f)
-							.HAlign(HAlign_Center)
-							.VAlign(VAlign_Center)
-							.AutoWidth()
-							[
-								SNew(SButton)
-								.Visibility_Lambda([this]() { return bIsInEditMode.Get() ? EVisibility::Visible : EVisibility::Collapsed; })
-								.ButtonStyle(FAppStyle::Get(), "ToggleButton")
-								.ToolTipText(LOCTEXT("NewGroupToolTip", "Create new group."))
-								.OnClicked(this, &SRCPanelExposedEntitiesList::OnCreateGroup)
-								[
-									SNew(SOverlay)
-
-									+SOverlay::Slot()
-									.HAlign(HAlign_Center)
-									.VAlign(VAlign_Center)
-									[
-										SNew(STextBlock)
-										.TextStyle(FRemoteControlPanelStyle::Get(), "RemoteControlPanel.Button.TextStyle")
-										.ColorAndOpacity(FColor::White)
-										.Font(FAppStyle::Get().GetFontStyle("FontAwesome.14"))
-										.Text(FEditorFontGlyphs::Circle_O)
-									]
-
-									+SOverlay::Slot()
-									.HAlign(HAlign_Center)
-									.VAlign(VAlign_Center)
-									[
-										SNew(STextBlock)
-										.TextStyle(FRemoteControlPanelStyle::Get(), "RemoteControlPanel.Button.TextStyle")
-										.ColorAndOpacity(FColor::White)
-										.Font(FAppStyle::Get().GetFontStyle("FontAwesome.7"))
-										.Text(FEditorFontGlyphs::Plus)
-									]
-								]
-							]
-						]
-					)
-				]
+				SNew(SImage)
+				.ColorAndOpacity(FSlateColor::UseForeground())
+				.Image(FAppStyle::GetBrush("Icons.PlusCircle"))
+			]
+		];
+
+	// Delete All Groups Button
+	TSharedPtr<SWidget> DeleteAllGroupsButton = SNew(SButton)
+		.AddMetaData<FTagMetaData>(FTagMetaData(TEXT("Delete All Groups")))
+		.IsEnabled_Lambda([this]() { return bIsInEditMode.Get() && FieldGroups.Num() > 0; })
+		.Visibility_Lambda([this]() { return (bIsInEditMode.Get() && FieldGroups.Num() > 0) ? EVisibility::Visible : EVisibility::Collapsed; })
+		.HAlign(HAlign_Center)
+		.VAlign(VAlign_Center)
+		.ForegroundColor(FSlateColor::UseForeground())
+		.ButtonStyle(&RCPanelStyle->FlatButtonStyle)
+		.ToolTipText(LOCTEXT("DeleteAllGroupsToolTip", "Delete all groups."))
+		.OnClicked(this, &SRCPanelExposedEntitiesList::RequestDeleteAllGroups)
+		[
+			SNew(SBox)
+			.WidthOverride(RCPanelStyle->IconSize.X)
+			.HeightOverride(RCPanelStyle->IconSize.Y)
+			[
+				SNew(SImage)
+				.ColorAndOpacity(FSlateColor::UseForeground())
+				.Image(FAppStyle::GetBrush("Icons.Delete"))
+			]
+		];
+
+	GroupDockPanel->AddHeaderToolbarItem(EToolbar::Left, NewGroupButton.ToSharedRef());
+	GroupDockPanel->AddHeaderToolbarItem(EToolbar::Right, DeleteAllGroupsButton.ToSharedRef());
+
+	ExposePanel->AddPanel(GroupDockPanel.ToSharedRef(), 0.25f);
+
+	// Fields List
+	SAssignNew(FieldsListView, STreeView<TSharedPtr<SRCPanelTreeNode>>)
+		.ItemHeight(24.f)
+		.OnGenerateRow(this, &SRCPanelExposedEntitiesList::OnGenerateRow)
+		.OnSelectionChanged(this, &SRCPanelExposedEntitiesList::OnSelectionChanged)
+		.SelectionMode(ESelectionMode::Single)
+		.TreeItemsSource(&FieldEntities)
+		.OnContextMenuOpening(this, &SRCPanelExposedEntitiesList::OnContextMenuOpening, SRCPanelTreeNode::Field)
+		.ClearSelectionOnClick(true)
+		.TreeViewStyle(&RCPanelStyle->TableViewStyle)
+		.OnGetChildren(this, &SRCPanelExposedEntitiesList::OnGetNodeChildren)
+		.HeaderRow(
+			SNew(SHeaderRow)
+			.Style(&RCPanelStyle->HeaderRowStyle)
+
+			+ SHeaderRow::Column(FRemoteControlPresetColumns::DragDropHandle)
+			.DefaultLabel(LOCTEXT("RCPresetDragDropHandleColumnHeader", ""))
+			.FixedWidth(25.f)
+			.HeaderContentPadding(RCPanelStyle->HeaderRowPadding)
+
+			+ SHeaderRow::Column(FRemoteControlPresetColumns::Description)
+			.DefaultLabel(LOCTEXT("RCPresetDescColumnHeader", "Description"))
+			.FillWidth(0.5f)
+			.HeaderContentPadding(RCPanelStyle->HeaderRowPadding)
+
+			+ SHeaderRow::Column(FRemoteControlPresetColumns::Value)
+			.DefaultLabel(LOCTEXT("RCPresetValueColumnHeader", "Value"))
+			.FillWidth(0.5f)
+			.HeaderContentPadding(RCPanelStyle->HeaderRowPadding)
+
+			+ SHeaderRow::Column(FRemoteControlPresetColumns::Reset)
+			.DefaultLabel(LOCTEXT("RCPresetResetButtonColumnHeader", ""))
+			.FixedWidth(48.f)
+			.HeaderContentPadding(RCPanelStyle->HeaderRowPadding)
+		);
+
+	// Exposed Entities Dock Panel
+	TSharedPtr<SRCMinorPanel> ExposeDockPanel = SNew(SRCMinorPanel)
+		.HeaderLabel(LOCTEXT("PropertiesLabel", "Properties"))
+		.EnableFooter(true)
+		[
+			FieldsListView.ToSharedRef()
+		];
+
+	// Placeholder Box
+	TSharedPtr<SWidget> PlaceholderBox = SNew(SBox)
+		.AddMetaData<FTagMetaData>(FTagMetaData(TEXT("Placeholder")))
+		.WidthOverride(30.f)
+		.HeightOverride(30.f)
+		[
+			SNullWidget::NullWidget
+		];
+
+	// Delete All Entities Button
+	TSharedPtr<SWidget> DeleteAllEntitiesButton = SNew(SButton)
+		.AddMetaData<FTagMetaData>(FTagMetaData(TEXT("Delete All Entities")))
+		.IsEnabled_Lambda([this]() { return bIsInEditMode.Get() && FieldWidgetMap.Num() > 0; })
+		.Visibility_Lambda([this]() { return (bIsInEditMode.Get() && FieldWidgetMap.Num() > 0) ? EVisibility::Visible : EVisibility::Collapsed; })
+		.HAlign(HAlign_Center)
+		.VAlign(VAlign_Center)
+		.ForegroundColor(FSlateColor::UseForeground())
+		.ButtonStyle(&RCPanelStyle->FlatButtonStyle)
+		.ToolTipText(LOCTEXT("DeleteAllEntitiesToolTip", "Delete all entities."))
+		.OnClicked(this, &SRCPanelExposedEntitiesList::RequestDeleteAllEntities)
+		[
+			SNew(SBox)
+			.WidthOverride(RCPanelStyle->IconSize.X)
+			.HeightOverride(RCPanelStyle->IconSize.Y)
+			[
+				SNew(SImage)
+				.ColorAndOpacity(FSlateColor::UseForeground())
+				.Image(FAppStyle::GetBrush("Icons.Delete"))
 			]
+		];
 
-			+SHorizontalBox::Slot()
-			.Padding(2.f, 4.f)
-			.FillWidth(0.75f)
+	// Expose Button
+	ExposeDockPanel->AddHeaderToolbarItem(EToolbar::Left, PlaceholderBox.ToSharedRef());
+	ExposeDockPanel->AddFooterToolbarItem(EToolbar::Right, InArgs._ExposeComboButton.Get().ToSharedRef());
+	ExposeDockPanel->AddFooterToolbarItem(EToolbar::Right, DeleteAllEntitiesButton.ToSharedRef());
+
+	ExposePanel->AddPanel(ExposeDockPanel.ToSharedRef(), 0.75f);
+
+	ChildSlot
+		.Padding(RCPanelStyle->PanelPadding)
+		[
+			SNew(SSplitter)
+			.Orientation(Orient_Horizontal)
+
+			+ SSplitter::Slot()
+			.Value(1.f)
 			[
-				SNew(SBorder)
-				.BorderImage(FAppStyle::Get().GetBrush("ToolPanel.DarkGroupBorder"))
-				[
-					SAssignNew(FieldsListView, SListView<TSharedPtr<SRCPanelTreeNode>>)
-					.ItemHeight(24.f)
-					.OnGenerateRow(this, &SRCPanelExposedEntitiesList::OnGenerateRow)
-					.OnSelectionChanged(this, &SRCPanelExposedEntitiesList::OnSelectionChanged)
-					.SelectionMode(ESelectionMode::Single)
-					.ListItemsSource(&FieldEntities)
-					.OnContextMenuOpening(this, &SRCPanelExposedEntitiesList::OnContextMenuOpening, SRCPanelTreeNode::Field)
-					.ClearSelectionOnClick(true)
-					.HeaderRow(
-						SNew(SHeaderRow)
-
-						+ SHeaderRow::Column(FRemoteControlPresetColumns::DragDropHandle)
-						.DefaultLabel(LOCTEXT("RCPresetDragDropHandleColumnHeader", ""))
-						.FixedWidth(36.0f)
-
-						+ SHeaderRow::Column(FRemoteControlPresetColumns::Description)
-						.DefaultLabel(LOCTEXT("RCPresetDescColumnHeader", "Description"))
-						.FillWidth(0.7f)
-
-						+ SHeaderRow::Column(FRemoteControlPresetColumns::Value)
-						.DefaultLabel(LOCTEXT("RCPresetValueColumnHeader", "Value"))
-						.FillWidth(0.7f)
-					)
-				]
+				ExposePanel.ToSharedRef()
 			]
-		]
-	];
+		];
 
 	RegisterEvents();
 	RegisterPresetDelegates();
 	Refresh();
 }
+
 SRCPanelExposedEntitiesList::~SRCPanelExposedEntitiesList()
 {
 	UnregisterPresetDelegates();
@@ -197,8 +338,6 @@ void SRCPanelExposedEntitiesList::Tick(const FGeometry& AllottedGeometry, const
 {
 	if (bSearchRequested)
 	{
-		UE_LOG(LogTemp, Warning, TEXT("SRCPanelExposedEntitiesList::Tick -> Search Requested (Total Requests : %d)."), ++SearchCount);
-
 		TryRefreshingSearch(*SearchedText);
 
 		bSearchRequested = false;
@@ -206,8 +345,6 @@ void SRCPanelExposedEntitiesList::Tick(const FGeometry& AllottedGeometry, const
 	
 	if (bFilterApplicationRequested)
 	{
-		UE_LOG(LogTemp, Warning, TEXT("SRCPanelExposedEntitiesList::Tick -> Filter Requested (Total Requests : %d)."), ++FilterCount);
-
 		ApplyFilters();
 
 		bFilterApplicationRequested = false;
@@ -437,9 +574,9 @@ void SRCPanelExposedEntitiesList::RefreshGroups()
 		TSharedRef<SRCPanelGroup> FieldGroup = SNew(SRCPanelGroup, Preset.Get(), ColumnSizeData)
 			.Id(RCGroup.Id)
 			.Name(RCGroup.Name)
-			.OnFieldDropEvent_Raw(this, &SRCPanelExposedEntitiesList::OnDropOnGroup)
-			.OnGetGroupId_Raw(this, &SRCPanelExposedEntitiesList::GetGroupId)
-			.OnDeleteGroup_Raw(this, &SRCPanelExposedEntitiesList::OnDeleteGroup)
+			.OnFieldDropEvent(this, &SRCPanelExposedEntitiesList::OnDropOnGroup)
+			.OnGetGroupId(this, &SRCPanelExposedEntitiesList::GetGroupId)
+			.OnDeleteGroup(this, &SRCPanelExposedEntitiesList::OnDeleteGroup)
 			.EditMode(bIsInEditMode);
 		
 		FieldGroups.Add(FieldGroup);
@@ -486,7 +623,7 @@ TSharedRef<ITableRow> SRCPanelExposedEntitiesList::OnGenerateRow(TSharedPtr<SRCP
 	{
 		return SNew(STableRow<TSharedPtr<SWidget>>, OwnerTable)
 			.Padding(1.f)
-			.Style(&FRemoteControlPanelStyle::Get()->GetWidgetStyle<FTableRowStyle>("RemoteControlPanel.GroupRow"))
+			.Style(&RCPanelStyle->TableRowStyle)
 			[
 				NodeWidget
 			];
@@ -494,15 +631,14 @@ TSharedRef<ITableRow> SRCPanelExposedEntitiesList::OnGenerateRow(TSharedPtr<SRCP
 	else
 	{
 		constexpr float LeftPadding = 3.f;
-		const FMargin Margin = Node->GetRCType() == SRCPanelTreeNode::FieldChild ? FMargin(LeftPadding + 10.f, 0.f, 0.f, 0.f) : FMargin(LeftPadding, 0.f, 0.f, 0.f);
-		return SNew(STableRow<TSharedPtr<FGuid>>, OwnerTable)
+		const FMargin Margin = Node->GetRCType() == SRCPanelTreeNode::FieldChild ? FMargin(LeftPadding + 10.f, 1.f, 1.f, 1.f) : FMargin(LeftPadding, 1.f, 1.f, 1.f);
+		return SNew(SEntityRow, OwnerTable)
 			.OnDragEnter_Lambda([Node](const FDragDropEvent& Event) { if (Node && Node->GetRCType() == SRCPanelTreeNode::Field) StaticCastSharedPtr<SRCPanelExposedField>(Node)->SetIsHovered(true); })
 			.OnDragLeave_Lambda([Node](const FDragDropEvent& Event) { if (Node && Node->GetRCType() == SRCPanelTreeNode::Field) StaticCastSharedPtr<SRCPanelExposedField>(Node)->SetIsHovered(false); })
 			.OnDrop_Lambda(OnDropLambda)
 			.Padding(Margin)
-			[
-				NodeWidget
-			];
+			.Style(&RCPanelStyle->TableRowStyle)
+			.Entity(Node);
 	}
 }
 
@@ -518,6 +654,10 @@ void SRCPanelExposedEntitiesList::OnSelectionChanged(TSharedPtr<SRCPanelTreeNode
 {
 	if (!Node || SelectInfo != ESelectInfo::OnMouseClick)
 	{
+		FieldsListView->ClearSelection();
+
+		OnSelectionChangeDelegate.Broadcast(nullptr);
+
 		return;
 	}
 
@@ -532,6 +672,8 @@ void SRCPanelExposedEntitiesList::OnSelectionChanged(TSharedPtr<SRCPanelTreeNode
 				bFilterApplicationRequested = true;
 			}
 		}
+	
+		FieldsListView->ClearSelection();
 	}
 	else
 	{
@@ -718,7 +860,7 @@ void SRCPanelExposedEntitiesList::RegisterPresetDelegates()
 	Layout.OnGroupAdded().AddSP(this, &SRCPanelExposedEntitiesList::OnGroupAdded);
 	Layout.OnGroupDeleted().AddSP(this, &SRCPanelExposedEntitiesList::OnGroupDeleted);
 	Layout.OnGroupOrderChanged().AddSP(this, &SRCPanelExposedEntitiesList::OnGroupOrderChanged);
-	Layout.OnGroupRenamed().AddRaw(this, &SRCPanelExposedEntitiesList::OnGroupRenamed);
+	Layout.OnGroupRenamed().AddSP(this, &SRCPanelExposedEntitiesList::OnGroupRenamed);
 	Layout.OnFieldAdded().AddSP(this, &SRCPanelExposedEntitiesList::OnFieldAdded);
 	Layout.OnFieldDeleted().AddSP(this, &SRCPanelExposedEntitiesList::OnFieldDeleted);
 	Layout.OnFieldOrderChanged().AddSP(this, &SRCPanelExposedEntitiesList::OnFieldOrderChanged);
@@ -797,9 +939,9 @@ void SRCPanelExposedEntitiesList::OnGroupAdded(const FRemoteControlPresetGroup&
 	TSharedRef<SRCPanelGroup> FieldGroup = SNew(SRCPanelGroup, Preset.Get(), ColumnSizeData)
 		.Id(Group.Id)
 		.Name(Group.Name)
-		.OnFieldDropEvent_Raw(this, &SRCPanelExposedEntitiesList::OnDropOnGroup)
-		.OnGetGroupId_Raw(this, &SRCPanelExposedEntitiesList::GetGroupId)
-		.OnDeleteGroup_Raw(this, &SRCPanelExposedEntitiesList::OnDeleteGroup)
+		.OnFieldDropEvent(this, &SRCPanelExposedEntitiesList::OnDropOnGroup)
+		.OnGetGroupId(this, &SRCPanelExposedEntitiesList::GetGroupId)
+		.OnDeleteGroup(this, &SRCPanelExposedEntitiesList::OnDeleteGroup)
 		.EditMode(bIsInEditMode);
 	
 	FieldGroups.Add(FieldGroup);
@@ -926,8 +1068,8 @@ void SRCPanelExposedEntitiesList::OnEntitiesUpdated(URemoteControlPreset*, const
 	{
 		if (TSharedPtr<SRCPanelExposedEntitiesList> ListPtr = WeakListPtr.Pin())
 		{
-			ListPtr->GroupsListView->RequestListRefresh();
-			ListPtr->FieldsListView->RequestListRefresh();
+			ListPtr->GroupsListView->RebuildList();
+			ListPtr->FieldsListView->RebuildList();
 		}
 	}));
 
@@ -993,6 +1135,68 @@ void SRCPanelExposedEntitiesList::RequestSearchOrFilter()
 	bFilterApplicationRequested = BackendFilter.HasAnyActiveFilters();
 }
 
+FReply SRCPanelExposedEntitiesList::RequestDeleteAllEntities()
+{
+	if (!FieldsListView.IsValid() || !Preset.IsValid())
+	{
+		return FReply::Unhandled();
+	}
+
+	const FString WarningMessage = FString::Format(TEXT("You are about to delete '{0}' entities. This action might not be undone.\nAre you sure you want to proceed?"), { FString::FromInt(FieldWidgetMap.Num()) });
+
+	EAppReturnType::Type UserResponse = FMessageDialog::Open(EAppMsgType::YesNo, FText::FromString(WarningMessage));
+
+	if (UserResponse == EAppReturnType::Yes)
+	{
+		FScopedTransaction Transaction(LOCTEXT("UnexposeAll", "Unexpose all remote control entities"));
+		Preset->Modify();
+
+		for (TMap<FGuid, TSharedPtr<SRCPanelTreeNode>>::TConstIterator EntityItr = FieldWidgetMap.CreateConstIterator(); EntityItr; ++EntityItr)
+		{
+			Preset->Unexpose(EntityItr->Key);
+		}
+	}
+
+	return FReply::Handled();
+}
+
+FReply SRCPanelExposedEntitiesList::RequestDeleteAllGroups()
+{
+	if (!GroupsListView.IsValid() || !Preset.IsValid())
+	{
+		return FReply::Unhandled();
+	}
+
+	const FString WarningMessage = FString::Format(TEXT("You are about to delete '{0}' groups. This action might not be undone.\nAre you sure you want to proceed?"), { FString::FromInt(FieldGroups.Num() - 1) });
+
+	EAppReturnType::Type UserResponse = FMessageDialog::Open(EAppMsgType::YesNo, FText::FromString(WarningMessage));
+
+	if (UserResponse == EAppReturnType::Yes)
+	{
+		FScopedTransaction Transaction(LOCTEXT("DeleteAllGroups", "Delete all remote control groups"));
+		Preset->Modify();
+
+		const TArray<FRemoteControlPresetGroup>& RCGroups = Preset->Layout.GetGroups();
+
+		// Perform deletion in reverse to avoid "Ensure condition failed: Lhs.CurrentNum == Lhs.InitialNum"
+		for (int32 GroupIndex = RCGroups.Num() - 1; GroupIndex >= 0; GroupIndex--)
+		{
+			const FRemoteControlPresetGroup& RCGroup = RCGroups[GroupIndex];
+
+			if (Preset->Layout.IsDefaultGroup(RCGroup.Id))
+			{
+				continue;
+			}
+		
+			Preset->Layout.DeleteGroup(RCGroup.Id);
+		}
+
+		Refresh();
+	}
+
+	return FReply::Handled();
+}
+
 bool FGroupDragEvent::IsDraggedFromSameGroup() const
 {
 	return DragOriginGroup->GetId() == DragTargetGroup->GetId();
diff --git a/Engine/Plugins/VirtualProduction/RemoteControl/Source/RemoteControlUI/Private/UI/SRCPanelExposedEntitiesList.h b/Engine/Plugins/VirtualProduction/RemoteControl/Source/RemoteControlUI/Private/UI/SRCPanelExposedEntitiesList.h
--- a/Engine/Plugins/VirtualProduction/RemoteControl/Source/RemoteControlUI/Private/UI/SRCPanelExposedEntitiesList.h
+++ b/Engine/Plugins/VirtualProduction/RemoteControl/Source/RemoteControlUI/Private/UI/SRCPanelExposedEntitiesList.h
@@ -26,6 +26,7 @@ class SRemoteControlTarget;
 struct SRCPanelExposedField;
 class STableViewBase;
 class URemoteControlPreset;
+struct FRCPanelStyle;
 
 /** Holds information about a group drag and drop event  */
 struct FGroupDragEvent
@@ -53,6 +54,7 @@ public:
 		: _EditMode(true)
 	{}
 		SLATE_ATTRIBUTE(bool, EditMode)
+		SLATE_ATTRIBUTE(TSharedPtr<SWidget>, ExposeComboButton)
 		SLATE_EVENT(FSimpleDelegate, OnEntityListUpdated)
 	SLATE_END_ARGS()
 
@@ -151,11 +153,16 @@ private:
 	/** Notifies us that we need to do a search or filter in the next tick. */
 	void RequestSearchOrFilter();
 
+	/** Delete all operation for entities list. */
+	FReply RequestDeleteAllEntities();
+	/** Delete all operation for groups list. */
+	FReply RequestDeleteAllGroups();
+
 private:
 	/** Holds the Groups list view. */
 	TSharedPtr<SListView<TSharedPtr<SRCPanelTreeNode>>> GroupsListView;
 	/** Holds the Fields list view. */
-	TSharedPtr<SListView<TSharedPtr<SRCPanelTreeNode>>> FieldsListView;
+	TSharedPtr<STreeView<TSharedPtr<SRCPanelTreeNode>>> FieldsListView;
 	/** Holds all the field groups. */
 	TArray<TSharedPtr<SRCPanelGroup>> FieldGroups;
 	/** Holds all the field entities. */
@@ -173,7 +180,7 @@ private:
 	/** The column data shared between all tree nodes in order to share a splitter amongst all rows. */
 	FRCColumnSizeData ColumnSizeData;
 	/** The actual width of the right column.  The left column is 1-ColumnWidth */
-	float ColumnWidth = 0.65f;
+	float ColumnWidth = 0.49f;
 	/** Event triggered when the entity list is updated. */
 	FSimpleDelegate OnEntityListUpdatedDelegate;
 	/** Holds the cache of widgets to be used by this list's entities. */
@@ -186,6 +193,6 @@ private:
 	bool bSearchRequested;
 	/** Holds the text that is being searched actively. */
 	TSharedPtr<FText> SearchedText;
-	int32 SearchCount;
-	int32 FilterCount;
+	/** Panel Style reference. */
+	const FRCPanelStyle* RCPanelStyle;
 };
\ No newline at end of file
diff --git a/Engine/Plugins/VirtualProduction/RemoteControl/Source/RemoteControlUI/Private/UI/SRCPanelExposedEntity.cpp b/Engine/Plugins/VirtualProduction/RemoteControl/Source/RemoteControlUI/Private/UI/SRCPanelExposedEntity.cpp
--- a/Engine/Plugins/VirtualProduction/RemoteControl/Source/RemoteControlUI/Private/UI/SRCPanelExposedEntity.cpp
+++ b/Engine/Plugins/VirtualProduction/RemoteControl/Source/RemoteControlUI/Private/UI/SRCPanelExposedEntity.cpp
@@ -22,6 +22,7 @@
 #include "SceneOutlinerFilters.h"
 #include "SceneOutlinerModule.h"
 #include "ScopedTransaction.h"
+#include "Styling/RemoteControlStyles.h"
 #include "Styling/SlateIconFinder.h"
 #include "Modules/ModuleManager.h"
 #include "Widgets/Input/SButton.h"
@@ -34,18 +35,6 @@
 
 #define LOCTEXT_NAMESPACE "RemoteControlPanel"
 
-void SRCPanelExposedEntity::Tick(const FGeometry&, const double, const float)
-{
-	if (bNeedsRename)
-	{
-		if (NameTextBox)
-		{
-			NameTextBox->EnterEditingMode();
-		}
-		bNeedsRename = false;
-	}
-}
-
 TSharedPtr<FRemoteControlEntity> SRCPanelExposedEntity::GetEntity() const
 {
 	if (Preset.IsValid())
@@ -104,7 +93,10 @@ TSharedPtr<SWidget> SRCPanelExposedEntity::GetContextMenu()
 
 void SRCPanelExposedEntity::EnterRenameMode()
 {
-	bNeedsRename = true;
+	if (NameTextBox.IsValid())
+	{
+		NameTextBox->EnterEditingMode();
+	}
 }
 
 void SRCPanelExposedEntity::Initialize(const FGuid& InEntityId, URemoteControlPreset* InPreset, const TAttribute<bool>& InbEditMode)
@@ -113,6 +105,8 @@ void SRCPanelExposedEntity::Initialize(const FGuid& InEntityId, URemoteControlPr
 	Preset = InPreset;
 	bEditMode = InbEditMode;
 
+	RCPanelStyle = &FRemoteControlPanelStyle::Get()->GetWidgetStyle<FRCPanelStyle>("RemoteControlPanel.MinorPanel");
+
 	if (ensure(InPreset))
 	{
 		if (const TSharedPtr<FRemoteControlEntity> RCEntity = InPreset->GetExposedEntity(InEntityId).Pin())
@@ -169,7 +163,7 @@ TSharedRef<SWidget> SRCPanelExposedEntity::CreateRebindAllPropertiesForActorMenu
 
 	if (TSharedPtr<FRemoteControlEntity> Entity = GetEntity())
 	{
-		Options.Filters->AddFilterPredicate<FActorTreeItem>(FActorTreeItem::FFilterPredicate::CreateRaw(this, &SRCPanelExposedEntity::IsActorSelectable));
+		Options.Filters->AddFilterPredicate<FActorTreeItem>(FActorTreeItem::FFilterPredicate::CreateSP(this, &SRCPanelExposedEntity::IsActorSelectable));
 	}
 
 	return SNew(SBox)
@@ -214,7 +208,7 @@ TSharedRef<SWidget> SRCPanelExposedEntity::CreateRebindMenuContent()
 	FSceneOutlinerModule& SceneOutlinerModule = FModuleManager::Get().LoadModuleChecked<FSceneOutlinerModule>("SceneOutliner");
 	FSceneOutlinerInitializationOptions Options;
 	Options.Filters = MakeShared<FSceneOutlinerFilters>();
-	Options.Filters->AddFilterPredicate<FActorTreeItem>(FActorTreeItem::FFilterPredicate::CreateRaw(this, &SRCPanelExposedEntity::IsActorSelectable));
+	Options.Filters->AddFilterPredicate<FActorTreeItem>(FActorTreeItem::FFilterPredicate::CreateSP(this, &SRCPanelExposedEntity::IsActorSelectable));
 	UWorld* PresetWorld = URemoteControlPreset::GetPresetWorld(Preset.Get());
 
 	return SNew(SBox)
@@ -297,16 +291,16 @@ bool SRCPanelExposedEntity::IsActorSelectable(const AActor* Actor) const
 	return false;
 }
 
-TSharedRef<SWidget> SRCPanelExposedEntity::CreateEntityWidget(TSharedPtr<SWidget> ValueWidget, const FText& OptionalWarningMessage)
+TSharedRef<SWidget> SRCPanelExposedEntity::CreateEntityWidget(TSharedPtr<SWidget> ValueWidget, TSharedPtr<SWidget> ResetWidget, const FText& OptionalWarningMessage)
 {
 	FMakeNodeWidgetArgs Args;
 
 	TSharedRef<SBorder> Widget = SNew(SBorder)
 		.Padding(0.0f)
-		.BorderImage_Raw(this, &SRCPanelExposedEntity::GetBorderImage);
+		.BorderImage(this, &SRCPanelExposedEntity::GetBorderImage);
 	
 	Args.DragHandle = SNew(SBox)
-		.Visibility_Raw(this, &SRCPanelExposedEntity::GetVisibilityAccordingToEditMode, EVisibility::Collapsed)
+		.Visibility(this, &SRCPanelExposedEntity::GetVisibilityAccordingToEditMode, EVisibility::Collapsed)
 		[
 			SNew(SRCPanelDragHandle<FExposedEntityDragDrop>, GetRCId())
 			.Widget(Widget)
@@ -332,17 +326,15 @@ TSharedRef<SWidget> SRCPanelExposedEntity::CreateEntityWidget(TSharedPtr<SWidget
 		[
 			SAssignNew(NameTextBox, SInlineEditableTextBlock)
 			.Text(FText::FromName(CachedLabel))
-			.OnTextCommitted_Raw(this, &SRCPanelExposedEntity::OnLabelCommitted)
-			.OnVerifyTextChanged_Raw(this, &SRCPanelExposedEntity::OnVerifyItemLabelChanged)
+			.OnTextCommitted(this, &SRCPanelExposedEntity::OnLabelCommitted)
+			.OnVerifyTextChanged(this, &SRCPanelExposedEntity::OnVerifyItemLabelChanged)
 			.IsReadOnly_Lambda([this]() { return !bEditMode.Get(); })
 			.HighlightText_Lambda([this]() { return HighlightText.Get().ToString().Len() > 3 ? HighlightText.Get() : FText::GetEmpty(); })
 		];
 
-	Args.RenameButton = SNullWidget::NullWidget;
-
 	Args.ValueWidget = ValueWidget;
 
-	Args.UnexposeButton = SNullWidget::NullWidget;
+	Args.ResetButton = ResetWidget;
 
 	Widget->SetContent(MakeNodeWidget(Args));
 	return Widget;
@@ -403,10 +395,10 @@ TSharedRef<SWidget> SRCPanelExposedEntity::CreateUseContextCheckbox()
 		.ToolTipText(LOCTEXT("UseRebindingContextTooltip", "Unchecking this will allow you to rebind this property to any object regardless of the underlying supported class."))
 
 		// Bind the button's "on checked" event to our object's method for this
-		.OnCheckStateChanged_Raw(this, &SRCPanelExposedEntity::OnUseContextChanged)
+		.OnCheckStateChanged(this, &SRCPanelExposedEntity::OnUseContextChanged)
 
 		// Bind the check box's "checked" state to our user interface action
-		.IsChecked_Raw(this, &SRCPanelExposedEntity::IsUseContextEnabled);
+		.IsChecked(this, &SRCPanelExposedEntity::IsUseContextEnabled);
 }
 
 void SRCPanelExposedEntity::OnUseContextChanged(ECheckBoxState State)
diff --git a/Engine/Plugins/VirtualProduction/RemoteControl/Source/RemoteControlUI/Private/UI/SRCPanelExposedEntity.h b/Engine/Plugins/VirtualProduction/RemoteControl/Source/RemoteControlUI/Private/UI/SRCPanelExposedEntity.h
--- a/Engine/Plugins/VirtualProduction/RemoteControl/Source/RemoteControlUI/Private/UI/SRCPanelExposedEntity.h
+++ b/Engine/Plugins/VirtualProduction/RemoteControl/Source/RemoteControlUI/Private/UI/SRCPanelExposedEntity.h
@@ -6,6 +6,7 @@
 #include "DragAndDrop/DecoratedDragDropOp.h"
 #include "RemoteControlPanelStyle.h"
 #include "Widgets/Layout/SBorder.h"
+#include "Widgets/SNullWidget.h"
 
 class AActor;
 class FMenuBuilder;
@@ -13,15 +14,13 @@ struct FRemoteControlEntity;
 class SInlineEditableTextBlock;
 class SWidget;
 class URemoteControlPreset;
+struct FRCPanelStyle;
 
 /**
  * @Note if you inherit from this struct, you must call SRCPanelExposedEntity::Initialize.
  */
 struct SRCPanelExposedEntity : public SRCPanelTreeNode
 {
-	//~ SWidget interface
-	virtual void Tick(const FGeometry&, const double, const float) override;
-
 	//~ SRCPanelTreeNode interface
 	TSharedPtr<FRemoteControlEntity> GetEntity() const;
 	virtual TSharedPtr<SWidget> GetContextMenu() override;
@@ -41,7 +40,7 @@ protected:
 	EVisibility GetVisibilityAccordingToEditMode(EVisibility NonEditModeVisibility) const;
 
 	/** Create an exposed entity widget with a drag handle and unexpose button. */
-	TSharedRef<SWidget> CreateEntityWidget(TSharedPtr<SWidget> ValueWidget, const FText& OptionalWarningMessage = FText::GetEmpty());
+	TSharedRef<SWidget> CreateEntityWidget(TSharedPtr<SWidget> ValueWidget, TSharedPtr<SWidget> ResetWidget = SNullWidget::NullWidget, const FText& OptionalWarningMessage = FText::GetEmpty());
 
 protected:
 	/** Id of the entity. */
@@ -54,6 +53,8 @@ protected:
 	FName CachedLabel;
 	/** Text to be highlighted while searching. */
 	TAttribute<FText> HighlightText;
+	/** Panel Style reference. */
+	const FRCPanelStyle* RCPanelStyle;
 
 private:
 	/** Handles changing the object this entity is bound to upon selecting an actor in the rebinding dropdown. */
@@ -79,8 +80,6 @@ private:
 	void OnLabelCommitted(const FText& InLabel, ETextCommit::Type InCommitInfo);
 	/** Returns whether or not the actor is selectable for a binding replacement. */
 	bool IsActorSelectable(const AActor* Parent) const;
-	/** Whether the editable text box for the label needs to enter edit mode. */
-	bool bNeedsRename = false;
 	/** Handle clicking on the unexpose button. */
 	void HandleUnexposeEntity();
 
diff --git a/Engine/Plugins/VirtualProduction/RemoteControl/Source/RemoteControlUI/Private/UI/SRCPanelExposedField.cpp b/Engine/Plugins/VirtualProduction/RemoteControl/Source/RemoteControlUI/Private/UI/SRCPanelExposedField.cpp
--- a/Engine/Plugins/VirtualProduction/RemoteControl/Source/RemoteControlUI/Private/UI/SRCPanelExposedField.cpp
+++ b/Engine/Plugins/VirtualProduction/RemoteControl/Source/RemoteControlUI/Private/UI/SRCPanelExposedField.cpp
@@ -18,12 +18,13 @@
 #include "ScopedTransaction.h"
 #include "Styling/AppStyle.h"
 #include "Styling/SlateBrush.h"
+#include "Styling/RemoteControlStyles.h"
 #include "UObject/Object.h"
 #include "Widgets/Input/SButton.h"
 #include "Widgets/Layout/SBorder.h"
 #include "Widgets/Text/SInlineEditableTextBlock.h"
 #include "Widgets/Text/STextBlock.h"
-
+#include "Widgets/SNullWidget.h"
 
 #define LOCTEXT_NAMESPACE "RemoteControlPanel"
 
@@ -73,6 +74,8 @@ void SRCPanelExposedField::Construct(const FArguments& InArgs, TWeakPtr<FRemoteC
 	
 	HighlightText = InArgs._HighlightText;
 
+	ResetButtonWidget = SNullWidget::NullWidget;
+
 	if (TSharedPtr<FRemoteControlField> FieldPtr = WeakField.Pin())
 	{
 		Initialize(FieldPtr->GetId(), InArgs._Preset.Get(), InArgs._EditMode);
@@ -101,6 +104,11 @@ SRCPanelTreeNode::ENodeType SRCPanelExposedField::GetRCType() const
 	return SRCPanelTreeNode::Field;
 }
 
+bool SRCPanelExposedField::HasChildren() const
+{
+	return ChildWidgets.Num() > 0;
+}
+
 FName SRCPanelExposedField::GetFieldLabel() const
 {
 	return CachedLabel;
@@ -173,7 +181,6 @@ TSharedRef<SWidget> SRCPanelExposedField::ConstructWidget()
 						ExposedFieldUtils::CreateNodeValueWidget(MoveTemp(Node))
 					];
 
-					
 					TSharedPtr<IPropertyHandle> PropertyHandle = Node->CreatePropertyHandle();
 					if (PropertyHandle->GetParentHandle()->GetPropertyDisplayName().ToString() == FString("Transform")
 						|| PropertyHandle->GetOuterBaseClass() == UMaterialInstanceDynamic::StaticClass())
@@ -196,39 +203,40 @@ TSharedRef<SWidget> SRCPanelExposedField::ConstructWidget()
 							return bVisible ? EVisibility::Visible : EVisibility::Hidden;
 						};
 						
-						FieldWidget->AddSlot()
-						.AutoWidth()
-						.VAlign(VAlign_Center)
-						.Padding(FMargin(1.f,2.f))
-						[
-							SNew(SButton)
-							.IsFocusable(false)
-							.ButtonStyle(FAppStyle::Get(), "SimpleButton")
-							.ContentPadding(0) 
-							.Visibility_Lambda(IsVisible)
-							.OnClicked_Lambda([Node]()
-							{
-								Node->CreatePropertyHandle()->ResetToDefault();
-								return FReply::Handled();
-							})
-							.Content()
+						ResetButtonWidget = SNew(SBox)
+							.HAlign(HAlign_Center)
+							.VAlign(VAlign_Center)
+							.Padding(2.f, 4.f)
 							[
-								SNew(SImage)
-								.Image(FAppStyle::GetBrush("PropertyWindow.DiffersFromDefault"))
-								.ColorAndOpacity(FSlateColor::UseForeground())
-							]
-						];
+								SNew(SButton)
+								.IsFocusable(false)
+								.ButtonStyle(&RCPanelStyle->FlatButtonStyle)
+								.ContentPadding(RCPanelStyle->PanelPadding)
+								.Visibility_Lambda(IsVisible)
+								.OnClicked_Lambda([Node]()
+								{
+									Node->CreatePropertyHandle()->ResetToDefault();
+									return FReply::Handled();
+								})
+								.Content()
+								[
+									SNew(SImage)
+									.Image(FAppStyle::GetBrush("PropertyWindow.DiffersFromDefault"))
+									.ColorAndOpacity(FSlateColor::UseForeground())
+								]
+							];
 					}
 					else
 					{
-						FieldWidget->AddSlot()
-						.AutoWidth()
-						.VAlign(VAlign_Center)
-						.Padding(FMargin(1.f,2.f))
-						[
-							SNew(SResetToDefaultPropertyEditor, PropertyHandle)
-						];
+						ResetButtonWidget = SNew(SBox)
+							.HAlign(HAlign_Center)
+							.VAlign(VAlign_Center)
+							.Padding(2.f, 4.f)
+							[
+								SNew(SResetToDefaultPropertyEditor, PropertyHandle)
+							];
 					}
+
 					return MakeFieldWidget(FieldWidget.ToSharedRef());
 				}
 			}
@@ -282,7 +290,7 @@ TSharedRef<SWidget> SRCPanelExposedField::MakeFieldWidget(const TSharedRef<SWidg
 		WarningMessage = Builder.ToText();
 	}
 	
-	return CreateEntityWidget(InWidget, WarningMessage);
+	return CreateEntityWidget(InWidget, ResetButtonWidget.ToSharedRef(), WarningMessage);
 }
 
 void SRCPanelExposedField::ConstructPropertyWidget()
@@ -349,7 +357,7 @@ TSharedRef<SWidget> SRCPanelExposedField::ConstructCallFunctionButton(bool bIsEn
 			SNew(SButton)
 			.IsEnabled(bIsEnabled)
 			.VAlign(VAlign_Center)
-			.OnClicked_Raw(this, &SRCPanelExposedField::OnClickFunctionButton)
+			.OnClicked(this, &SRCPanelExposedField::OnClickFunctionButton)
 			[
 				SNew(STextBlock)
 				.Text(LOCTEXT("CallFunctionLabel", "Call Function"))
diff --git a/Engine/Plugins/VirtualProduction/RemoteControl/Source/RemoteControlUI/Private/UI/SRCPanelExposedField.h b/Engine/Plugins/VirtualProduction/RemoteControl/Source/RemoteControlUI/Private/UI/SRCPanelExposedField.h
--- a/Engine/Plugins/VirtualProduction/RemoteControl/Source/RemoteControlUI/Private/UI/SRCPanelExposedField.h
+++ b/Engine/Plugins/VirtualProduction/RemoteControl/Source/RemoteControlUI/Private/UI/SRCPanelExposedField.h
@@ -21,6 +21,7 @@ class IDetailTreeNode;
 class SInlineEditableTextBlock;
 struct SRCPanelFieldChildNode;
 class URemoteControlPreset;
+class SWidget;
 
 /**
  * Widget that displays an exposed field.
@@ -44,6 +45,7 @@ struct SRCPanelExposedField : public SRCPanelExposedEntity
 	//~ SRCPanelTreeNode Interface 
 	virtual void GetNodeChildren(TArray<TSharedPtr<SRCPanelTreeNode>>& OutChildren) const override;
 	virtual ENodeType GetRCType() const override;
+	virtual bool HasChildren() const override;
 	virtual void Refresh() override;
 	//~ End SRCPanelTreeNode Interface
 
@@ -60,7 +62,6 @@ struct SRCPanelExposedField : public SRCPanelExposedEntity
 	/** Get this field's type. */
 	EExposedFieldType GetFieldType() const;
 
-
 	/** Returns this widget's underlying objects. */
 	void GetBoundObjects(TSet<UObject*>& OutBoundObjects) const;
 
@@ -92,6 +93,8 @@ private:
 	TWeakPtr<FRCPanelWidgetRegistry> WidgetRegistry;
 	/** Holds the zeroed Default Value of the ExposedField */
 	TSharedPtr<uint8> DefaultValue;
+	/** Holds the shared reference of reset button for this field. */
+	TSharedPtr<SWidget> ResetButtonWidget;
 };
 
 
diff --git a/Engine/Plugins/VirtualProduction/RemoteControl/Source/RemoteControlUI/Private/UI/SRCPanelTreeNode.cpp b/Engine/Plugins/VirtualProduction/RemoteControl/Source/RemoteControlUI/Private/UI/SRCPanelTreeNode.cpp
--- a/Engine/Plugins/VirtualProduction/RemoteControl/Source/RemoteControlUI/Private/UI/SRCPanelTreeNode.cpp
+++ b/Engine/Plugins/VirtualProduction/RemoteControl/Source/RemoteControlUI/Private/UI/SRCPanelTreeNode.cpp
@@ -10,12 +10,34 @@
 
 #define LOCTEXT_NAMESPACE "RemoteControlPanelNode"
 
+TSharedRef<SWidget> SRCPanelTreeNode::GetWidget(const FName& ForColumnName)
+{
+	if (ForColumnName == FRemoteControlPresetColumns::DragDropHandle)
+	{
+		return DragHandleWidget.ToSharedRef();
+	}
+	else if (ForColumnName == FRemoteControlPresetColumns::Description)
+	{
+		return NodeNameWidget.ToSharedRef();
+	}
+	else if (ForColumnName == FRemoteControlPresetColumns::Reset)
+	{
+		return ResetValueWidget.ToSharedRef();
+	}
+	else if (ForColumnName == FRemoteControlPresetColumns::Value)
+	{
+		return NodeValueWidget.ToSharedRef();
+	}
+
+	return SNullWidget::NullWidget;
+}
+
 TSharedRef<SWidget> SRCPanelTreeNode::MakeSplitRow(TSharedRef<SWidget> LeftColumn, TSharedRef<SWidget> RightColumn)
 {
 	TAttribute<float> LeftColumnAttribute;
-	LeftColumnAttribute.BindRaw(this, &SRCPanelTreeNode::GetLeftColumnWidth);
+	LeftColumnAttribute.Bind(this, &SRCPanelTreeNode::GetLeftColumnWidth);
 	TAttribute<float> RightColumnAttribute;
-	RightColumnAttribute.BindRaw(this, &SRCPanelTreeNode::GetRightColumnWidth);
+	RightColumnAttribute.Bind(this, &SRCPanelTreeNode::GetRightColumnWidth);
 	
 	return SNew(SSplitter)
 		.Style(FAppStyle::Get(), "DetailsView.Splitter")
@@ -23,13 +45,13 @@ TSharedRef<SWidget> SRCPanelTreeNode::MakeSplitRow(TSharedRef<SWidget> LeftColum
 		.HitDetectionSplitterHandleSize(5.0f)
 		+ SSplitter::Slot()
 		.Value(MoveTemp(LeftColumnAttribute))
-		.OnSlotResized(SSplitter::FOnSlotResized::CreateRaw(this, &SRCPanelTreeNode::OnLeftColumnResized))
+		.OnSlotResized(SSplitter::FOnSlotResized::CreateSP(this, &SRCPanelTreeNode::OnLeftColumnResized))
 		[
 			LeftColumn
 		]
 		+ SSplitter::Slot()
 		.Value(MoveTemp(RightColumnAttribute))
-		.OnSlotResized(SSplitter::FOnSlotResized::CreateRaw(this, &SRCPanelTreeNode::SetColumnWidth))
+		.OnSlotResized(SSplitter::FOnSlotResized::CreateSP(this, &SRCPanelTreeNode::SetColumnWidth))
 		[
 			RightColumn
 		];
@@ -37,10 +59,9 @@ TSharedRef<SWidget> SRCPanelTreeNode::MakeSplitRow(TSharedRef<SWidget> LeftColum
 
 TSharedRef<SWidget> SRCPanelTreeNode::MakeNodeWidget(const FMakeNodeWidgetArgs& Args)
 {
-	auto WidgetOrNull = [](const TSharedPtr<SWidget>& Widget) {return Widget ? Widget.ToSharedRef() : SNullWidget::NullWidget; };
+	MakeNodeWidgets(Args);
 
-	TSharedRef<SWidget> LeftColumn = 
-		SNew(SHorizontalBox)
+	TSharedRef<SWidget> LeftColumn = SNew(SHorizontalBox)
 		.Clipping(EWidgetClipping::OnDemand)
 		// Drag and drop handle
 		+ SHorizontalBox::Slot()
@@ -48,52 +69,53 @@ TSharedRef<SWidget> SRCPanelTreeNode::MakeNodeWidget(const FMakeNodeWidgetArgs&
 		.VAlign(VAlign_Center)
 		.AutoWidth()
 		[
-			WidgetOrNull(Args.DragHandle)
+			DragHandleWidget.ToSharedRef()
 		]
 		// Field name
 		+ SHorizontalBox::Slot()
 		.VAlign(VAlign_Center)
 		.AutoWidth()
 		[
-			WidgetOrNull(Args.NameWidget)
-		]
-		// Rename button
-		+ SHorizontalBox::Slot()
-		.VAlign(VAlign_Center)
-		.AutoWidth()
-		[
-			WidgetOrNull(Args.RenameButton)
+			NodeNameWidget.ToSharedRef()
 		];
 
-	TSharedRef<SWidget> RightColumn =
-		SNew(SOverlay)
-		+ SOverlay::Slot()
-
+	TSharedRef<SWidget> RightColumn = SNew(SHorizontalBox)
+		.Clipping(EWidgetClipping::OnDemand)
+		// Node Value
+		+ SHorizontalBox::Slot()
 		.HAlign(HAlign_Fill)
 		[
 			SNew(SBox)
 			.HAlign(HAlign_Left)
 			.VAlign(VAlign_Center)
 			[
-				WidgetOrNull(Args.ValueWidget)			
+				NodeValueWidget.ToSharedRef()
 			]
 		]
-
-		+ SOverlay::Slot()
-		.HAlign(HAlign_Right)
+		// Reset button
+		+ SHorizontalBox::Slot()
 		.VAlign(VAlign_Center)
+		.AutoWidth()
 		[
-			SNew(SBorder)
-			.BorderImage(FAppStyle::Get().GetBrush("ToolPanel.GroupBorder"))
-			.Padding(0.f)
-			[
-				WidgetOrNull(Args.UnexposeButton)
-			]
+			ResetValueWidget.ToSharedRef()
 		];
-		
+
 	return MakeSplitRow(LeftColumn, RightColumn);
 }
 
+void SRCPanelTreeNode::MakeNodeWidgets(const FMakeNodeWidgetArgs& Args)
+{
+	auto WidgetOrNull = [](const TSharedPtr<SWidget>& Widget) {return Widget ? Widget.ToSharedRef() : SNullWidget::NullWidget; };
+
+	DragHandleWidget = WidgetOrNull(Args.DragHandle);
+
+	NodeNameWidget = WidgetOrNull(Args.NameWidget);
+
+	NodeValueWidget = WidgetOrNull(Args.ValueWidget);
+
+	ResetValueWidget = WidgetOrNull(Args.ResetButton);
+}
+
 void SRCPanelTreeNode::OnLeftColumnResized(float) const
 {
 	// This has to be bound or the splitter will take it upon itself to determine the size
diff --git a/Engine/Plugins/VirtualProduction/RemoteControl/Source/RemoteControlUI/Private/UI/SRCPanelTreeNode.h b/Engine/Plugins/VirtualProduction/RemoteControl/Source/RemoteControlUI/Private/UI/SRCPanelTreeNode.h
--- a/Engine/Plugins/VirtualProduction/RemoteControl/Source/RemoteControlUI/Private/UI/SRCPanelTreeNode.h
+++ b/Engine/Plugins/VirtualProduction/RemoteControl/Source/RemoteControlUI/Private/UI/SRCPanelTreeNode.h
@@ -14,6 +14,14 @@ struct SRCPanelExposedField;
 struct SRCPanelExposedActor;
 struct SRCPanelExposedMaterial;
 
+namespace FRemoteControlPresetColumns
+{
+	static FName DragDropHandle = TEXT("DragDropHandle");
+	static FName Description = TEXT("Description");
+	static FName Value = TEXT("Value");
+	static FName Reset = TEXT("Reset");
+}
+
 /** A node in the panel tree view. */
 struct SRCPanelTreeNode : public SCompoundWidget
 {
@@ -35,6 +43,8 @@ struct SRCPanelTreeNode : public SCompoundWidget
 	virtual FGuid GetRCId() const { return FGuid(); }
 	/** Get get this node's type. */
 	virtual ENodeType GetRCType() const { return ENodeType::Invalid; };
+	/** Returns true if this tree node has childen. */
+	virtual bool HasChildren() const { return false; }
 	/** Refresh the node. */
 	virtual void Refresh() {};
 	/** Get the context menu for this node. */
@@ -45,23 +55,24 @@ struct SRCPanelTreeNode : public SCompoundWidget
 	virtual void EnterRenameMode() {};
 	/** Updates the highlight text to active search term. */
 	virtual void SetHighlightText(const FText& InHightlightText = FText::GetEmpty()) {};
-
-	virtual FReply OnMouseButtonDoubleClick(const FGeometry& InMyGeometry, const FPointerEvent& InMouseEvent) override { if (InMouseEvent.GetEffectingButton() == EKeys::LeftMouseButton) { EnterRenameMode(); } return SCompoundWidget::OnMouseButtonDoubleClick(InMyGeometry, InMouseEvent); };
+	/** Retrieves the referenced widget corresponding to the given column name. */
+	virtual TSharedRef<SWidget> GetWidget(const FName& ForColumnName);
 
 protected:
 	struct FMakeNodeWidgetArgs
 	{
 		TSharedPtr<SWidget> DragHandle;
 		TSharedPtr<SWidget> NameWidget;
-		TSharedPtr<SWidget> RenameButton;
 		TSharedPtr<SWidget> ValueWidget;
-		TSharedPtr<SWidget> UnexposeButton;
+		TSharedPtr<SWidget> ResetButton;
 	};
 	
 	/** Create a widget that represents a row with a splitter. */
 	TSharedRef<SWidget> MakeSplitRow(TSharedRef<SWidget> LeftColumn, TSharedRef<SWidget> RightColumn);
 	/** Create a widget that represents a node in the panel tree hierarchy. */
 	TSharedRef<SWidget> MakeNodeWidget(const FMakeNodeWidgetArgs& Args);
+	/** Creates cached copies of underlying widgets. */
+	void MakeNodeWidgets(const FMakeNodeWidgetArgs& Args);
 
 private:
 	/** Stub handler for column resize callback to prevent the splitter from handling it internally.  */
@@ -76,6 +87,20 @@ protected:
 	/** Holds the row's columns' width. */
 	FRCColumnSizeData ColumnSizeData;
 
+private:
+
+	/** Cached widget of drag handle. */
+	TSharedPtr<SWidget> DragHandleWidget;
+
+	/** Cached widget of node name widget. */
+	TSharedPtr<SWidget> NodeNameWidget;
+
+	/** Cached widget of node value widget. */
+	TSharedPtr<SWidget> NodeValueWidget;
+	
+	/** Cached widget of reset arrow widget. */
+	TSharedPtr<SWidget> ResetValueWidget;
+
 private:
 	/** The splitter offset to align the group splitter with the other row's splitters. */
 	static constexpr float SplitterOffset = 0.008f;
diff --git a/Engine/Plugins/VirtualProduction/RemoteControl/Source/RemoteControlUI/Private/UI/SRemoteControlPanel.cpp b/Engine/Plugins/VirtualProduction/RemoteControl/Source/RemoteControlUI/Private/UI/SRemoteControlPanel.cpp
--- a/Engine/Plugins/VirtualProduction/RemoteControl/Source/RemoteControlUI/Private/UI/SRemoteControlPanel.cpp
+++ b/Engine/Plugins/VirtualProduction/RemoteControl/Source/RemoteControlUI/Private/UI/SRemoteControlPanel.cpp
@@ -4,7 +4,6 @@
 
 #include "Action/SRCActionPanel.h"
 #include "ActorEditorUtils.h"
-#include "AssetRegistry/AssetRegistryModule.h"
 #include "Behaviour/SRCBehaviourPanel.h"
 #include "ClassViewerFilter.h"
 #include "ClassViewerModule.h"
@@ -49,6 +48,7 @@
 #include "SRCPanelExposedActor.h"
 #include "SRCPanelExposedField.h"
 #include "SRCPanelTreeNode.h"
+#include "Styling/RemoteControlStyles.h"
 #include "Styling/ToolBarStyle.h"
 #include "Subsystems/Subsystem.h"
 #include "SWarningOrErrorBox.h"
@@ -60,6 +60,8 @@
 #include "UI/BaseLogicUI/RCLogicModeBase.h"
 #include "UI/BaseLogicUI/SRCLogicPanelListBase.h"
 #include "UI/Filters/SRCPanelFilter.h"
+#include "UI/Drawers/SRCPanelDrawer.h"
+#include "UI/Panels/SRCDockPanel.h"
 #include "Widgets/DeclarativeSyntaxSupport.h"
 #include "Widgets/Docking/SDockTab.h"
 #include "Widgets/Images/SLayeredImage.h"
@@ -79,6 +81,17 @@ const FName SRemoteControlPanel::DefaultRemoteControlPanelToolBarName("RemoteCon
 const FName SRemoteControlPanel::AuxiliaryRemoteControlPanelToolBarName("RemoteControlPanel.AuxiliaryToolBar");
 const float SRemoteControlPanel::MinimumPanelWidth = 640.f;
 
+TSharedRef<SBox> SRemoteControlPanel::NoneSelectedWidget = SNew(SBox)
+			.Padding(0.f)
+			.HAlign(HAlign_Center)
+			.VAlign(VAlign_Center)
+			[
+				SNew(STextBlock)
+				.Text(LOCTEXT("NoneSelected", "Select an entity to view details."))
+				.TextStyle(&FAppStyle::GetWidgetStyle<FTextBlockStyle>("NormalText"))
+				.Justification(ETextJustify::Center)
+			];
+
 static TAutoConsoleVariable<bool> CVarRemoteControlEnableLogicUI(TEXT("RemoteControl.EnableLogicUI"), false, TEXT("Toggles visibility of Remote Control Logic User Interface (part of the Remote Control Preset asset)"));
 
 namespace RemoteControlPanelUtils
@@ -152,10 +165,11 @@ public:
 		ChildSlot
 		.Padding(ToolBarStyle.ButtonPadding.Left, 0.f, ToolBarStyle.ButtonPadding.Right, 0.f)
 			[
-				SNew(SButton)
-				.ContentPadding(ToolBarStyle.ButtonPadding)
-				.ButtonStyle(&ToolBarStyle.ButtonStyle)
-				.OnClicked(this, &SAutoResizeButton::OnClicked)
+				SNew(SCheckBox)
+				.Padding(ToolBarStyle.CheckBoxPadding)
+				.Style(&ToolBarStyle.ToggleButton)
+				.IsChecked(this , &SAutoResizeButton::HandleIsChecked)
+				.OnCheckStateChanged(this, &SAutoResizeButton::OnCheckStateChanged)
 				.ToolTipText(ActualToolTip)
 				.Content()
 				[
@@ -166,6 +180,7 @@ public:
 					.AutoWidth()
 					.VAlign(VAlign_Center)
 					.HAlign(HAlign_Center)
+					.Padding(2.f, 4.f)
 					[
 						IconWidget
 					]
@@ -186,6 +201,7 @@ public:
 	}
 
 private:
+
 	/** Called by Slate to determine whether labels are visible */
 	EVisibility GetLabelVisibility() const
 	{
@@ -203,6 +219,7 @@ private:
 		return ActualIcon.GetIcon();
 	}
 
+	/** Retrieves the color used by icon brush. */
 	FSlateColor GetIconForegroundColor() const
 	{
 		// If any brush has a tint, don't assume it should be subdued
@@ -216,9 +233,9 @@ private:
 	}
 
 	/**
-	 * Called by Slate when this tool bar button's button is clicked
+	 * Called by Slate to check whether this toolbar button is checked or not.
 	 */
-	FReply OnClicked()
+	ECheckBoxState HandleIsChecked() const
 	{
 		IMainFrameModule& MainFrame = FModuleManager::Get().LoadModuleChecked<IMainFrameModule>("MainFrame");
 
@@ -226,16 +243,36 @@ private:
 
 		if (ActionList.IsValid() && UICommand.IsValid())
 		{
-			ActionList->ExecuteAction(UICommand.ToSharedRef());
+			return ActionList->GetCheckState(UICommand.ToSharedRef());
 		}
 
-		return FReply::Handled();
+		return ECheckBoxState::Undetermined;
+	}
+	
+	/**
+	 * Called by Slate when this tool bar button's button is clicked
+	 */
+	void OnCheckStateChanged(ECheckBoxState NewState)
+	{
+		IMainFrameModule& MainFrame = FModuleManager::Get().LoadModuleChecked<IMainFrameModule>("MainFrame");
+
+		TSharedPtr<FUICommandList> ActionList = MainFrame.GetMainFrameCommandBindings();
+
+		if (ActionList.IsValid() && UICommand.IsValid())
+		{
+			ActionList->TryExecuteAction(UICommand.ToSharedRef());
+		}
 	}
 
+private:
+
+	/** Should we use small icons or not. */
 	TAttribute<bool> bForceSmallIcons;
 
+	/** Holds the UI Command information of this button. */
 	TSharedPtr<const FUICommandInfo> UICommand;
 
+	/** Overriden icon to be used instead of default one. */
 	TAttribute<FSlateIcon> IconOverride;
 };
 
@@ -246,6 +283,10 @@ void SRemoteControlPanel::Construct(const FArguments& InArgs, URemoteControlPres
 	WidgetRegistry = MakeShared<FRCPanelWidgetRegistry>();
 	ToolkitHost = InToolkitHost;
 
+	ActivePanel = ERCPanels::RCP_Properties;
+
+	RCPanelStyle = &FRemoteControlPanelStyle::Get()->GetWidgetStyle<FRCPanelStyle>("RemoteControlPanel.MinorPanel");
+
 	TArray<TSharedRef<SWidget>> ExtensionWidgets;
 	FRemoteControlUIModule::Get().GetExtensionGenerators().Broadcast(ExtensionWidgets);
 
@@ -253,17 +294,39 @@ void SRemoteControlPanel::Construct(const FArguments& InArgs, URemoteControlPres
 
 	// Rebind All
 	AddToolbarWidget(SNew(SButton)
+		.ButtonStyle(&RCPanelStyle->FlatButtonStyle)
 		.Visibility_Lambda([this]() { return bShowRebindButton ? EVisibility::Visible : EVisibility::Collapsed; })
-		.OnClicked_Raw(this, &SRemoteControlPanel::OnClickRebindAllButton)
+		.OnClicked(this, &SRemoteControlPanel::OnClickRebindAllButton)
+		.ContentPadding(2.f)
 		[
-			SNew(STextBlock)
-			.ToolTipText(LOCTEXT("RebindButtonToolTip", "Attempt to rebind all unbound entites of the preset."))
-			.Text(LOCTEXT("RebindButtonText", "Rebind All"))
+			SNew(SHorizontalBox)
+
+			+SHorizontalBox::Slot()
+			.HAlign(HAlign_Center)
+			.VAlign(VAlign_Center)
+			.AutoWidth()
+			.Padding(2.f)
+			[
+				SNew(SImage)
+				.Image(FAppStyle::Get().GetBrush("Icons.Link"))
+				.ColorAndOpacity(FSlateColor::UseForeground())
+			]
+
+			+SHorizontalBox::Slot()
+			.VAlign(VAlign_Center)
+			.Padding(4.f, 2.f)
+			.FillWidth(1.f)
+			[	
+				SNew(STextBlock)
+				.ToolTipText(LOCTEXT("RebindButtonToolTip", "Attempt to rebind all unbound entites of the preset."))
+				.Text(LOCTEXT("RebindButtonText", "Rebind All"))
+				.TextStyle(&RCPanelStyle->PanelTextStyle)
+			]
 		]);
 
 	// Show Log
 	AddToolbarWidget(SNew(SCheckBox)
-		.Style(FRemoteControlPanelStyle::Get(), "RemoteControlPanel.Switch")
+		.Style(&RCPanelStyle->ToggleButtonStyle)
 		.ToolTipText(LOCTEXT("ShowLogTooltip", "Show/Hide remote control log."))
 		.IsChecked_Lambda([]() { return FRemoteControlLogger::Get().IsEnabled() ? ECheckBoxState::Checked : ECheckBoxState::Unchecked; })
 		.OnCheckStateChanged(this, &SRemoteControlPanel::OnLogCheckboxToggle)
@@ -272,11 +335,12 @@ void SRemoteControlPanel::Construct(const FArguments& InArgs, URemoteControlPres
 			SNew(STextBlock)
 			.Text(LOCTEXT("ShowLogLabel", "Log"))
 			.Justification(ETextJustify::Center)
+			.TextStyle(&RCPanelStyle->PanelTextStyle)
 		]);
 
 	// Edit Mode
 	AddToolbarWidget(SNew(SCheckBox)
-		.Style(FRemoteControlPanelStyle::Get(), "RemoteControlPanel.Switch")
+		.Style(&RCPanelStyle->ToggleButtonStyle)
 		.ToolTipText(LOCTEXT("EditModeTooltip", "Toggle Editing (Ctrl + E)"))
 		.IsChecked_Lambda([this]() { return this->bIsInEditMode ? ECheckBoxState::Checked : ECheckBoxState::Unchecked; })
 		.Padding(4.f)
@@ -285,6 +349,8 @@ void SRemoteControlPanel::Construct(const FArguments& InArgs, URemoteControlPres
 		[
 			SNew(STextBlock)
 			.Text(LOCTEXT("EditLabel", "Edit"))
+			.Justification(ETextJustify::Center)
+			.TextStyle(&RCPanelStyle->PanelTextStyle)
 		]
 	);
 
@@ -314,19 +380,19 @@ void SRemoteControlPanel::Construct(const FArguments& InArgs, URemoteControlPres
 	TSharedPtr<ISlateMetaData> FilterComboButtonMetaData = MakeShared<FTagMetaData>(TEXT("ContentBrowserFiltersCombo"));
 	FilterComboButton->AddMetadata(FilterComboButtonMetaData.ToSharedRef());
 
-	TSharedPtr<SHorizontalBox> TopExtensions;
-
 	EntityProtocolDetails = SNew(SBox);
 	
 	EntityList = SNew(SRCPanelExposedEntitiesList, Preset.Get(), WidgetRegistry)
+		.ExposeComboButton(CreateExposeButton())
 		.OnEntityListUpdated_Lambda([this] ()
-		{
-			UpdateEntityDetailsView(EntityList->GetSelectedEntity());
-			UpdateRebindButtonVisibility();
-			CachedExposedPropertyArgs.Reset();
-		})
-		.EditMode_Lambda([this](){ return bIsInEditMode; });
-	
+			{
+				UpdateEntityDetailsView(EntityList->GetSelectedEntity());
+				UpdateRebindButtonVisibility();
+				CachedExposedPropertyArgs.Reset();
+			}
+		)
+		.EditMode_Lambda([this]() { return bIsInEditMode; });
+
 	EntityList->OnSelectionChange().AddSP(this, &SRemoteControlPanel::UpdateEntityDetailsView);
 
 	const TAttribute<float> TreeBindingSplitRatioTop = TAttribute<float>::Create(
@@ -343,10 +409,9 @@ void SRemoteControlPanel::Construct(const FArguments& InArgs, URemoteControlPres
 			return 1.0f - Settings->TreeBindingSplitRatio;
 		}));
 
-	TSharedRef<SWidget> ExposedPropertyPanel =
-		SNew(SVerticalBox)
+	TSharedRef<SWidget> HeaderPanel = SNew(SVerticalBox)
 		+ SVerticalBox::Slot()
-		.Padding(3.f)
+		.Padding(3.f, 0.f)
 		.AutoHeight()
 		[
 			ToolbarWidgetContent.ToSharedRef()
@@ -362,7 +427,7 @@ void SRemoteControlPanel::Construct(const FArguments& InArgs, URemoteControlPres
 			.Orientation(EOrientation::Orient_Horizontal)
 		]
 		+ SVerticalBox::Slot()
-		.Padding(3.f)
+		.Padding(3.f, 0.f)
 		.AutoHeight()
 		[
 			// Auxiliary Toolbar
@@ -397,10 +462,10 @@ void SRemoteControlPanel::Construct(const FArguments& InArgs, URemoteControlPres
 			.Padding(3.f, 3.f, 3.f, 3.f)
 			[
 				SNew(SButton)
+				.ButtonStyle(&RCPanelStyle->FlatButtonStyle)
 				.ContentPadding(2.0f)
-				.ButtonStyle(FAppStyle::Get(), "FlatButton")
 				.TextStyle(FRemoteControlPanelStyle::Get(), "RemoteControlPanel.Button.TextStyle")
-				.OnClicked_Raw(this, &SRemoteControlPanel::OnClickSettingsButton)
+				.OnClicked(this, &SRemoteControlPanel::OnClickSettingsButton)
 				.ToolTipText(LOCTEXT("OpenRemoteControlSettings", "Open Remote Control settings."))
 				[
 					SNew(SImage)
@@ -429,333 +494,246 @@ void SRemoteControlPanel::Construct(const FArguments& InArgs, URemoteControlPres
 			[
 				AuxiliaryToolbarWidgetContent.ToSharedRef()
 			]
-		]
-		+ SVerticalBox::Slot()
-		.Padding(5.0f, 5.0f)
+		];
+
+	TSharedRef<SWidget> FooterPanel = SNew(SVerticalBox)
+		// Use less CPU Warning
+		+SVerticalBox::Slot()
 		.AutoHeight()
+		.Padding(FMargin(5.f, 8.f, 5.f, 5.f))
 		[
-			// Top tool bar
-			SNew(SHorizontalBox)
+			CreateCPUThrottleWarning()
+		];
 
-			+ SHorizontalBox::Slot()
-			.Padding(FMargin(0, 0, 2.0f, 0.0f))
-			.VAlign(VAlign_Center)
-			.AutoWidth()
+	// Output Log Dock panel
+	TSharedRef<SRCMinorPanel> OutputLogDockPanel = SNew(SRCMinorPanel)
+		.HeaderLabel(LOCTEXT("LogPanelLabel", "Log"))
+		[
+			SNew(SRCLogger)
+		];
+
+	TSharedRef<SWidget> OutputLogIcon = SNew(SImage)
+		.Image(FAppStyle::Get().GetBrush("MessageLog.TabIcon"))
+		.ColorAndOpacity(FSlateColor::UseForeground());
+
+	OutputLogDockPanel->AddHeaderToolbarItem(EToolbar::Left, OutputLogIcon);
+
+	// Make 3 Columns with Controllers + Behaviours + Actions
+	TSharedRef<SRCMajorPanel> LogicPanel = SNew(SRCMajorPanel)
+		.HeaderLabel(LOCTEXT("LogicPanelHeader", "LOGIC"))
+		.EnableFooter(false);
+	
+	// Controllers with Behaviours panel
+	TSharedRef<SRCMajorPanel> ControllersAndBehavioursPanel = SNew(SRCMajorPanel)
+		.EnableHeader(false)
+		.EnableFooter(true);
+
+	ControllerPanel = SNew(SRCControllerPanel, SharedThis(this))
+		.Visibility_Lambda([this] { return bIsLogicPanelEnabled && (ActivePanel == ERCPanels::RCP_Properties || ActivePanel == ERCPanels::RCP_None) ? EVisibility::Visible : EVisibility::Collapsed; });
+
+	BehaviourPanel = SNew(SRCBehaviourPanel, SharedThis(this))
+		.Visibility_Lambda([this] { return bIsLogicPanelEnabled && (ActivePanel == ERCPanels::RCP_Properties || ActivePanel == ERCPanels::RCP_None) ? EVisibility::Visible : EVisibility::Collapsed; });
+
+	ControllersAndBehavioursPanel->AddPanel(ControllerPanel.ToSharedRef(), 0.5f);
+		
+	ControllersAndBehavioursPanel->AddPanel(BehaviourPanel.ToSharedRef(), 0.5f);
+
+	// TODO : Replace this with actual Favorites Panel.
+	TSharedRef<SWidget> FavoritesItem = SNew(SBox)
+		.AddMetaData<FTagMetaData>(FTagMetaData(TEXT("Favorites")))
+		.HeightOverride(30.f)
+		.VAlign(VAlign_Center)
+		[
+			SNew(STextBlock)
+			.Text(LOCTEXT("FavoritesLabel", "Favorites"))
+			.TextStyle(&RCPanelStyle->PanelTextStyle)
+		];
+
+	ControllersAndBehavioursPanel->AddFooterToolItem(EToolbar::Left, FavoritesItem);
+
+	// Actions Panel.
+	ActionPanel = SNew(SRCActionPanel, SharedThis(this))
+		.Visibility_Lambda([this] { return bIsLogicPanelEnabled && (ActivePanel == ERCPanels::RCP_Properties || ActivePanel == ERCPanels::RCP_None) ? EVisibility::Visible : EVisibility::Collapsed; });
+
+	LogicPanel->AddPanel(ControllersAndBehavioursPanel, 0.67f);
+	LogicPanel->AddPanel(ActionPanel.ToSharedRef(), 0.33f);
+
+	// Make 2 Columns with Panel Drawer + Main Panel
+	TSharedRef<SSplitter> ContentPanel = SNew(SSplitter)
+		.Orientation(Orient_Horizontal)
+	
+		// Panel Drawer
+		+SSplitter::Slot()
+		.Value(0.05f)
+		.Resizable(false)
+		.SizeRule(SSplitter::ESizeRule::SizeToContent)
+		[
+			SAssignNew(PanelDrawer, SRCPanelDrawer)
+		]
+
+		// Main Panel
+		+SSplitter::Slot()
+		.Value(0.35f)
+		[
+			SNew(SWidgetSwitcher)
+			.WidgetIndex_Lambda([this](){ return (int32)ActivePanel; })
+
+			// 1. Nothing to be shown as all panels might be hidden (Index : 0)
+			// Since this is widget switcher we should show NullWidget in such case.
+			+ SWidgetSwitcher::Slot()
 			[
-				SNew(SButton)
-				.Visibility_Lambda([this]() { 
-					return (bIsInEditMode && Preset.IsValid() && !Preset->IsEmbeddedPreset()) ? EVisibility::Visible : EVisibility::Collapsed; 
-				})
-				.ButtonStyle(FAppStyle::Get(), "FlatButton")
-				.OnClicked(this, &SRemoteControlPanel::OnSavePreset)
+				SNew(SWidgetSwitcher)
+				.WidgetIndex_Lambda([this]() { return !bIsLogicPanelEnabled ? 0 : 1; })
+
+				+ SWidgetSwitcher::Slot()
 				[
-					SNew(SHorizontalBox)
-					+ SHorizontalBox::Slot()
-					.AutoWidth()
-					.Padding(0,0,5, 0)
-					[
-						SNew(SImage)
-						.Image(FAppStyle::Get().GetBrush("Icons.Save"))
-					]
-					+ SHorizontalBox::Slot()
-					.FillWidth(1.0f)
-					[
-						SNew(STextBlock)
-						.Text(LOCTEXT("Save", "Save"))
-					]
+					SNullWidget::NullWidget
 				]
-			]
-			+ SHorizontalBox::Slot()
-			.Padding(FMargin(0, 0, 5.0f, 0.0f))
-			.VAlign(VAlign_Center)
-			.AutoWidth()
-			[
-				SNew(SButton)
-				.Visibility_Lambda([this]() { return bIsInEditMode ? EVisibility::Visible : EVisibility::Collapsed; })
-				.ButtonStyle(FAppStyle::Get(), "FlatButton")
-				.OnClicked(this, &SRemoteControlPanel::OnCreateGroup)
+
+				+ SWidgetSwitcher::Slot()
 				[
-					SNew(SImage)
-					.Image(FAppStyle::Get().GetBrush("SceneOutliner.NewFolderIcon"))
-				]
+					LogicPanel
+				]	
 			]
-			// Function library picker
-			+ SHorizontalBox::Slot()
-			.Padding(FMargin(5.0f, 0.0f))
-			.AutoWidth()
-			[
-				CreateExposeButton()
-			]
-			// Right aligned widgets
-			+ SHorizontalBox::Slot()
-			.VAlign(VAlign_Center)
-			.HAlign(HAlign_Right)
-			.FillWidth(1.0f)
-			.Padding(0, 7.0f)
+		
+			// 2. Exposed entities List with Logic Editor (Index : 1).
+			+ SWidgetSwitcher::Slot()
 			[
-				SAssignNew(TopExtensions, SHorizontalBox)
-				// Rebind button
-				+ SHorizontalBox::Slot()
-				.VAlign(VAlign_Center)
-				.AutoWidth()
+				SNew(SWidgetSwitcher)
+				.WidgetIndex_Lambda([this]() { return !bIsLogicPanelEnabled ? 0 : 1; })
+
+				+ SWidgetSwitcher::Slot()
 				[
-					SNew(SButton)
-					.Visibility_Lambda([this]() { return bShowRebindButton ? EVisibility::Visible : EVisibility::Collapsed; })
-					.OnClicked_Raw(this, &SRemoteControlPanel::OnClickRebindAllButton)
-					[
-						SNew(STextBlock)
-						.ToolTipText(LOCTEXT("RebindButtonToolTip", "Attempt to rebind all unbound entites of the preset."))
-						.Text(LOCTEXT("RebindButtonText", "Rebind All"))
-					]
+					EntityList.ToSharedRef()
 				]
 
-				+ SHorizontalBox::Slot()
-				.VAlign(VAlign_Center)
-				.AutoWidth()
+				+ SWidgetSwitcher::Slot()
 				[
-					SNew(SButton)
-					.ButtonStyle(FAppStyle::Get(), "FlatButton")
-					.TextStyle(FRemoteControlPanelStyle::Get(), "RemoteControlPanel.Button.TextStyle")
-					.ToolTipText(LOCTEXT("EntityDetailsToolTip", "Open the details panel for the selected exposed entity."))
-					.OnClicked_Lambda([this](){ ToggleDetailsView(); return FReply::Handled(); })
+					SNew(SSplitter)
+					.Orientation(Orient_Horizontal)
+
+					+SSplitter::Slot()
+					.Value(0.4)
 					[
-						SNew(SImage)
-						.Image(FAppStyle::Get().GetBrush("LevelEditor.Tabs.Details"))
+						EntityList.ToSharedRef()
+					]
+			
+					+SSplitter::Slot()
+					.Value(0.6)
+					[
+						LogicPanel
 					]
 				]
-				// Edit Mode
-				+ SHorizontalBox::Slot()
-				.AutoWidth()
-				.Padding(2.f, 0.f)
-				.VAlign(VAlign_Center)
-				[
-					SNew(STextBlock)
-					.Text(LOCTEXT("EditModeLabel_Editor", "Edit Mode : "))
-				]
-				+ SHorizontalBox::Slot()
-				.AutoWidth()
-				.Padding(2.f, 0.f)
-				.VAlign(VAlign_Center)
-				[
-					SNew(SCheckBox)
-					.Style(FRemoteControlPanelStyle::Get(), "RemoteControlPanel.Switch")
-					.ToolTipText(LOCTEXT("EditModeTooltip", "Toggle Editing (Ctrl + E)"))
-					.IsChecked_Lambda([this]() { return this->bIsInEditMode ? ECheckBoxState::Checked : ECheckBoxState::Unchecked; })
-					.Padding(2.f)
-					.OnCheckStateChanged(this, &SRemoteControlPanel::OnEditModeCheckboxToggle)
-				]
-				// Enable Log
-				+ SHorizontalBox::Slot()
-				.VAlign(VAlign_Center)
-				.Padding(4.0f, 0)
-				.AutoWidth()
-				[
-					SNew(STextBlock)
-					.Text(LOCTEXT("EnableLogLabel", "Enable Log: "))
-				]
-				+ SHorizontalBox::Slot()
-				.VAlign(VAlign_Center)
-				.AutoWidth()
-				[
-					SNew(SCheckBox)
-					.IsChecked_Lambda([]() { return FRemoteControlLogger::Get().IsEnabled() ? ECheckBoxState::Checked : ECheckBoxState::Unchecked; })
-					.OnCheckStateChanged(this, &SRemoteControlPanel::OnLogCheckboxToggle)
-				]
-				+ SHorizontalBox::Slot()
-				.AutoWidth()
-				.Padding(FMargin(5.0f, 0.f))
-				[
-					SNew(SSeparator)
-					.Orientation(Orient_Vertical)
-				]
-				+ SHorizontalBox::Slot()
-				.VAlign(VAlign_Center)
-				.AutoWidth()
-				[
-					SAssignNew(PresetNameTextBlock, STextBlock)
-					.Font(FAppStyle::Get().GetFontStyle("DetailsView.CategoryFontStyle"))
-					.Text(FText::FromName(Preset->GetPresetName()))
-				]
-				+ SHorizontalBox::Slot()
-				.AutoWidth()
-				.Padding(FMargin(5.0f, 0.f))
+			]
+
+			// 3. Properties with Entity Details (Index : 2).
+			+ SWidgetSwitcher::Slot()
+			[
+				SNew(SSplitter)
+				.Orientation(EOrientation::Orient_Vertical)						
+				
+				// Exposed entities List
+				+ SSplitter::Slot()
+				.Value(TreeBindingSplitRatioTop)
+				.OnSlotResized(SSplitter::FOnSlotResized::CreateLambda([](float InNewSize)
+				{
+					URemoteControlSettings* Settings = GetMutableDefault<URemoteControlSettings>();
+					Settings->TreeBindingSplitRatio = InNewSize;
+					Settings->PostEditChange();
+					Settings->SaveConfig();
+				}))
 				[
-					SNew(SSeparator)
-					.Orientation(Orient_Vertical)
+					EntityList.ToSharedRef()
 				]
-				+ SHorizontalBox::Slot()
-				.VAlign(VAlign_Center)
-				.AutoWidth()
+					
+				// Entity Details
+				+ SSplitter::Slot()
+				.Value(TreeBindingSplitRatioBottom)
 				[
-					SNew(SButton)
-					.ContentPadding(2.0f)
-					.ButtonStyle(FAppStyle::Get(), "FlatButton")
-					.TextStyle(FRemoteControlPanelStyle::Get(), "RemoteControlPanel.Button.TextStyle")
-					.OnClicked_Raw(this, &SRemoteControlPanel::OnClickSettingsButton)
-					[
-						SNew(STextBlock)
-						.ColorAndOpacity(FColor::White)
-						.ToolTipText(LOCTEXT("OpenRemoteControlSettings", "Open Remote Control settings."))
-						.Font(FAppStyle::Get().GetFontStyle("FontAwesome.10"))
-						.Text(FEditorFontGlyphs::Cogs)
-					]
+					CreateEntityDetailsView()
 				]
 			]
-		]
-		+ SVerticalBox::Slot()
-		[
-			SNew(SSplitter)
-			.Orientation(Orient_Vertical)
-			+ SSplitter::Slot()
-			.Value(.8f)
+				
+			// 4. Properties with Protocols (Index : 3).
+			+ SWidgetSwitcher::Slot()
 			[
-				SNew(SBorder)
-				.Padding(FMargin(0.f, 5.f, 0.f, 0.f))
-				.BorderImage(FAppStyle::Get().GetBrush("ToolPanel.DarkGroupBorder"))
+				SNew(SSplitter)
+				.Orientation(EOrientation::Orient_Vertical)						
+				
+				// Exposed entities List
+				+ SSplitter::Slot()
+				.Value(TreeBindingSplitRatioTop)
+				.OnSlotResized(SSplitter::FOnSlotResized::CreateLambda([](float InNewSize)
+				{
+					URemoteControlSettings* Settings = GetMutableDefault<URemoteControlSettings>();
+					Settings->TreeBindingSplitRatio = InNewSize;
+					Settings->PostEditChange();
+					Settings->SaveConfig();
+				}))
 				[
-					SNew(SWidgetSwitcher)
-					.WidgetIndex_Lambda([this](){ return !bIsInEditMode ? 0 : 1; })
-					+ SWidgetSwitcher::Slot()
-					[
-						// Exposed entities List
-						SNew(SBorder)
-						.BorderImage(FAppStyle::Get().GetBrush("ToolPanel.GroupBorder"))
-						[
-							EntityList.ToSharedRef()
-						]
-					]
-					+ SWidgetSwitcher::Slot()
-					[
-						SNew(SSplitter)
-						.Orientation(EOrientation::Orient_Vertical)						
-						+ SSplitter::Slot()
-						.Value(TreeBindingSplitRatioTop)
-						.OnSlotResized(SSplitter::FOnSlotResized::CreateLambda([](float InNewSize)
-						{
-							URemoteControlSettings* Settings = GetMutableDefault<URemoteControlSettings>();
-							Settings->TreeBindingSplitRatio = InNewSize;
-							Settings->PostEditChange();
-							Settings->SaveConfig();
-						}))
-						[
-							// Exposed entities List
-							SNew(SBorder)
-							.BorderImage(FAppStyle::Get().GetBrush("ToolPanel.GroupBorder"))
-							[
-								EntityList.ToSharedRef()
-							]
-						]
-						+ SSplitter::Slot()
-						.Value(TreeBindingSplitRatioBottom)
-						[
-							SNew(SBorder)
-							.BorderImage(FAppStyle::Get().GetBrush("ToolPanel.GroupBorder"))
-							[
-								EntityProtocolDetails.ToSharedRef()
-							]
-						]
-					]
+					EntityList.ToSharedRef()
+				]
+			
+				// Protocol Details
+				+ SSplitter::Slot()
+				.Value(TreeBindingSplitRatioBottom)
+				[
+					EntityProtocolDetails.ToSharedRef()
 				]
 			]
-			+ SSplitter::Slot()
-			.Value(.2f)
+					
+			// 5. Output Log (Index : 4).
+			+ SWidgetSwitcher::Slot()
 			[
-				SNew(SRCLogger)
+				OutputLogDockPanel
 			]
-		]
-		// Use less CPU Warning
-		+SVerticalBox::Slot()
-		.AutoHeight()
-		.Padding(FMargin(5.f, 8.f, 5.f, 5.f))
-		[
-			CreateCPUThrottleWarning()
 		];
 
-	TabManager = FGlobalTabmanager::New();
-	const TSharedPtr<FTabManager::FStack> ControlPanelStackManager = FTabManager::NewStack();
-	const constexpr TCHAR* ExposedPropertyTabName = TEXT("Expose Property Tab");
-	const constexpr TCHAR* ControllerTabName = TEXT("Controller Tab");
-	const FOnSpawnTab OnExposedPropertySpawnTab = FOnSpawnTab::CreateLambda([ExposedPropertyPanel] (const FSpawnTabArgs&)
-	{
-		return SNew(SDockTab)
-		[
-			ExposedPropertyPanel
-		];
-	});
-	const FOnSpawnTab OnControllerSpawnTab = FOnSpawnTab::CreateLambda([this] (const FSpawnTabArgs&)
-	{
-		return SNew(SDockTab)
-		[
-			SAssignNew(ControllerPanel, SRCControllerPanel, SharedThis(this))
-		];
-	});
-	const FCanSpawnTab CanSpawnPanelTab = FCanSpawnTab::CreateLambda([this](const FSpawnTabArgs& InArgs)
-	{
-		return true;
-	});
-	if (!TabManager->HasTabSpawner(ExposedPropertyTabName))
-	{
-		TabManager->RegisterTabSpawner(ExposedPropertyTabName, OnExposedPropertySpawnTab, CanSpawnPanelTab);
-	}
-	if (!TabManager->HasTabSpawner(ControllerTabName))
-	{
-		TabManager->RegisterTabSpawner(ControllerTabName, OnControllerSpawnTab, CanSpawnPanelTab);
-	}
-	
-	const TSharedRef<FTabManager::FLayout> Layout = FTabManager::NewLayout("RemoteControlPanelLayout1.0") 
-	->AddArea
-	(
-		FTabManager::NewPrimaryArea()
-		->SetOrientation(Orient_Vertical)
-		->Split
-		(
-			FTabManager::NewStack()
-			->SetSizeCoefficient(0.5f)
-			->AddTab(ControllerTabName, ETabState::OpenedTab)
-		)
-		->Split
-		(
-			FTabManager::NewStack()
-			->SetSizeCoefficient(0.5f)
-			->AddTab(ExposedPropertyTabName, ETabState::OpenedTab)
-		)
-	);
-		
-	// Make 3 Columns with Exposed Properties + Controllers / Behaviours / Actions
-	TSharedRef<SSplitter> SplitterBox = SNew(SSplitter)
-	.Orientation(Orient_Horizontal)
-	+ SSplitter::Slot()
-	.Value(0.33)
-	[
-		TabManager->RestoreFrom(Layout, nullptr).ToSharedRef()
-	]
-	+ SSplitter::Slot()
-	.Value(0.33)
-	[
-		SAssignNew(BehaviourPanel, SRCBehaviourPanel, SharedThis(this))
-	]
-	+ SSplitter::Slot()
-	.Value(0.33)
-	[
-		SAssignNew(ActionPanel, SRCActionPanel, SharedThis(this))
-	];
-
-	if (CVarRemoteControlEnableLogicUI.GetValueOnAnyThread())
-	{
-		ChildSlot
+	ChildSlot
 		[
-			SplitterBox
-		];
-	}
-	else
-	{
-		ChildSlot
+			SNew(SVerticalBox)
+			+SVerticalBox::Slot()
+			.HAlign(HAlign_Fill)
+			.AutoHeight()
 			[
-				ExposedPropertyPanel
-			];
-	}
+				HeaderPanel
+			]
+			+ SVerticalBox::Slot()
+			.HAlign(HAlign_Fill)
+			.AutoHeight()
+			[
+				// Separator
+				SNew(SSeparator)
+				.SeparatorImage(FAppStyle::Get().GetBrush("Separator"))
+				.Thickness(5.f)
+				.Orientation(EOrientation::Orient_Horizontal)
+			]
+			+ SVerticalBox::Slot()
+			.HAlign(HAlign_Fill)
+			.FillHeight(1.f)
+			[
+				ContentPanel
+			]
+			+ SVerticalBox::Slot()
+			.HAlign(HAlign_Fill)
+			.AutoHeight()
+			[
+				// Separator
+				SNew(SSeparator)
+				.SeparatorImage(FAppStyle::Get().GetBrush("Separator"))
+				.Thickness(5.f)
+				.Orientation(EOrientation::Orient_Horizontal)
+			]
+			+ SVerticalBox::Slot()
+			.HAlign(HAlign_Fill)
+			.AutoHeight()
+			[
+				FooterPanel
+			]
+		];
 
 	RegisterEvents();
+	RegisterPanels();
 	CacheLevelClasses();
 	Refresh();
 	LoadSettings(InPreset->GetPresetId());
@@ -766,6 +744,7 @@ SRemoteControlPanel::~SRemoteControlPanel()
 	SaveSettings();
 
 	UnbindRemoteControlCommands();
+	UnregisterPanels();
 	UnregisterEvents();
 
 	// Clear the log
@@ -833,7 +812,6 @@ bool SRemoteControlPanel::IsExposed(const FRCExposesPropertyArgs& InPropertyArgs
 	return false;
 }
 
-
 bool SRemoteControlPanel::IsAllObjectsExposed(TArray<UObject*> InOuterObjects, const FString& InPath, bool bUsingDuplicatesInPath)
 {
 	TArray<TSharedPtr<FRemoteControlProperty>, TInlineAllocator<1>> PotentialMatches;
@@ -888,7 +866,6 @@ void SRemoteControlPanel::ToggleProperty(const FRCExposesPropertyArgs& InPropert
 		return;
 	}
 
-
 	auto PreExpose = [this]()
 	{
 		FScopedTransaction Transaction(LOCTEXT("ExposeProperty", "Expose Property"));
@@ -984,19 +961,19 @@ TSharedRef<SWidget> SRemoteControlPanel::CreateExposeButton()
 		.AllowDefaultObjects(true)
 		.Label(LOCTEXT("FunctionLibrariesLabel", "Function Libraries"))
 		.ObjectClass(UBlueprintFunctionLibrary::StaticClass())
-		.OnSelectFunction_Raw(this, &SRemoteControlPanel::ExposeFunction);
+		.OnSelectFunction(this, &SRemoteControlPanel::ExposeFunction);
 
 	SAssignNew(SubsystemFunctionPicker, SRCPanelFunctionPicker)
 		.RemoteControlPanel(SharedThis(this))
 		.Label(LOCTEXT("SubsystemFunctionLabel", "Subsystem Functions"))
 		.ObjectClass(USubsystem::StaticClass())
-		.OnSelectFunction_Raw(this, &SRemoteControlPanel::ExposeFunction);
+		.OnSelectFunction(this, &SRemoteControlPanel::ExposeFunction);
 
 	SAssignNew(ActorFunctionPicker, SRCPanelFunctionPicker)
 		.RemoteControlPanel(SharedThis(this))
 		.Label(LOCTEXT("ActorFunctionsLabel", "Actor Functions"))
 		.ObjectClass(AActor::StaticClass())
-		.OnSelectFunction_Raw(this, &SRemoteControlPanel::ExposeFunction);
+		.OnSelectFunction(this, &SRemoteControlPanel::ExposeFunction);
 	
 	MenuBuilder.BeginSection(NAME_None, LOCTEXT("ExposeHeader", "Expose"));
 	{
@@ -1083,16 +1060,25 @@ TSharedRef<SWidget> SRemoteControlPanel::CreateExposeButton()
 	MenuBuilder.EndSection();
 	
 	return SAssignNew(ExposeComboButton, SComboButton)
-		.Visibility_Lambda([this]() { return this->bIsInEditMode ? EVisibility::Visible : EVisibility::Collapsed; })
-		.ButtonStyle(FAppStyle::Get(), "PropertyEditor.AssetComboStyle")
+		.AddMetaData<FTagMetaData>(FTagMetaData(TEXT("Expose Functions")))
+		.IsEnabled_Lambda([this]() { return this->bIsInEditMode; })
+		.HAlign(HAlign_Center)
+		.VAlign(VAlign_Center)
+		.ButtonStyle(&RCPanelStyle->FlatButtonStyle)
 		.ForegroundColor(FSlateColor::UseForeground())
 		.CollapseMenuOnParentFocus(true)
-		.ContentPadding(FMargin(10.0f, 0.f))
+		.HasDownArrow(false)
+		.ContentPadding(FMargin(4.f, 2.f))
 		.ButtonContent()
 		[
-			SNew(STextBlock)
-			.TextStyle(FAppStyle::Get(), "ContentBrowser.TopBar.Font")
-			.Text(LOCTEXT("ExposeButtonLabel", "Expose"))
+			SNew(SBox)
+			.WidthOverride(RCPanelStyle->IconSize.X)
+			.HeightOverride(RCPanelStyle->IconSize.Y)
+			[
+				SNew(SImage)
+				.ColorAndOpacity(FSlateColor::UseForeground())
+				.Image(FAppStyle::GetBrush("GraphEditor.Function_16x"))
+			]
 		]
 		.MenuContent()
 		[
@@ -1264,20 +1250,19 @@ void SRemoteControlPanel::BindRemoteControlCommands()
 		FExecuteAction::CreateSP(this, &SRemoteControlPanel::FindInContentBrowser_Execute),
 		FCanExecuteAction::CreateSP(this, &SRemoteControlPanel::CanFindInContentBrowser));
 
-	ActionList.MapAction(
-		Commands.ToggleExposeFunctions,
-		FExecuteAction::CreateLambda([this] { OnToggleExposeFunctions(); }),
-		FCanExecuteAction::CreateLambda([] { return true; }));
-
 	ActionList.MapAction(
 		Commands.ToggleProtocolMappings,
-		FExecuteAction::CreateLambda([this] { OnToggleProtocolMappings(); }),
-		FCanExecuteAction::CreateLambda([] { return true; }));
+		FExecuteAction::CreateSP(this, &SRemoteControlPanel::ToggleProtocolMappings_Execute),
+		FCanExecuteAction::CreateSP(this, &SRemoteControlPanel::CanToggleProtocolsMode),
+		FIsActionChecked::CreateSP(this, &SRemoteControlPanel::IsInProtocolsMode),
+		FIsActionButtonVisible::CreateSP(this, &SRemoteControlPanel::CanToggleProtocolsMode));
 
 	ActionList.MapAction(
 		Commands.ToggleLogicEditor,
-		FExecuteAction::CreateLambda([this] { OnToggleLogicEditor(); }),
-		FCanExecuteAction::CreateLambda([] { return true; }));
+		FExecuteAction::CreateSP(this, &SRemoteControlPanel::ToggleLogicEditor_Execute),
+		FCanExecuteAction::CreateSP(this, &SRemoteControlPanel::CanToggleLogicPanel),
+		FIsActionChecked::CreateSP(this, &SRemoteControlPanel::IsLogicPanelEnabled),
+		FIsActionButtonVisible::CreateSP(this, &SRemoteControlPanel::CanToggleLogicPanel));
 
 	ActionList.MapAction(
 		Commands.DeleteEntity,
@@ -1302,7 +1287,6 @@ void SRemoteControlPanel::UnbindRemoteControlCommands()
 
 		ActionList.UnmapAction(Commands.SavePreset);
 		ActionList.UnmapAction(Commands.FindPresetInContentBrowser);
-		ActionList.UnmapAction(Commands.ToggleExposeFunctions);
 		ActionList.UnmapAction(Commands.ToggleProtocolMappings);
 		ActionList.UnmapAction(Commands.ToggleLogicEditor);
 		ActionList.UnmapAction(Commands.DeleteEntity);
@@ -1312,8 +1296,6 @@ void SRemoteControlPanel::UnbindRemoteControlCommands()
 
 void SRemoteControlPanel::RegisterEvents()
 {
-	FAssetRegistryModule& AssetRegistryModule = FModuleManager::LoadModuleChecked<FAssetRegistryModule>(TEXT("AssetRegistry"));
-	AssetRegistryModule.Get().OnAssetRenamed().AddSP(this, &SRemoteControlPanel::OnAssetRenamed);
 	FEditorDelegates::MapChange.AddSP(this, &SRemoteControlPanel::OnMapChange);
 	
 	if (GEditor)
@@ -1352,13 +1334,91 @@ void SRemoteControlPanel::UnregisterEvents()
 	}
 
 	FEditorDelegates::MapChange.RemoveAll(this);
-	
-	FAssetRegistryModule& AssetRegistryModule = FModuleManager::LoadModuleChecked<FAssetRegistryModule>(TEXT("AssetRegistry"));
-	AssetRegistryModule.Get().OnAssetRenamed().RemoveAll(this);
-	
+
 	UMaterial::OnMaterialCompilationFinished().RemoveAll(this);
 }
 
+void SRemoteControlPanel::RegisterPanels()
+{
+	if (PanelDrawer.IsValid())
+	{
+		PanelDrawer->OnRCPanelToggled().BindSP(this, &SRemoteControlPanel::OnRCPanelToggled);
+
+		{// Properties Panel
+			TSharedRef<FRCPanelDrawerArgs> PropertiesPanel = MakeShared<FRCPanelDrawerArgs>(ERCPanels::RCP_Properties);
+
+			PropertiesPanel->bDrawnByDefault = true;
+			PropertiesPanel->bRotateIconBy90 = true;
+			PropertiesPanel->Label = LOCTEXT("PropertiesPanelLabel", "Expose");
+			PropertiesPanel->ToolTip = LOCTEXT("PropertiesPanelTooltip", "Open exposed properties panel.");
+			PropertiesPanel->Icon = FSlateIcon(FAppStyle::GetAppStyleSetName(), "EditorPreferences.TabIcon");
+
+			PanelDrawer->RegisterPanel(PropertiesPanel);
+
+			RegisteredDrawers.Add(PropertiesPanel->GetPanelID(), PropertiesPanel);
+		}
+		
+		{// Properties With Details Panel
+			TSharedRef<FRCPanelDrawerArgs> EntityDetailsPanel = MakeShared<FRCPanelDrawerArgs>(ERCPanels::RCP_EntityDetails);
+
+			EntityDetailsPanel->bDrawnByDefault = false;
+			EntityDetailsPanel->bRotateIconBy90 = false;
+			EntityDetailsPanel->Label = LOCTEXT("EntityDetailsPanelLabel", "Details");
+			EntityDetailsPanel->ToolTip = LOCTEXT("EntityDetailsPanelTooltip", "Open entity details panel.");
+			EntityDetailsPanel->Icon = FSlateIcon(FAppStyle::GetAppStyleSetName(), "LevelEditor.Tabs.Details");
+
+			PanelDrawer->RegisterPanel(EntityDetailsPanel);
+
+			RegisteredDrawers.Add(EntityDetailsPanel->GetPanelID(), EntityDetailsPanel);
+		}
+
+		{// Properties With Protocols Panel
+			TSharedRef<FRCPanelDrawerArgs> EntityProtocolsPanel = MakeShared<FRCPanelDrawerArgs>(ERCPanels::RCP_Protocols);
+
+			EntityProtocolsPanel->bDrawnByDefault = true;
+			EntityProtocolsPanel->bRotateIconBy90 = false;
+			EntityProtocolsPanel->DrawerVisibility = bIsInProtocolsMode ? EVisibility::Visible : EVisibility::Collapsed;
+			EntityProtocolsPanel->Label = LOCTEXT("EntityProtocolsPanelLabel", "Protocols");
+			EntityProtocolsPanel->ToolTip = LOCTEXT("EntityProtocolsPanelTooltip", "Open entity protocols panel.");
+			EntityProtocolsPanel->Icon = FSlateIcon(FAppStyle::GetAppStyleSetName(), "LevelEditor.Tabs.StatsViewer");
+
+			PanelDrawer->RegisterPanel(EntityProtocolsPanel);
+
+			RegisteredDrawers.Add(EntityProtocolsPanel->GetPanelID(), EntityProtocolsPanel);
+		}
+
+		{// Output Log Panel
+			TSharedRef<FRCPanelDrawerArgs> OutputLogPanel = MakeShared<FRCPanelDrawerArgs>(ERCPanels::RCP_OutputLog);
+
+			OutputLogPanel->bDrawnByDefault = true;
+			OutputLogPanel->bRotateIconBy90 = false;
+			OutputLogPanel->DrawerVisibility = FRemoteControlLogger::Get().IsEnabled() ? EVisibility::Visible : EVisibility::Collapsed;
+			OutputLogPanel->Label = LOCTEXT("OutputLogPanelPanelLabel", "Log");
+			OutputLogPanel->ToolTip = LOCTEXT("OutputLogPanelPanelTooltip", "Open output log panel.");
+			OutputLogPanel->Icon = FSlateIcon(FAppStyle::GetAppStyleSetName(), "MessageLog.TabIcon");
+
+			PanelDrawer->RegisterPanel(OutputLogPanel);
+
+			RegisteredDrawers.Add(OutputLogPanel->GetPanelID(), OutputLogPanel);
+		}
+	}
+}
+
+void SRemoteControlPanel::UnregisterPanels()
+{
+	if (PanelDrawer.IsValid())
+	{
+		PanelDrawer->OnRCPanelToggled().Unbind();
+
+		for (TMap<ERCPanels, TSharedRef<FRCPanelDrawerArgs>>::TIterator RegisteredDrawer = RegisteredDrawers.CreateIterator(); RegisteredDrawer; ++RegisteredDrawer)
+		{
+			PanelDrawer->UnregisterPanel(RegisteredDrawer->Value);
+
+			RegisteredDrawer.RemoveCurrent();
+		}
+	}
+}
+
 void SRemoteControlPanel::Refresh()
 {
 	BlueprintPicker->Refresh();
@@ -1424,7 +1484,6 @@ void SRemoteControlPanel::Unexpose(const FRCExposesPropertyArgs& InPropertyArgs)
 	}
 }
 
-
 void SRemoteControlPanel::OnEditModeCheckboxToggle(ECheckBoxState State)
 {
 	bIsInEditMode = (State == ECheckBoxState::Checked) ? true : false;
@@ -1435,6 +1494,23 @@ void SRemoteControlPanel::OnLogCheckboxToggle(ECheckBoxState State)
 {
 	const bool bIsLogEnabled = (State == ECheckBoxState::Checked) ? true : false;
 	FRemoteControlLogger::Get().EnableLog(bIsLogEnabled);
+	
+	if (PanelDrawer.IsValid())
+	{
+		TSharedRef<FRCPanelDrawerArgs> OutputLogPanel = RegisteredDrawers.FindChecked(ERCPanels::RCP_OutputLog);
+
+		OutputLogPanel->DrawerVisibility = bIsLogEnabled ? EVisibility::Visible : EVisibility::Collapsed;
+
+		// When we are not enabling log collapse the drawer.
+		PanelDrawer->TogglePanel(OutputLogPanel, !bIsLogEnabled);
+
+		if (!bIsLogEnabled)
+		{
+			TSharedRef<FRCPanelDrawerArgs> PropertiesPanel = RegisteredDrawers.FindChecked(ERCPanels::RCP_Properties);
+
+			PanelDrawer->TogglePanel(PropertiesPanel);
+		}
+	}
 }
  
 void SRemoteControlPanel::OnBlueprintReinstanced()
@@ -1442,33 +1518,6 @@ void SRemoteControlPanel::OnBlueprintReinstanced()
 	Refresh();
 }
 
-FReply SRemoteControlPanel::OnCreateGroup()
-{
-	FScopedTransaction Transaction(LOCTEXT("CreateGroup", "Create Group"));
-	Preset->Modify();
-	Preset->Layout.CreateGroup();
-	return FReply::Handled();
-}
-
-FReply SRemoteControlPanel::OnSavePreset()
-{
-	if (!Preset.IsValid() || !Preset->IsAsset())
-	{
-		UE_LOG(LogRemoteControl, Log, TEXT("Invalid object to save: %s"), (Preset.IsValid()) ? *Preset->GetFullName() : TEXT("Null Object"));
-		
-		return FReply::Handled();
-	}
-
-	TArray<UPackage*> PackagesToSave;
-	constexpr bool bCheckDirtyOnAssetSave = false;
-	constexpr bool bPromptToSave = false;
-	
-	PackagesToSave.Add(Preset->GetOutermost());
-	FEditorFileUtils::PromptForCheckoutAndSave(PackagesToSave, bCheckDirtyOnAssetSave, bPromptToSave);
-	
-	return FReply::Handled();
-}
-
 void SRemoteControlPanel::ExposeProperty(UObject* Object, FRCFieldPathInfo Path)
 {
 	if (Path.Resolve(Object))
@@ -1537,6 +1586,24 @@ void SRemoteControlPanel::ToggleDetailsView()
 
 TSharedRef<SWidget> SRemoteControlPanel::CreateEntityDetailsView()
 {
+	WrappedEntityDetailsView = SNew(SBorder)
+		.BorderImage(&RCPanelStyle->ContentAreaBrush)
+		.Padding(RCPanelStyle->PanelPadding);
+
+	// Details Panel Dock Widget.
+	TSharedRef<SRCMinorPanel> EntityDetailsDockPanel = SNew(SRCMinorPanel)
+		.HeaderLabel(LOCTEXT("EntityDetailsLabel", "Details"))
+		[
+			WrappedEntityDetailsView.ToSharedRef()
+		];
+
+	// Details Panel Icon
+	TSharedRef<SWidget> DetailsIcon = SNew(SImage)
+		.Image(FAppStyle::Get().GetBrush("LevelEditor.Tabs.Details"))
+		.ColorAndOpacity(FSlateColor::UseForeground());
+
+	EntityDetailsDockPanel->AddHeaderToolbarItem(EToolbar::Left, DetailsIcon);
+
 	FDetailsViewArgs Args;
 	Args.bShowOptions = false;
 	Args.bAllowFavoriteSystem = false;
@@ -1546,12 +1613,22 @@ TSharedRef<SWidget> SRemoteControlPanel::CreateEntityDetailsView()
 	EntityDetailsView = FModuleManager::GetModuleChecked<FPropertyEditorModule>("PropertyEditor").CreateStructureDetailView(MoveTemp(Args), FStructureDetailsViewArgs(), nullptr);
 
 	UpdateEntityDetailsView(EntityList->GetSelectedEntity());
-	if (ensure(EntityDetailsView && EntityDetailsView->GetWidget()))
+	
+	const bool bCanShowDetailsView = SelectedEntity.IsValid() && (SelectedEntity->GetRCType() != SRCPanelTreeNode::Group) && (SelectedEntity->GetRCType() != SRCPanelTreeNode::FieldChild);
+
+	if (bCanShowDetailsView)
 	{
-		return EntityDetailsView->GetWidget().ToSharedRef();
+		if (ensure(EntityDetailsView && EntityDetailsView->GetWidget()))
+		{
+			WrappedEntityDetailsView->SetContent(EntityDetailsView->GetWidget().ToSharedRef());
+		}
+	}
+	else
+	{
+		WrappedEntityDetailsView->SetContent(NoneSelectedWidget);
 	}
 
-	return SNullWidget::NullWidget;
+	return EntityDetailsDockPanel;
 }
 
 void SRemoteControlPanel::UpdateEntityDetailsView(const TSharedPtr<SRCPanelTreeNode>& SelectedNode)
@@ -1570,9 +1647,20 @@ void SRemoteControlPanel::UpdateEntityDetailsView(const TSharedPtr<SRCPanelTreeN
 		}
 	}
 
-	if (EntityDetailsView)
+	const bool bCanShowDetailsView = SelectedEntity.IsValid() && (SelectedEntity->GetRCType() != SRCPanelTreeNode::Group) && (SelectedEntity->GetRCType() != SRCPanelTreeNode::FieldChild);
+
+	if (bCanShowDetailsView)
 	{
-		EntityDetailsView->SetStructureData(SelectedEntityPtr);
+		if (EntityDetailsView)
+		{
+			EntityDetailsView->SetStructureData(SelectedEntityPtr);
+
+			WrappedEntityDetailsView->SetContent(EntityDetailsView->GetWidget().ToSharedRef());
+		}
+	}
+	else
+	{
+		WrappedEntityDetailsView->SetContent(NoneSelectedWidget);
 	}
 
 	static const FName ProtocolWidgetsModuleName = "RemoteControlProtocolWidgets";	
@@ -1587,10 +1675,14 @@ void SRemoteControlPanel::UpdateEntityDetailsView(const TSharedPtr<SRCPanelTreeN
 			}
 			else
 			{
-				EntityProtocolDetails->SetContent(SNullWidget::NullWidget);
+				EntityProtocolDetails->SetContent(NoneSelectedWidget);
 			}
 		}
 	}
+	else
+	{
+		EntityProtocolDetails->SetContent(NoneSelectedWidget);
+	}
 
 	// Trigger search to list the search results specific to selected group.
 	if (SearchedText.IsValid() && !SearchedText->IsEmptyOrWhitespace() && SearchedText->ToString().Len() > 3)
@@ -1645,14 +1737,6 @@ void SRemoteControlPanel::UpdateActorFunctionPicker()
 	}
 }
 
-void SRemoteControlPanel::OnAssetRenamed(const FAssetData& Asset, const FString&)
-{
-	if (Preset.IsValid() && PresetNameTextBlock)
-	{
-		PresetNameTextBlock->SetText(FText::FromName(Preset->GetPresetName()));
-	}
-}
-
 void SRemoteControlPanel::OnEntityExposed(URemoteControlPreset* InPreset, const FGuid& InEntityId)
 {
 	CachedExposedPropertyArgs.Empty();
@@ -1787,21 +1871,12 @@ void SRemoteControlPanel::RegisterAuxiliaryToolBar()
 		{
 			FToolMenuSection& ToolsSection = ToolbarBuilder->AddSection("Tools");
 
-			ToolsSection.AddEntry(FToolMenuEntry::InitWidget("Functions"
-				, SNew(SAutoResizeButton)
-					.UICommand(FRemoteControlCommands::Get().ToggleExposeFunctions)
-					.ForceSmallIcons_Static(SRemoteControlPanel::ShouldForceSmallIcons)
-					.IconOverride(FSlateIcon(FAppStyle::GetAppStyleSetName(), TEXT("GraphEditor.Function_24x")))
-				, LOCTEXT("FunctionsLabel", "Functions")
-			)
-			);
-			
-			ToolsSection.AddEntry(FToolMenuEntry::InitWidget("Mappings"
+			ToolsSection.AddEntry(FToolMenuEntry::InitWidget("Protocols"
 				, SNew(SAutoResizeButton)
 					.UICommand(FRemoteControlCommands::Get().ToggleProtocolMappings)
 					.ForceSmallIcons_Static(SRemoteControlPanel::ShouldForceSmallIcons)
 					.IconOverride(FSlateIcon(FAppStyle::GetAppStyleSetName(), "LevelEditor.Tabs.StatsViewer"))
-				, LOCTEXT("MappingsLabel", "Mappings")
+				, LOCTEXT("ProtocolsLabel", "Protocols")
 			)
 			);
 			
@@ -1925,16 +2000,66 @@ bool SRemoteControlPanel::ShouldForceSmallIcons()
 	return false;
 }
 
-void SRemoteControlPanel::OnToggleExposeFunctions() const
+void SRemoteControlPanel::ToggleProtocolMappings_Execute()
 {
+	bIsInProtocolsMode = !bIsInProtocolsMode;
+
+	if (PanelDrawer.IsValid())
+	{
+		TSharedRef<FRCPanelDrawerArgs> ProtocolsPanel = RegisteredDrawers.FindChecked(ERCPanels::RCP_Protocols);
+		
+		ProtocolsPanel->DrawerVisibility = bIsInProtocolsMode ? EVisibility::Visible : EVisibility::Collapsed;
+
+		// When we are not in Protocols Mode collapse the drawer.
+		PanelDrawer->TogglePanel(ProtocolsPanel, !bIsInProtocolsMode);
+
+		if (!bIsInProtocolsMode)
+		{
+			TSharedRef<FRCPanelDrawerArgs> PropertiesPanel = RegisteredDrawers.FindChecked(ERCPanels::RCP_Properties);
+
+			PanelDrawer->TogglePanel(PropertiesPanel);
+		}
+	}
+}
+
+bool SRemoteControlPanel::CanToggleProtocolsMode() const
+{
+	return bIsInEditMode;
 }
 
-void SRemoteControlPanel::OnToggleProtocolMappings() const
+bool SRemoteControlPanel::IsInProtocolsMode() const
 {
+	return bIsInEditMode && bIsInProtocolsMode;
 }
 
-void SRemoteControlPanel::OnToggleLogicEditor() const
+void SRemoteControlPanel::ToggleLogicEditor_Execute()
 {
+	bIsLogicPanelEnabled = !bIsLogicPanelEnabled;
+
+	if (PanelDrawer.IsValid() && (ActivePanel != ERCPanels::RCP_Properties || ActivePanel == ERCPanels::RCP_None))
+	{
+		TSharedRef<FRCPanelDrawerArgs> PropertiesPanel = RegisteredDrawers.FindChecked(ERCPanels::RCP_Properties);
+
+		PanelDrawer->TogglePanel(PropertiesPanel);
+	}
+}
+
+bool SRemoteControlPanel::CanToggleLogicPanel() const
+{
+	return bIsInEditMode;
+}
+
+bool SRemoteControlPanel::IsLogicPanelEnabled() const
+{
+	return bIsInEditMode && bIsLogicPanelEnabled;
+}
+
+void SRemoteControlPanel::OnRCPanelToggled(ERCPanels InPanelID)
+{
+	if (InPanelID != ActivePanel)
+	{
+		ActivePanel = InPanelID;
+	}
 }
 
 TSharedPtr<class SRCLogicPanelBase> SRemoteControlPanel::GetActiveLogicPanel() const
@@ -1963,9 +2088,9 @@ void SRemoteControlPanel::DeleteEntity_Execute()
 	// ~ Delete Logic Item ~
 	// 
 	// If the user focus is currently active on a Logic panel then route the Delete command to it and return.
-	if (TSharedPtr<class SRCLogicPanelBase> ActivePanel = GetActiveLogicPanel())
+	if (TSharedPtr<class SRCLogicPanelBase> ActiveLogicPanel = GetActiveLogicPanel())
 	{
-		ActivePanel->DeleteSelectedPanelItem();
+		ActiveLogicPanel->DeleteSelectedPanelItem();
 
 		return; // handled
 	}
@@ -1995,7 +2120,7 @@ void SRemoteControlPanel::DeleteEntity_Execute()
 
 bool SRemoteControlPanel::CanDeleteEntity() const
 {
-	if (TSharedPtr<class SRCLogicPanelBase> ActivePanel = GetActiveLogicPanel())
+	if (TSharedPtr<class SRCLogicPanelBase> ActiveLogicPanel = GetActiveLogicPanel())
 	{
 		return true; // User has focus on a logic panel
 	}
diff --git a/Engine/Plugins/VirtualProduction/RemoteControl/Source/RemoteControlUI/Private/UI/SRemoteControlPanel.h b/Engine/Plugins/VirtualProduction/RemoteControl/Source/RemoteControlUI/Private/UI/SRemoteControlPanel.h
--- a/Engine/Plugins/VirtualProduction/RemoteControl/Source/RemoteControlUI/Private/UI/SRemoteControlPanel.h
+++ b/Engine/Plugins/VirtualProduction/RemoteControl/Source/RemoteControlUI/Private/UI/SRemoteControlPanel.h
@@ -14,12 +14,15 @@
 #include "Widgets/SCompoundWidget.h"
 #include "Widgets/DeclarativeSyntaxSupport.h"
 
+enum class ERCPanels : uint8;
 class AActor;
 struct EVisibility;
 struct FAssetData;
 class FExposedEntityDragDrop;
 struct FListEntry;
+struct FRCPanelDrawerArgs;
 class FRCPanelWidgetRegistry;
+struct FRCPanelStyle;
 struct FRemoteControlEntity;
 class FReply;
 class IToolkitHost;
@@ -31,6 +34,7 @@ class SComboButton;
 struct SRCPanelTreeNode;
 class SRCPanelFunctionPicker;
 class SRemoteControlPanel;
+class SRCPanelDrawer;
 class SRCPanelExposedEntitiesList;
 class SRCPanelFilter;
 class SSearchBox;
@@ -147,6 +151,11 @@ private:
 	/** Unregister editor events */
 	void UnregisterEvents();
 
+	/** Register panels to the drawer. */
+	void RegisterPanels();
+	/** Unregister panels from the drawer. */
+	void UnregisterPanels();
+
 	/** Unexpose a field from the preset. */
 	void Unexpose(const FRCExposesPropertyArgs& InArgs);
 
@@ -156,12 +165,6 @@ private:
 	/** Handler called when a blueprint is reinstanced. */
 	void OnBlueprintReinstanced();
 
-	/** Handles creating a new group. */
-	FReply OnCreateGroup();
-
-	/** Handles saving the Preset. */
-	FReply OnSavePreset();
-
 	/** Expose a property using its handle. */
 	void ExposeProperty(UObject* Object, FRCFieldPathInfo FieldPath);
 
@@ -225,9 +228,6 @@ private:
 	/** Triggers a next frame update of the actor function picker to ensure that added actors are valid. */
 	void UpdateActorFunctionPicker();
 	
-	/** Handle updating the preset name textblock when it's renamed. */
-	void OnAssetRenamed(const FAssetData& Asset, const FString&);
-
 	/** Handle clicking on the setting button. */
 	FReply OnClickSettingsButton();
 
@@ -262,12 +262,16 @@ private:
 
 	static bool ShouldForceSmallIcons();
 
-	void OnToggleExposeFunctions() const;
+	void ToggleProtocolMappings_Execute();
+	bool CanToggleProtocolsMode() const;
+	bool IsInProtocolsMode() const;
 
-	void OnToggleProtocolMappings() const;
+	void ToggleLogicEditor_Execute();
+	bool CanToggleLogicPanel() const;
+	bool IsLogicPanelEnabled() const;
+
+	void OnRCPanelToggled(ERCPanels InPanelID);
 
-	void OnToggleLogicEditor() const;
-	
 	/** Called when user attempts to delete a group/exposed entity. */
 	void DeleteEntity_Execute();
 
@@ -294,6 +298,9 @@ private:
 	/** Saves settings from config based on the preset identifier. */
 	void SaveSettings();
 
+	/** Retrieves active logic panel. */
+	TSharedPtr<class SRCLogicPanelBase> GetActiveLogicPanel() const;
+
 private:
 	static const FName DefaultRemoteControlPanelToolBarName;
 	static const FName AuxiliaryRemoteControlPanelToolBarName;
@@ -301,6 +308,10 @@ private:
 	TStrongObjectPtr<URemoteControlPreset> Preset;
 	/** Whether the panel is in edit mode. */
 	bool bIsInEditMode = true;
+	/** Whether the panel is in protocols mode. */
+	bool bIsInProtocolsMode = false;
+	/** Whether the logic panel is enabled or not. */
+	bool bIsLogicPanelEnabled = false;
 	/** Whether objects need to be re-resolved because PIE Started or ended. */
 	bool bTriggerRefreshForPIE = false;
 	/** Delegate called when the edit mode changes. */
@@ -321,14 +332,16 @@ private:
 	TSharedPtr<SClassViewer> ClassPicker;
 	/** Holds the field's details. */
 	TSharedPtr<class IStructureDetailsView> EntityDetailsView;
+	/** Wrapper widget for entity details view. */
+	TSharedPtr<SBorder> WrappedEntityDetailsView;
+	/** Helper widget for entity details view and protocol details view. */
+	static TSharedRef<SBox> NoneSelectedWidget;
 	/** Holds the field's protocol details. */
 	TSharedPtr<SBox> EntityProtocolDetails;
 	/** Whether to show the rebind all button. */
 	bool bShowRebindButton = false;
 	/** Cache of exposed property arguments. */
 	TSet<FRCExposesPropertyArgs> CachedExposedPropertyArgs;
-	/** Preset name widget. */
-	TSharedPtr<STextBlock> PresetNameTextBlock;
 	/** Holds a cache of widgets. */
 	TSharedPtr<FRCPanelWidgetRegistry> WidgetRegistry;
 	/** Holds the handle to a timer set for next tick. Used to not schedule more than once event per frame */
@@ -355,26 +368,27 @@ private:
 	TSharedPtr<SRCPanelFilter> FilterPtr;
 	/** Actively serached term. */
 	TSharedPtr<FText> SearchedText;
+	/** Panel Drawer widget holds all docked panels. */
+	TSharedPtr<SRCPanelDrawer> PanelDrawer;
+	/** Map of Opened Drawers. */
+	TMap<ERCPanels, TSharedRef<FRCPanelDrawerArgs>> RegisteredDrawers;
+	/** Panel Style reference. */
+	const FRCPanelStyle* RCPanelStyle;
+	/** Stores the active panel that is drawn. */
+	ERCPanels ActivePanel;
 
 	// ~ Remote Control Logic Panels ~
 
 	/** Controller panel UI widget for Remote Control Logic*/
 	TSharedPtr<class SRCControllerPanel> ControllerPanel;
-
 	/** Behaviour panel UI widget for Remote Control Logic*/
 	TSharedPtr<class SRCBehaviourPanel> BehaviourPanel;
-
 	/** Action panel UI widget for Remote Control Logic*/
 	TSharedPtr<class SRCActionPanel> ActionPanel;
 
-
 public:
-	static const float MinimumPanelWidth;
 
-	/** Tab Manager used to handle spawning of Controller and Exposed Property tabs*/
-	TSharedPtr<FGlobalTabmanager> TabManager;
-
-public:
+	static const float MinimumPanelWidth;
 	// Global Delegates for Remote Control Logic
 	FOnControllerAdded OnControllerAdded;	
 	FOnBehaviourAdded OnBehaviourAdded;
@@ -384,7 +398,4 @@ public:
 	FOnEmptyControllers OnEmptyControllers;
 	FOnEmptyBehaviours OnEmptyBehaviours;
 	FOnEmptyActions OnEmptyActions;
-
-private:
-	TSharedPtr<class SRCLogicPanelBase> GetActiveLogicPanel() const;
 };
\ No newline at end of file
