Change check to ensure until we can reproduce

diff --git a/Engine/Plugins/VirtualProduction/RemoteControl/Source/WebRemoteControl/Private/WebSocketMessageHandler.cpp b/Engine/Plugins/VirtualProduction/RemoteControl/Source/WebRemoteControl/Private/WebSocketMessageHandler.cpp
--- a/Engine/Plugins/VirtualProduction/RemoteControl/Source/WebRemoteControl/Private/WebSocketMessageHandler.cpp
+++ b/Engine/Plugins/VirtualProduction/RemoteControl/Source/WebRemoteControl/Private/WebSocketMessageHandler.cpp
@@ -1525,8 +1525,11 @@ bool FWebSocketMessageHandler::WriteActorPropertyChangePayload(URemoteControlPre
 		}
 	}
 
-	FStructOnScope ActorsModifedOnScope = WebSocketMessageHandlerStructUtils::CreateModifiedActorsStructOnScope(InPreset, ModifiedActorsOnScope);
-	WebRemoteControlInternalUtils::SerializeStructOnScope(ActorsModifedOnScope, InWriter);
+	if (ensure(ModifiedActorsOnScope.Num()))
+	{
+		FStructOnScope ActorsModifedOnScope = WebSocketMessageHandlerStructUtils::CreateModifiedActorsStructOnScope(InPreset, ModifiedActorsOnScope);
+		WebRemoteControlInternalUtils::SerializeStructOnScope(ActorsModifedOnScope, InWriter);
+	}
 
 	return bHasProperty;
 }
