Remote Control Behaviour entry in content browser context menu should be with Remote Control Preset

diff --git a/Engine/Plugins/VirtualProduction/RemoteControl/Source/RemoteControlUI/Private/AssetTools/RemoteControlPresetActions.cpp b/Engine/Plugins/VirtualProduction/RemoteControl/Source/RemoteControlUI/Private/AssetTools/RemoteControlPresetActions.cpp
--- a/Engine/Plugins/VirtualProduction/RemoteControl/Source/RemoteControlUI/Private/AssetTools/RemoteControlPresetActions.cpp
+++ b/Engine/Plugins/VirtualProduction/RemoteControl/Source/RemoteControlUI/Private/AssetTools/RemoteControlPresetActions.cpp
@@ -5,6 +5,7 @@
 #include "Engine/Engine.h"
 #include "RemoteControlPreset.h"
 #include "AssetEditor/RemoteControlPresetEditorToolkit.h"
+#include "IRemoteControlUIModule.h"
 
 #define LOCTEXT_NAMESPACE "AssetTypeActions"
 
@@ -14,7 +15,7 @@ FRemoteControlPresetActions::FRemoteControlPresetActions(const TSharedRef<ISlate
 
 uint32 FRemoteControlPresetActions::GetCategories()
 {
-	return EAssetTypeCategories::Misc;
+	return IRemoteControlUIModule::Get().GetRemoteControlAssetCategory();
 }
 
 FText FRemoteControlPresetActions::GetName() const
diff --git a/Engine/Plugins/VirtualProduction/RemoteControl/Source/RemoteControlUI/Private/Factories/RCBehaviourNodeFactory.cpp b/Engine/Plugins/VirtualProduction/RemoteControl/Source/RemoteControlUI/Private/Factories/RCBehaviourNodeFactory.cpp
--- a/Engine/Plugins/VirtualProduction/RemoteControl/Source/RemoteControlUI/Private/Factories/RCBehaviourNodeFactory.cpp
+++ b/Engine/Plugins/VirtualProduction/RemoteControl/Source/RemoteControlUI/Private/Factories/RCBehaviourNodeFactory.cpp
@@ -2,8 +2,8 @@
 
 #include "RCBehaviourNodeFactory.h"
 
-#include "AssetToolsModule.h"
 #include "Behaviour/RCBehaviourBlueprintNode.h"
+#include "IRemoteControlUIModule.h"
 #include "Kismet2/KismetEditorUtilities.h"
 
 #define LOCTEXT_NAMESPACE "RCBehaviourNodeFactory"
@@ -19,7 +19,7 @@ URCBehaviourNodeFactory::URCBehaviourNodeFactory()
 
 FText URCBehaviourNodeFactory::GetDisplayName() const
 {
-	return LOCTEXT("DisplayName", "Remote Control Behaviour Node");
+	return LOCTEXT("DisplayName", "Remote Control Behavior");
 }
 
 UObject* URCBehaviourNodeFactory::FactoryCreateNew(UClass* Class, UObject* InParent, FName Name, EObjectFlags Flags, UObject* Context, FFeedbackContext* Warn)
@@ -35,8 +35,7 @@ UObject* URCBehaviourNodeFactory::FactoryCreateNew(UClass* Class, UObject* InPar
 
 uint32 URCBehaviourNodeFactory::GetMenuCategories() const
 {
-	IAssetTools& AssetTools = FModuleManager::LoadModuleChecked<FAssetToolsModule>("AssetTools").Get();
-	return AssetTools.RegisterAdvancedAssetCategory("Remote Control", LOCTEXT("AssetCategoryName", "Remote Control Behaviour Node"));
+	return IRemoteControlUIModule::Get().GetRemoteControlAssetCategory();
 }
 
 #undef LOCTEXT_NAMESPACE
\ No newline at end of file
diff --git a/Engine/Plugins/VirtualProduction/RemoteControl/Source/RemoteControlUI/Private/RemoteControlUIModule.cpp b/Engine/Plugins/VirtualProduction/RemoteControl/Source/RemoteControlUI/Private/RemoteControlUIModule.cpp
--- a/Engine/Plugins/VirtualProduction/RemoteControl/Source/RemoteControlUI/Private/RemoteControlUIModule.cpp
+++ b/Engine/Plugins/VirtualProduction/RemoteControl/Source/RemoteControlUI/Private/RemoteControlUIModule.cpp
@@ -424,10 +424,17 @@ URemoteControlPreset* FRemoteControlUIModule::GetActivePreset() const
 	return nullptr;
 }
 
+uint32 FRemoteControlUIModule::GetRemoteControlAssetCategory() const
+{
+	return RemoteControlAssetCategoryBit;
+}
+
 void FRemoteControlUIModule::RegisterAssetTools()
 {
 	if (FAssetToolsModule* AssetToolsModule = FModuleManager::GetModulePtr<FAssetToolsModule>("AssetTools"))
 	{
+		RemoteControlAssetCategoryBit = AssetToolsModule->Get().RegisterAdvancedAssetCategory(FName(TEXT("Remote Control")), LOCTEXT("RemoteControlAssetCategory", "Remote Control"));
+
 		RemoteControlPresetActions = MakeShared<FRemoteControlPresetActions>(FRemoteControlPanelStyle::Get().ToSharedRef());
 		AssetToolsModule->Get().RegisterAssetTypeActions(RemoteControlPresetActions.ToSharedRef());
 	}
diff --git a/Engine/Plugins/VirtualProduction/RemoteControl/Source/RemoteControlUI/Private/RemoteControlUIModule.h b/Engine/Plugins/VirtualProduction/RemoteControl/Source/RemoteControlUI/Private/RemoteControlUIModule.h
--- a/Engine/Plugins/VirtualProduction/RemoteControl/Source/RemoteControlUI/Private/RemoteControlUIModule.h
+++ b/Engine/Plugins/VirtualProduction/RemoteControl/Source/RemoteControlUI/Private/RemoteControlUIModule.h
@@ -3,6 +3,7 @@
 
 #include "IRemoteControlUIModule.h"
 
+#include "AssetTypeCategories.h"
 #include "CoreMinimal.h"	
 #include "LevelEditor.h"
 #include "Modules/ModuleManager.h"
@@ -19,6 +20,11 @@ class URemoteControlPreset;
 class FRemoteControlUIModule : public IRemoteControlUIModule
 {
 public:
+	
+	FRemoteControlUIModule()
+		: RemoteControlAssetCategoryBit(EAssetTypeCategories::Misc)
+	{
+	}
 
 	/**
 	 * Singleton-like access to this module's interface.  This is just for convenience!
@@ -44,6 +50,7 @@ public:
 	virtual void RegisterMetadataCustomization(FName MetadataKey, FOnCustomizeMetadataEntry OnCustomizeCallback) override;
 	virtual void UnregisterMetadataCustomization(FName MetadataKey) override;
 	virtual URemoteControlPreset* GetActivePreset() const override;
+	virtual uint32 GetRemoteControlAssetCategory() const override;
 	virtual void RegisterWidgetFactoryForType(UScriptStruct* RemoteControlEntityType, const FOnGenerateRCWidget& OnGenerateRCWidgetDelegate) override;
 	virtual void UnregisterWidgetFactoryForType(UScriptStruct* RemoteControlEntityType) override;
 
@@ -200,4 +207,7 @@ private:
 	TMap<FName, FOnCustomizeMetadataEntry> ExternalEntityMetadataCustomizations;
 
 	TMap<TWeakObjectPtr<UScriptStruct>, FOnGenerateRCWidget> GenerateWidgetDelegates;
+
+	/** Holds the advanced asset type category bit registered by Remote Control. */
+	EAssetTypeCategories::Type RemoteControlAssetCategoryBit;
 };
diff --git a/Engine/Plugins/VirtualProduction/RemoteControl/Source/RemoteControlUI/Private/UI/Panels/SRCDockPanel.cpp b/Engine/Plugins/VirtualProduction/RemoteControl/Source/RemoteControlUI/Private/UI/Panels/SRCDockPanel.cpp
--- a/Engine/Plugins/VirtualProduction/RemoteControl/Source/RemoteControlUI/Private/UI/Panels/SRCDockPanel.cpp
+++ b/Engine/Plugins/VirtualProduction/RemoteControl/Source/RemoteControlUI/Private/UI/Panels/SRCDockPanel.cpp
@@ -4,10 +4,10 @@
 
 #include "SlateOptMacros.h"
 #include "Styling/RemoteControlStyles.h"
+#include "Styling/StyleColors.h"
 
 #include "UI/RemoteControlPanelStyle.h"
 
-#include "Widgets/Layout/SBorder.h"
 #include "Widgets/Layout/SBox.h"
 #include "Widgets/Layout/SSpacer.h"
 #include "Widgets/Layout/SSplitter.h"
@@ -16,6 +16,33 @@
 #include "Widgets/SWidget.h"
 #include "Widgets/Text/STextBlock.h"
 
+SLATE_IMPLEMENT_WIDGET(SRCDockPanel)
+
+void SRCDockPanel::PrivateRegisterAttributes(FSlateAttributeInitializer& AttributeInitializer)
+{
+}
+
+void SRCDockPanel::Construct(const FArguments& InArgs)
+{
+	RCPanelStyle = &FRemoteControlPanelStyle::Get()->GetWidgetStyle<FRCPanelStyle>("RemoteControlPanel.MinorPanel");
+
+	bIsFooterEnabled = false;
+	bIsHeaderEnabled = false;
+
+	Clipping = EWidgetClipping::ClipToBounds;
+
+	SBorder::FArguments SuperArgs = SBorder::FArguments();
+
+	SuperArgs.BorderImage(&RCPanelStyle->ContentAreaBrush);
+	SuperArgs.Padding(RCPanelStyle->PanelPadding);
+
+	SBorder::Construct(SuperArgs
+		[
+			InArgs._Content.Widget
+		]
+	);
+}
+
 SLATE_IMPLEMENT_WIDGET(SRCMajorPanel)
 
 void SRCMajorPanel::PrivateRegisterAttributes(FSlateAttributeInitializer& AttributeInitializer)
@@ -26,30 +53,23 @@ BEGIN_SLATE_FUNCTION_BUILD_OPTIMIZATION
 
 void SRCMajorPanel::Construct(const SRCMajorPanel::FArguments& InArgs)
 {
-	RCPanelStyle = &FRemoteControlPanelStyle::Get()->GetWidgetStyle<FRCPanelStyle>("RemoteControlPanel.MinorPanel");
-
-	bIsFooterEnabled = InArgs._EnableFooter;
-	bIsHeaderEnabled = InArgs._EnableHeader;
-
-	ChildSlot
-	.Padding(RCPanelStyle->PanelPadding)
-	[
-		SNew(SBorder)
-		.BorderImage(&RCPanelStyle->ContentAreaBrush)
-		.Padding(RCPanelStyle->PanelPadding)
-		[	
+	SRCDockPanel::Construct(SRCDockPanel::FArguments()
+		[
 			SAssignNew(ContentPanel, SSplitter)
 			.Orientation(InArgs._Orientation)
 
 			// Content Panel
-			+SSplitter::Slot()
+			+ SSplitter::Slot()
 			.Value(1.f)
 			[
 				SAssignNew(Children, SSplitter)
 				.Orientation(InArgs._ChildOrientation)
 			]
 		]
-	];
+	);
+
+	bIsFooterEnabled = InArgs._EnableFooter;
+	bIsHeaderEnabled = InArgs._EnableHeader;
 
 	// Enable Header upon explicit request.
 	if (bIsHeaderEnabled.Get() && ContentPanel.IsValid())
@@ -177,29 +197,22 @@ BEGIN_SLATE_FUNCTION_BUILD_OPTIMIZATION
 
 void SRCMinorPanel::Construct(const SRCMinorPanel::FArguments& InArgs)
 {
-	RCPanelStyle = &FRemoteControlPanelStyle::Get()->GetWidgetStyle<FRCPanelStyle>("RemoteControlPanel.MinorPanel");
-
-	bIsFooterEnabled = InArgs._EnableFooter;
-	bIsHeaderEnabled = InArgs._EnableHeader;
-
-	ChildSlot
-	.Padding(RCPanelStyle->PanelPadding)
-	[
-		SNew(SBorder)
-		.BorderImage(&RCPanelStyle->ContentAreaBrush)
-		.Padding(RCPanelStyle->PanelPadding)
-		[	
+	SRCDockPanel::Construct(SRCDockPanel::FArguments()
+		[
 			SAssignNew(ContentPanel, SSplitter)
 			.Orientation(InArgs._Orientation)
 
 			// Content Panel
-			+SSplitter::Slot()
+			+ SSplitter::Slot()
 			.Value(1.f)
 			[
 				InArgs._Content.Widget
 			]
 		]
-	];
+	);
+
+	bIsFooterEnabled = InArgs._EnableFooter;
+	bIsHeaderEnabled = InArgs._EnableHeader;
 
 	// Enable Header upon explicit request.
 	if (bIsHeaderEnabled.Get() && ContentPanel.IsValid())
diff --git a/Engine/Plugins/VirtualProduction/RemoteControl/Source/RemoteControlUI/Private/UI/Panels/SRCDockPanel.h b/Engine/Plugins/VirtualProduction/RemoteControl/Source/RemoteControlUI/Private/UI/Panels/SRCDockPanel.h
--- a/Engine/Plugins/VirtualProduction/RemoteControl/Source/RemoteControlUI/Private/UI/Panels/SRCDockPanel.h
+++ b/Engine/Plugins/VirtualProduction/RemoteControl/Source/RemoteControlUI/Private/UI/Panels/SRCDockPanel.h
@@ -3,7 +3,7 @@
 #pragma once
 
 #include "CoreMinimal.h"
-#include "Widgets/SCompoundWidget.h"
+#include "Widgets/Layout/SBorder.h"
 
 struct FRCPanelStyle;
 class SHorizontalBox;
@@ -19,26 +19,62 @@ enum EToolbar
 	Right
 };
 
+/**
+ * A custom slate widget representing a dockable panel on RC Panel.
+ */
+class SRCDockPanel : public SBorder
+{
+	SLATE_DECLARE_WIDGET(SRCDockPanel, SCompoundWidget)
+
+public:
+
+	SLATE_BEGIN_ARGS(SRCDockPanel)
+		: _Content()
+	{}
+
+		SLATE_DEFAULT_SLOT(FArguments, Content)
+
+	SLATE_END_ARGS()
+
+	/** Constructs this widget with InArgs */
+	void Construct(const FArguments& InArgs);
+
+protected:
+
+	/** Holds whether the footer panel is enabled or not. */
+	TAttribute<bool> bIsFooterEnabled;
+
+	/** Holds whether the header panel is enabled or not. */
+	TAttribute<bool> bIsHeaderEnabled;
+
+	/** Holds the entire content panel. */
+	TSharedPtr<SSplitter> ContentPanel;
+
+	/** Panel Style reference. */
+	const FRCPanelStyle* RCPanelStyle;
+};
+
 /**
  * A custom slate widget representing a major panel on RC Panel.
  */
-class SRCMajorPanel : public SCompoundWidget
+class SRCMajorPanel : public SRCDockPanel
 {
-	SLATE_DECLARE_WIDGET(SRCMajorPanel, SCompoundWidget)
+	SLATE_DECLARE_WIDGET(SRCMajorPanel, SRCDockPanel)
 
 public:
 
 	SLATE_BEGIN_ARGS(SRCMajorPanel)
-		: _EnableFooter(true)
+		: _EnableFooter(false)
 		, _EnableHeader(true)
 		, _HeaderLabel(FText::GetEmpty())
 		, _Orientation(Orient_Vertical)
 		, _ChildOrientation(Orient_Horizontal)
 	{}
+
 		SLATE_ARGUMENT(bool, EnableFooter)
 
 		SLATE_ARGUMENT(bool, EnableHeader)
-			
+
 		SLATE_ARGUMENT(FText, HeaderLabel)
 
 		SLATE_ARGUMENT(EOrientation, Orientation)
@@ -72,34 +108,22 @@ public:
 	
 private:
 
-	/** Holds whether the footer panel is enabled or not. */
-	TAttribute<bool> bIsFooterEnabled;
-	
-	/** Holds whether the header panel is enabled or not. */
-	TAttribute<bool> bIsHeaderEnabled;
-
 	/** Actual toolbar widget located left to the footer. */
 	TSharedPtr<SHorizontalBox> LeftToolbar;
 
 	/** Actual toolbar widget located right to the footer. */
 	TSharedPtr<SHorizontalBox> RightToolbar;
 
-	/** Holds the entire content panel. */
-	TSharedPtr<SSplitter> ContentPanel;
-	
 	/** Holds the entire child panels. */
 	TSharedPtr<SSplitter> Children;
-
-	/** Panel Style reference. */
-	const FRCPanelStyle* RCPanelStyle;
 };
 
 /**
  * A custom slate widget representing a minor panel on RC Panel.
  */
-class SRCMinorPanel : public SCompoundWidget
+class SRCMinorPanel : public SRCDockPanel
 {
-	SLATE_DECLARE_WIDGET(SRCMinorPanel, SCompoundWidget)
+	SLATE_DECLARE_WIDGET(SRCMinorPanel, SRCDockPanel)
 
 public:
 
@@ -150,12 +174,6 @@ public:
 	
 private:
 
-	/** Holds whether the footer panel is enabled or not. */
-	TAttribute<bool> bIsFooterEnabled;
-	
-	/** Holds whether the header panel is enabled or not. */
-	TAttribute<bool> bIsHeaderEnabled;
-
 	/** Actual toolbar widget located left to the footer. */
 	TSharedPtr<SHorizontalBox> LeftFooterToolbar;
 
@@ -167,10 +185,4 @@ private:
 
 	/** Actual toolbar widget located right to the header. */
 	TSharedPtr<SHorizontalBox> RightHeaderToolbar;
-
-	/** Holds the entire content panel. */
-	TSharedPtr<SSplitter> ContentPanel;
-
-	/** Panel Style reference. */
-	const FRCPanelStyle* RCPanelStyle;
 };
diff --git a/Engine/Plugins/VirtualProduction/RemoteControl/Source/RemoteControlUI/Public/IRemoteControlUIModule.h b/Engine/Plugins/VirtualProduction/RemoteControl/Source/RemoteControlUI/Public/IRemoteControlUIModule.h
--- a/Engine/Plugins/VirtualProduction/RemoteControl/Source/RemoteControlUI/Public/IRemoteControlUIModule.h
+++ b/Engine/Plugins/VirtualProduction/RemoteControl/Source/RemoteControlUI/Public/IRemoteControlUIModule.h
@@ -176,6 +176,11 @@ public:
 	 */
 	virtual URemoteControlPreset* GetActivePreset() const = 0;
 
+	/**
+	 * Retrieves the advanced asset category type registered by Remote Control Module.
+	 */
+	virtual uint32 GetRemoteControlAssetCategory() const = 0;
+
 	/**
 	 * Register a widget factory to handle creating a widget in the control panel.
 	 */
