[UE-157410] Enable Logic UI by default (remove console variable)

No longer required as the latest UI revamp has a formal location for Logic UI now (at the top-right corner of the screen as a "Logic" header button)

diff --git a/Engine/Plugins/VirtualProduction/RemoteControl/Source/RemoteControlUI/Private/UI/SRemoteControlPanel.cpp b/Engine/Plugins/VirtualProduction/RemoteControl/Source/RemoteControlUI/Private/UI/SRemoteControlPanel.cpp
--- a/Engine/Plugins/VirtualProduction/RemoteControl/Source/RemoteControlUI/Private/UI/SRemoteControlPanel.cpp
+++ b/Engine/Plugins/VirtualProduction/RemoteControl/Source/RemoteControlUI/Private/UI/SRemoteControlPanel.cpp
@@ -92,8 +92,6 @@ TSharedRef<SBox> SRemoteControlPanel::NoneSelectedWidget = SNew(SBox)
 				.Justification(ETextJustify::Center)
 			];
 
-static TAutoConsoleVariable<bool> CVarRemoteControlEnableLogicUI(TEXT("RemoteControl.EnableLogicUI"), false, TEXT("Toggles visibility of Remote Control Logic User Interface (part of the Remote Control Preset asset)"));
-
 namespace RemoteControlPanelUtils
 {
 	bool IsExposableActor(AActor* Actor)
