[RemoteControl] Reset static shared pointers to fix crash on shutdown

diff --git a/Engine/Plugins/VirtualProduction/RemoteControl/Source/RemoteControlUI/Private/RemoteControlUIModule.cpp b/Engine/Plugins/VirtualProduction/RemoteControl/Source/RemoteControlUI/Private/RemoteControlUIModule.cpp
--- a/Engine/Plugins/VirtualProduction/RemoteControl/Source/RemoteControlUI/Private/RemoteControlUIModule.cpp
+++ b/Engine/Plugins/VirtualProduction/RemoteControl/Source/RemoteControlUI/Private/RemoteControlUIModule.cpp
@@ -25,6 +25,8 @@
 #include "ScopedTransaction.h"
 #include "Styling/AppStyle.h"
 #include "Textures/SlateIcon.h"
+#include "UI/Action/SRCActionPanel.h"
+#include "UI/Behaviour/SRCBehaviourPanel.h"
 #include "UI/Customizations/FPassphraseCustomization.h"
 #include "UI/Customizations/RemoteControlEntityCustomization.h"
 #include "UI/SRCPanelExposedField.h"
@@ -268,6 +270,9 @@ void FRemoteControlUIModule::ShutdownModule()
 	UnregisterAssetTools();
 	UnbindRemoteControlCommands();
 	FRemoteControlPanelStyle::Shutdown();
+	SRemoteControlPanel::Shutdown();
+	SRCActionPanel::Shutdown();
+	SRCBehaviourPanel::Shutdown();
 }
 
 FDelegateHandle FRemoteControlUIModule::AddPropertyFilter(FOnDisplayExposeIcon OnDisplayExposeIcon)
diff --git a/Engine/Plugins/VirtualProduction/RemoteControl/Source/RemoteControlUI/Private/UI/Action/SRCActionPanel.cpp b/Engine/Plugins/VirtualProduction/RemoteControl/Source/RemoteControlUI/Private/UI/Action/SRCActionPanel.cpp
--- a/Engine/Plugins/VirtualProduction/RemoteControl/Source/RemoteControlUI/Private/UI/Action/SRCActionPanel.cpp
+++ b/Engine/Plugins/VirtualProduction/RemoteControl/Source/RemoteControlUI/Private/UI/Action/SRCActionPanel.cpp
@@ -35,7 +35,7 @@
 
 #define LOCTEXT_NAMESPACE "SRCActionPanel"
 
-TSharedRef<SBox> SRCActionPanel::NoneSelectedWidget = SNew(SBox)
+TSharedPtr<SBox> SRCActionPanel::NoneSelectedWidget = SNew(SBox)
 			.Padding(0.f)
 			.VAlign(VAlign_Center)
 			.HAlign(HAlign_Center)
@@ -69,6 +69,11 @@ void SRCActionPanel::Construct(const FArguments& InArgs, const TSharedRef<SRemot
 
 END_SLATE_FUNCTION_BUILD_OPTIMIZATION
 
+void SRCActionPanel::Shutdown()
+{
+	NoneSelectedWidget.Reset();
+}
+
 void SRCActionPanel::OnBehaviourSelectionChanged(TSharedPtr<FRCBehaviourModel> InBehaviourItem)
 {
 	SelectedBehaviourItemWeakPtr = InBehaviourItem;
@@ -210,7 +215,7 @@ void SRCActionPanel::UpdateWrappedWidget(TSharedPtr<FRCBehaviourModel> InBehavio
 	}
 	else
 	{
-		WrappedBoxWidget->SetContent(NoneSelectedWidget);
+		WrappedBoxWidget->SetContent(NoneSelectedWidget.ToSharedRef());
 	}
 }
 
diff --git a/Engine/Plugins/VirtualProduction/RemoteControl/Source/RemoteControlUI/Private/UI/Action/SRCActionPanel.h b/Engine/Plugins/VirtualProduction/RemoteControl/Source/RemoteControlUI/Private/UI/Action/SRCActionPanel.h
--- a/Engine/Plugins/VirtualProduction/RemoteControl/Source/RemoteControlUI/Private/UI/Action/SRCActionPanel.h
+++ b/Engine/Plugins/VirtualProduction/RemoteControl/Source/RemoteControlUI/Private/UI/Action/SRCActionPanel.h
@@ -27,6 +27,9 @@ public:
 	/** Constructs this widget with InArgs */
 	void Construct(const FArguments& InArgs, const TSharedRef<SRemoteControlPanel>& InPanel);
 
+	/** Shutdown panel */
+	static void Shutdown();
+
 	/** Whether the Actions list widget currently has focus.*/
 	bool IsListFocused() const;
 
@@ -82,7 +85,7 @@ private:
 	TSharedPtr<class SRCActionPanelList> ActionPanelList;
 
 	/** Helper widget for behavior details. */
-	static TSharedRef<SBox> NoneSelectedWidget;
+	static TSharedPtr<SBox> NoneSelectedWidget;
 
 	/** Panel Style reference. */
 	const FRCPanelStyle* RCPanelStyle;
diff --git a/Engine/Plugins/VirtualProduction/RemoteControl/Source/RemoteControlUI/Private/UI/Behaviour/SRCBehaviourPanel.cpp b/Engine/Plugins/VirtualProduction/RemoteControl/Source/RemoteControlUI/Private/UI/Behaviour/SRCBehaviourPanel.cpp
--- a/Engine/Plugins/VirtualProduction/RemoteControl/Source/RemoteControlUI/Private/UI/Behaviour/SRCBehaviourPanel.cpp
+++ b/Engine/Plugins/VirtualProduction/RemoteControl/Source/RemoteControlUI/Private/UI/Behaviour/SRCBehaviourPanel.cpp
@@ -23,7 +23,7 @@
 
 #define LOCTEXT_NAMESPACE "SRCBehaviourPanel"
 
-TSharedRef<SBox> SRCBehaviourPanel::NoneSelectedWidget = SNew(SBox)
+TSharedPtr<SBox> SRCBehaviourPanel::NoneSelectedWidget = SNew(SBox)
 			.Padding(0.f)
 			.VAlign(VAlign_Center)
 			.HAlign(HAlign_Center)
@@ -57,6 +57,11 @@ void SRCBehaviourPanel::Construct(const FArguments& InArgs, const TSharedRef<SRe
 
 END_SLATE_FUNCTION_BUILD_OPTIMIZATION
 
+void SRCBehaviourPanel::Shutdown()
+{
+	NoneSelectedWidget.Reset();
+}
+
 void SRCBehaviourPanel::OnControllerSelectionChanged(TSharedPtr<FRCControllerModel> InControllerItem)
 {
 	SelectedControllerItemWeakPtr = InControllerItem;
@@ -128,7 +133,7 @@ void SRCBehaviourPanel::UpdateWrappedWidget(TSharedPtr<FRCControllerModel> InCon
 	}
 	else
 	{
-		WrappedBoxWidget->SetContent(NoneSelectedWidget);
+		WrappedBoxWidget->SetContent(NoneSelectedWidget.ToSharedRef());
 	}
 }
 
diff --git a/Engine/Plugins/VirtualProduction/RemoteControl/Source/RemoteControlUI/Private/UI/Behaviour/SRCBehaviourPanel.h b/Engine/Plugins/VirtualProduction/RemoteControl/Source/RemoteControlUI/Private/UI/Behaviour/SRCBehaviourPanel.h
--- a/Engine/Plugins/VirtualProduction/RemoteControl/Source/RemoteControlUI/Private/UI/Behaviour/SRCBehaviourPanel.h
+++ b/Engine/Plugins/VirtualProduction/RemoteControl/Source/RemoteControlUI/Private/UI/Behaviour/SRCBehaviourPanel.h
@@ -27,6 +27,9 @@ public:
 	/** Constructs this widget with InArgs */
 	void Construct(const FArguments& InArgs, const TSharedRef<SRemoteControlPanel>& InPanel);
 
+	/** Shutdown panel */
+	static void Shutdown();
+
 	/** Whether the Behaviour list widget currently has focus. Used for Delete Item UI command */
 	bool IsListFocused() const;
 
@@ -70,7 +73,7 @@ private:
 	TSharedPtr<class SRCBehaviourPanelList> BehaviourPanelList;
 
 	/** Helper widget for behavior details. */
-	static TSharedRef<SBox> NoneSelectedWidget;
+	static TSharedPtr<SBox> NoneSelectedWidget;
 
 	/** Panel Style reference. */
 	const FRCPanelStyle* RCPanelStyle;
diff --git a/Engine/Plugins/VirtualProduction/RemoteControl/Source/RemoteControlUI/Private/UI/SRemoteControlPanel.cpp b/Engine/Plugins/VirtualProduction/RemoteControl/Source/RemoteControlUI/Private/UI/SRemoteControlPanel.cpp
--- a/Engine/Plugins/VirtualProduction/RemoteControl/Source/RemoteControlUI/Private/UI/SRemoteControlPanel.cpp
+++ b/Engine/Plugins/VirtualProduction/RemoteControl/Source/RemoteControlUI/Private/UI/SRemoteControlPanel.cpp
@@ -81,7 +81,7 @@ const FName SRemoteControlPanel::DefaultRemoteControlPanelToolBarName("RemoteCon
 const FName SRemoteControlPanel::AuxiliaryRemoteControlPanelToolBarName("RemoteControlPanel.AuxiliaryToolBar");
 const float SRemoteControlPanel::MinimumPanelWidth = 640.f;
 
-TSharedRef<SBox> SRemoteControlPanel::NoneSelectedWidget = SNew(SBox)
+TSharedPtr<SBox> SRemoteControlPanel::NoneSelectedWidget = SNew(SBox)
 			.Padding(0.f)
 			.HAlign(HAlign_Center)
 			.VAlign(VAlign_Center)
@@ -753,6 +753,11 @@ SRemoteControlPanel::~SRemoteControlPanel()
 	ProtocolWidgetsModule.ResetProtocolBindingList();	
 }
 
+void SRemoteControlPanel::Shutdown()
+{
+	NoneSelectedWidget.Reset();
+}
+
 void SRemoteControlPanel::PostUndo(bool bSuccess)
 {
 	Refresh();
@@ -1623,7 +1628,7 @@ TSharedRef<SWidget> SRemoteControlPanel::CreateEntityDetailsView()
 	}
 	else
 	{
-		WrappedEntityDetailsView->SetContent(NoneSelectedWidget);
+		WrappedEntityDetailsView->SetContent(NoneSelectedWidget.ToSharedRef());
 	}
 
 	return EntityDetailsDockPanel;
@@ -1658,7 +1663,7 @@ void SRemoteControlPanel::UpdateEntityDetailsView(const TSharedPtr<SRCPanelTreeN
 	}
 	else
 	{
-		WrappedEntityDetailsView->SetContent(NoneSelectedWidget);
+		WrappedEntityDetailsView->SetContent(NoneSelectedWidget.ToSharedRef());
 	}
 
 	static const FName ProtocolWidgetsModuleName = "RemoteControlProtocolWidgets";	
@@ -1673,13 +1678,13 @@ void SRemoteControlPanel::UpdateEntityDetailsView(const TSharedPtr<SRCPanelTreeN
 			}
 			else
 			{
-				EntityProtocolDetails->SetContent(NoneSelectedWidget);
+				EntityProtocolDetails->SetContent(NoneSelectedWidget.ToSharedRef());
 			}
 		}
 	}
 	else
 	{
-		EntityProtocolDetails->SetContent(NoneSelectedWidget);
+		EntityProtocolDetails->SetContent(NoneSelectedWidget.ToSharedRef());
 	}
 
 	// Trigger search to list the search results specific to selected group.
diff --git a/Engine/Plugins/VirtualProduction/RemoteControl/Source/RemoteControlUI/Private/UI/SRemoteControlPanel.h b/Engine/Plugins/VirtualProduction/RemoteControl/Source/RemoteControlUI/Private/UI/SRemoteControlPanel.h
--- a/Engine/Plugins/VirtualProduction/RemoteControl/Source/RemoteControlUI/Private/UI/SRemoteControlPanel.h
+++ b/Engine/Plugins/VirtualProduction/RemoteControl/Source/RemoteControlUI/Private/UI/SRemoteControlPanel.h
@@ -74,6 +74,7 @@ public:
 	
 	void Construct(const FArguments& InArgs, URemoteControlPreset* InPreset, TSharedPtr<IToolkitHost> InToolkitHost);
 	~SRemoteControlPanel();
+	static void Shutdown();
 
 	//~ FEditorUndoClient interface
 	virtual void PostUndo(bool bSuccess) override;
@@ -335,7 +336,7 @@ private:
 	/** Wrapper widget for entity details view. */
 	TSharedPtr<SBorder> WrappedEntityDetailsView;
 	/** Helper widget for entity details view and protocol details view. */
-	static TSharedRef<SBox> NoneSelectedWidget;
+	static TSharedPtr<SBox> NoneSelectedWidget;
 	/** Holds the field's protocol details. */
 	TSharedPtr<SBox> EntityProtocolDetails;
 	/** Whether to show the rebind all button. */
