CIS Fix for Build Engine Localization In Preview For Source Code - Remote Control

diff --git a/Engine/Plugins/VirtualProduction/RemoteControl/Source/RemoteControlUI/Private/UI/SRCPanelExposedEntitiesList.cpp b/Engine/Plugins/VirtualProduction/RemoteControl/Source/RemoteControlUI/Private/UI/SRCPanelExposedEntitiesList.cpp
--- a/Engine/Plugins/VirtualProduction/RemoteControl/Source/RemoteControlUI/Private/UI/SRCPanelExposedEntitiesList.cpp
+++ b/Engine/Plugins/VirtualProduction/RemoteControl/Source/RemoteControlUI/Private/UI/SRCPanelExposedEntitiesList.cpp
@@ -174,7 +174,7 @@ void SRCPanelExposedEntitiesList::Construct(const FArguments& InArgs, URemoteCon
 
 	// Group Dock Panel
 	TSharedPtr<SRCMinorPanel> GroupDockPanel = SNew(SRCMinorPanel)
-		.HeaderLabel(LOCTEXT("GroupsLabel", "Group"))
+		.HeaderLabel(LOCTEXT("GroupLabel", "Group"))
 		.EnableFooter(false)
 		[
 			GroupsListView.ToSharedRef()
diff --git a/Engine/Plugins/VirtualProduction/RemoteControl/Source/RemoteControlUI/Private/UI/SRemoteControlPanel.cpp b/Engine/Plugins/VirtualProduction/RemoteControl/Source/RemoteControlUI/Private/UI/SRemoteControlPanel.cpp
--- a/Engine/Plugins/VirtualProduction/RemoteControl/Source/RemoteControlUI/Private/UI/SRemoteControlPanel.cpp
+++ b/Engine/Plugins/VirtualProduction/RemoteControl/Source/RemoteControlUI/Private/UI/SRemoteControlPanel.cpp
@@ -1377,11 +1377,12 @@ void SRemoteControlPanel::RegisterPanels()
 
 		{// Properties With Protocols Panel
 			TSharedRef<FRCPanelDrawerArgs> EntityProtocolsPanel = MakeShared<FRCPanelDrawerArgs>(ERCPanels::RCP_Protocols);
+			const FRemoteControlCommands& Commands = FRemoteControlCommands::Get();
 
 			EntityProtocolsPanel->bDrawnByDefault = true;
 			EntityProtocolsPanel->bRotateIconBy90 = false;
 			EntityProtocolsPanel->DrawerVisibility = bIsInProtocolsMode ? EVisibility::Visible : EVisibility::Collapsed;
-			EntityProtocolsPanel->Label = LOCTEXT("EntityProtocolsPanelLabel", "Protocols");
+			EntityProtocolsPanel->Label = Commands.ToggleProtocolMappings->GetLabel();
 			EntityProtocolsPanel->ToolTip = LOCTEXT("EntityProtocolsPanelTooltip", "Open entity protocols panel.");
 			EntityProtocolsPanel->Icon = FSlateIcon(FAppStyle::GetAppStyleSetName(), "LevelEditor.Tabs.StatsViewer");
 
@@ -1874,12 +1875,14 @@ void SRemoteControlPanel::RegisterAuxiliaryToolBar()
 		{
 			FToolMenuSection& ToolsSection = ToolbarBuilder->AddSection("Tools");
 
+			const FRemoteControlCommands& Commands = FRemoteControlCommands::Get();
+
 			ToolsSection.AddEntry(FToolMenuEntry::InitWidget("Protocols"
 				, SNew(SAutoResizeButton)
 					.UICommand(FRemoteControlCommands::Get().ToggleProtocolMappings)
 					.ForceSmallIcons_Static(SRemoteControlPanel::ShouldForceSmallIcons)
 					.IconOverride(FSlateIcon(FAppStyle::GetAppStyleSetName(), "LevelEditor.Tabs.StatsViewer"))
-				, LOCTEXT("ProtocolsLabel", "Protocols")
+				, Commands.ToggleProtocolMappings->GetLabel()
 			)
 			);
 			
@@ -1888,7 +1891,7 @@ void SRemoteControlPanel::RegisterAuxiliaryToolBar()
 					.UICommand(FRemoteControlCommands::Get().ToggleLogicEditor)
 					.ForceSmallIcons_Static(SRemoteControlPanel::ShouldForceSmallIcons)
 					.IconOverride(FSlateIcon(FAppStyle::GetAppStyleSetName(), TEXT("GraphEditor.StateMachine_24x")))
-				, LOCTEXT("LogicLabel", "Logic")
+				, Commands.ToggleLogicEditor->GetLabel()
 			)
 			);
 		}
