Changing the transient allowed list in RC to support Avalanche.

Also renamed the GetPresetWorld to GetWorld and GetEmbeddedPresetWorld to GetEmbeddedWorld.

Fixed bindings appearing in the wrong panel when using embeded presets.

diff --git a/Engine/Plugins/VirtualProduction/RemoteControl/Source/RemoteControl/Private/RemoteControlBinding.cpp b/Engine/Plugins/VirtualProduction/RemoteControl/Source/RemoteControl/Private/RemoteControlBinding.cpp
--- a/Engine/Plugins/VirtualProduction/RemoteControl/Source/RemoteControl/Private/RemoteControlBinding.cpp
+++ b/Engine/Plugins/VirtualProduction/RemoteControl/Source/RemoteControl/Private/RemoteControlBinding.cpp
@@ -375,7 +375,8 @@ UWorld* URemoteControlLevelDependantBinding::GetCurrentWorld() const
 	URemoteControlPreset* Preset = Cast<URemoteControlPreset>(GetOuter());
 
 	// Since this is used to retrieve the binding in the map, we never use the PIE world in editor.
-	return URemoteControlPreset::GetPresetWorld(Preset, false);
+	constexpr bool bAllowPIE = false;
+	return URemoteControlPreset::GetWorld(Preset, bAllowPIE);
 }
 
 void URemoteControlLevelDependantBinding::SetBoundObject(const TSoftObjectPtr<ULevel>& Level, const TSoftObjectPtr<UObject>& BoundObject)
diff --git a/Engine/Plugins/VirtualProduction/RemoteControl/Source/RemoteControl/Private/RemoteControlPreset.cpp b/Engine/Plugins/VirtualProduction/RemoteControl/Source/RemoteControl/Private/RemoteControlPreset.cpp
--- a/Engine/Plugins/VirtualProduction/RemoteControl/Source/RemoteControl/Private/RemoteControlPreset.cpp
+++ b/Engine/Plugins/VirtualProduction/RemoteControl/Source/RemoteControl/Private/RemoteControlPreset.cpp
@@ -983,16 +983,15 @@ void URemoteControlPreset::PostLoadProperties()
 	}
 }
 
-UWorld* URemoteControlPreset::GetPresetWorld(const URemoteControlPreset* Preset, bool bAllowPIE /*= false*/)
+UWorld* URemoteControlPreset::GetWorld(const URemoteControlPreset* Preset, bool bAllowPIE)
 {
-	if (Preset)
+	if (Preset && Preset->IsEmbeddedPreset())
 	{
-		if (UObject* Outer = Preset->GetOuter())
+		UWorld* EmbeddedPresetWorld = Preset->GetEmbeddedWorld();
+
+		if (EmbeddedPresetWorld)
 		{
-			if (UWorld* OuterWorld = Outer->GetWorld())
-			{
-				return OuterWorld;
-			}
+			return EmbeddedPresetWorld;
 		}
 	}
 
@@ -1029,9 +1028,22 @@ UWorld* URemoteControlPreset::GetPresetWorld(const URemoteControlPreset* Preset,
 	return nullptr;
 }
 
-UWorld* URemoteControlPreset::GetPresetWorld(bool bAllowPIE /*= false*/) const
+UWorld* URemoteControlPreset::GetWorld(bool bAllowPIE) const
+{
+	return URemoteControlPreset::GetWorld(this, bAllowPIE);
+}
+
+UWorld* URemoteControlPreset::GetEmbeddedWorld() const
 {
-	return URemoteControlPreset::GetPresetWorld(this, bAllowPIE);
+	if (UObject* Outer = GetOuter())
+	{
+		if (UWorld* OuterWorld = Outer->GetWorld())
+		{
+			return OuterWorld;
+		}
+	}
+
+	return nullptr;
 }
 
 bool URemoteControlPreset::IsEmbeddedPreset() const
diff --git a/Engine/Plugins/VirtualProduction/RemoteControl/Source/RemoteControl/Private/RemoteControlPresetRebindingManager.cpp b/Engine/Plugins/VirtualProduction/RemoteControl/Source/RemoteControl/Private/RemoteControlPresetRebindingManager.cpp
--- a/Engine/Plugins/VirtualProduction/RemoteControl/Source/RemoteControl/Private/RemoteControlPresetRebindingManager.cpp
+++ b/Engine/Plugins/VirtualProduction/RemoteControl/Source/RemoteControl/Private/RemoteControlPresetRebindingManager.cpp
@@ -235,7 +235,8 @@ void FRemoteControlPresetRebindingManager::Rebind_Legacy(URemoteControlPreset* P
 	// Fetch any relevant object from the level based on if their class is relevant for the set of entities we have to rebind.
 	TArray<UClass*> EntitiesOwnerClasses;
 	EntitiesGroupedBySupportedOwnerClass.GenerateKeyArray(EntitiesOwnerClasses);
-	Context.ObjectsGroupedByRelevantClass = RCPresetRebindingManager::GetLevelObjectsGroupedByClass(EntitiesOwnerClasses, Preset->GetPresetWorld(false));
+	constexpr bool bAllowPIE = false;
+	Context.ObjectsGroupedByRelevantClass = RCPresetRebindingManager::GetLevelObjectsGroupedByClass(EntitiesOwnerClasses, Preset->GetWorld(bAllowPIE));
 
 	for (TPair<UClass*, TArray<TSharedPtr<FRemoteControlEntity>>>& Entry : EntitiesGroupedBySupportedOwnerClass)
 	{
@@ -432,7 +433,8 @@ void FRemoteControlPresetRebindingManager::Rebind_NewAlgo(URemoteControlPreset*
 		{
 			const bool bRebindingComponent = !InvalidBinding->BindingContext.ComponentName.IsNone();
 			TArray<AActor*> ObjectsWithSupportedClass;
-			if (!RCPresetRebindingManager::GetActorsOfClass(Preset->GetPresetWorld(false), OwnerClass, ObjectsWithSupportedClass))
+			constexpr bool bAllowPIE = false;
+			if (!RCPresetRebindingManager::GetActorsOfClass(Preset->GetWorld(bAllowPIE), OwnerClass, ObjectsWithSupportedClass))
 			{
 				continue;
 			}
diff --git a/Engine/Plugins/VirtualProduction/RemoteControl/Source/RemoteControl/Public/RemoteControlPreset.h b/Engine/Plugins/VirtualProduction/RemoteControl/Source/RemoteControl/Public/RemoteControlPreset.h
--- a/Engine/Plugins/VirtualProduction/RemoteControl/Source/RemoteControl/Public/RemoteControlPreset.h
+++ b/Engine/Plugins/VirtualProduction/RemoteControl/Source/RemoteControl/Public/RemoteControlPreset.h
@@ -304,6 +304,8 @@ class REMOTECONTROL_API URemoteControlPreset : public UObject
 {
 public:
 	GENERATED_BODY()
+
+	using UObject::GetWorld;
 	
 	/** Callback for post remote control preset load, called by URemoteControlPreset::PostLoad function */
 	DECLARE_MULTICAST_DELEGATE_OneParam(FOnPostLoadRemoteControlPreset, URemoteControlPreset* /* InPreset */);
@@ -843,8 +845,9 @@ private:
 
 public:
 
-	static UWorld* GetPresetWorld(const URemoteControlPreset* Preset = nullptr, bool bAllowPIE = false);
-	UWorld* GetPresetWorld(bool bAllowPIE = false) const;
+	static UWorld* GetWorld(const URemoteControlPreset* Preset = nullptr, bool bAllowPIE = false);
+	UWorld* GetWorld(bool bAllowPIE = false) const;
+	UWorld* GetEmbeddedWorld() const;
 	
 	/** Returns true if the preset is hosted within another asset. */
 	bool IsEmbeddedPreset() const;
diff --git a/Engine/Plugins/VirtualProduction/RemoteControl/Source/RemoteControlUI/Private/AssetEditor/RemoteControlPresetEditorToolkit.cpp b/Engine/Plugins/VirtualProduction/RemoteControl/Source/RemoteControlUI/Private/AssetEditor/RemoteControlPresetEditorToolkit.cpp
--- a/Engine/Plugins/VirtualProduction/RemoteControl/Source/RemoteControlUI/Private/AssetEditor/RemoteControlPresetEditorToolkit.cpp
+++ b/Engine/Plugins/VirtualProduction/RemoteControl/Source/RemoteControlUI/Private/AssetEditor/RemoteControlPresetEditorToolkit.cpp
@@ -35,9 +35,9 @@ void FRemoteControlPresetEditorToolkit::InitRemoteControlPresetEditor(const EToo
 	// tab ID will already be registered within EditorTabManager
 	TSharedPtr<FTabManager> HostTabManager;
 
-	if (ToolkitHost.IsValid())
+	if (InitToolkitHost.IsValid())
 	{
-		HostTabManager = ToolkitHost.Pin()->GetTabManager();
+		HostTabManager = InitToolkitHost->GetTabManager();
 	}
 	else
 	{
diff --git a/Engine/Plugins/VirtualProduction/RemoteControl/Source/RemoteControlUI/Private/RemoteControlUIModule.cpp b/Engine/Plugins/VirtualProduction/RemoteControl/Source/RemoteControlUI/Private/RemoteControlUIModule.cpp
--- a/Engine/Plugins/VirtualProduction/RemoteControl/Source/RemoteControlUI/Private/RemoteControlUIModule.cpp
+++ b/Engine/Plugins/VirtualProduction/RemoteControl/Source/RemoteControlUI/Private/RemoteControlUIModule.cpp
@@ -11,6 +11,7 @@
 #include "Engine/SkeletalMesh.h"
 #include "Framework/MultiBox/MultiBoxBuilder.h"
 #include "HAL/PlatformApplicationMisc.h"
+#include "IRemoteControlModule.h"
 #include "Interfaces/IMainFrameModule.h"
 #include "Kismet2/ComponentEditorUtils.h"
 #include "MaterialEditor/DEditorParameterValue.h"
@@ -227,18 +228,18 @@ namespace RemoteControlUIModule
 			return true;
 		}
 
-		/*
-		 * Embedded presets may be part of a transient package.
-		 * If the object's world's outer is transient, it's probably hosted. Allow it.
-		 */
+		 // Embedded preset worlds may be part of a transient package.
 		if (UWorld* ObjectWorld = Object->GetWorld())
 		{
-			if (UObject* ObjectWorldOuter = ObjectWorld->GetOuter())
+			IRemoteControlModule& RCModule = IRemoteControlModule::Get();
+			TArray<TWeakObjectPtr<URemoteControlPreset>> EmbeddedPresets;
+			RCModule.GetEmbeddedPresets(EmbeddedPresets);
+
+			for (TWeakObjectPtr<URemoteControlPreset> RCPreset : EmbeddedPresets)
 			{
-				if (ObjectWorldOuter->HasAnyFlags(RF_Transient))
+				if (RCPreset.IsValid() && RCPreset->GetEmbeddedWorld() == ObjectWorld)
 				{
-					// Check if it's also a remote control preset so that we don't accidentally include some other objects.
-					return  Object->IsA<URemoteControlPreset>();
+					return true;
 				}
 			}
 		}
@@ -322,6 +323,8 @@ TSharedRef<SRemoteControlPanel> FRemoteControlUIModule::CreateRemoteControlPanel
 				}
 			});
 
+	RegisteredRemoteControlPanels.Add(TWeakPtr<SRemoteControlPanel>(PanelRef));
+
 	// NOTE : Reregister the module with the detail panel when this panel is created.
 
 	if (!SharedDetailsPanel.IsValid())
@@ -399,7 +402,7 @@ void FRemoteControlUIModule::UnregisterEvents()
 
 void FRemoteControlUIModule::ToggleEditMode()
 {
-	if (TSharedPtr<SRemoteControlPanel> Panel = WeakActivePanel.Pin()) // Check whether the panel is active.
+	if (TSharedPtr<SRemoteControlPanel> Panel = GetPanelForObject(nullptr)) // Check whether the panel is active.
 	{
 		Panel->SetEditMode(Panel->IsInEditMode() ? false : true);
 	}
@@ -407,12 +410,12 @@ void FRemoteControlUIModule::ToggleEditMode()
 
 bool FRemoteControlUIModule::CanToggleEditMode() const
 {
-	return WeakActivePanel.IsValid();
+	return GetPanelForObject(nullptr).IsValid();
 }
 
 bool FRemoteControlUIModule::IsInEditMode() const
 {
-	if (TSharedPtr<SRemoteControlPanel> Panel = WeakActivePanel.Pin()) // Check whether the panel is active.
+	if (TSharedPtr<SRemoteControlPanel> Panel = GetPanelForObject(nullptr)) // Check whether the panel is active.
 	{
 		return Panel->IsInEditMode();
 	}
@@ -422,7 +425,7 @@ bool FRemoteControlUIModule::IsInEditMode() const
 
 URemoteControlPreset* FRemoteControlUIModule::GetActivePreset() const
 {
-	if (const TSharedPtr<SRemoteControlPanel> Panel = WeakActivePanel.Pin())
+	if (const TSharedPtr<SRemoteControlPanel> Panel = GetPanelForObject(nullptr))
 	{
 		return Panel->GetPreset();
 	}
@@ -533,11 +536,68 @@ void FRemoteControlUIModule::HandleCreatePropertyRowExtension(const FOnGenerateG
 	WeakDetailsTreeNode = InArgs.OwnerTreeNode;
 }
 
+TSharedPtr<SRemoteControlPanel> FRemoteControlUIModule::GetPanelForObject(const UObject* Object) const
+{
+	if (Object)
+	{
+		if (UWorld* OwnerWorld = Object->GetWorld())
+		{
+			for (const TWeakPtr<SRemoteControlPanel>& RegisteredPanelWeak : RegisteredRemoteControlPanels)
+			{
+				if (RegisteredPanelWeak.IsValid())
+				{
+					TSharedPtr<SRemoteControlPanel> RegisteredPanel = RegisteredPanelWeak.Pin();
+					const URemoteControlPreset* Preset = RegisteredPanel->GetPreset();
+
+					if (Preset && Preset->IsEmbeddedPreset() && Preset->GetEmbeddedWorld() == OwnerWorld)
+					{
+						return RegisteredPanel;
+					}
+				}
+			}
+		}
+	}
+
+	return WeakActivePanel.Pin();
+}
+
+TSharedPtr<SRemoteControlPanel> FRemoteControlUIModule::GetPanelForProperty(const FRCExposesPropertyArgs& InPropertyArgs) const
+{
+	const FRCExposesPropertyArgs::EType ExtensionArgsType = InPropertyArgs.GetType();
+
+	if (ExtensionArgsType == FRCExposesPropertyArgs::EType::E_Handle)
+	{
+		TArray<UObject*> OuterObjects;
+		InPropertyArgs.PropertyHandle->GetOuterObjects(OuterObjects);
+
+		if (OuterObjects.Num() > 0)
+		{
+			return GetPanelForObject(OuterObjects[0]);
+		}		
+	}
+	else if (ExtensionArgsType == FRCExposesPropertyArgs::EType::E_OwnerObject)
+	{
+		return GetPanelForObject(InPropertyArgs.OwnerObject);
+	}
+
+	return GetPanelForObject(nullptr);
+}
+
+TSharedPtr<SRemoteControlPanel> FRemoteControlUIModule::GetPanelForPropertyChangeEvent(const FPropertyChangedEvent& InPropertyChangeEvent) const
+{
+	if (InPropertyChangeEvent.GetNumObjectsBeingEdited() > 0)
+	{
+		return GetPanelForObject(InPropertyChangeEvent.GetObjectBeingEdited(0));
+	}
+
+	return GetPanelForObject(nullptr);
+}
+
 FSlateIcon FRemoteControlUIModule::OnGetExposedIcon(const FRCExposesPropertyArgs& InPropertyArgs) const
 {
 	FName BrushName("NoBrush");
 
-	if (TSharedPtr<SRemoteControlPanel> Panel = WeakActivePanel.Pin())
+	if (TSharedPtr<SRemoteControlPanel> Panel = GetPanelForProperty(InPropertyArgs))
 	{
 		if (Panel->GetPreset())
 		{
@@ -558,7 +618,7 @@ FSlateIcon FRemoteControlUIModule::OnGetExposedIcon(const FRCExposesPropertyArgs
 
 bool FRemoteControlUIModule::CanToggleExposeProperty(const FRCExposesPropertyArgs InPropertyArgs) const
 {
-	if (TSharedPtr<SRemoteControlPanel> Panel = WeakActivePanel.Pin())
+	if (TSharedPtr<SRemoteControlPanel> Panel = GetPanelForProperty(InPropertyArgs))
 	{
 		return InPropertyArgs.IsValid() && ShouldDisplayExposeIcon(InPropertyArgs);
 	}
@@ -569,7 +629,7 @@ bool FRemoteControlUIModule::CanToggleExposeProperty(const FRCExposesPropertyArg
 
 ECheckBoxState FRemoteControlUIModule::GetPropertyExposedCheckState(const FRCExposesPropertyArgs InPropertyArgs) const
 {
-	if (TSharedPtr<SRemoteControlPanel> Panel = WeakActivePanel.Pin())
+	if (TSharedPtr<SRemoteControlPanel> Panel = GetPanelForProperty(InPropertyArgs))
 	{
 		if (Panel->GetPreset())
 		{
@@ -591,7 +651,7 @@ void FRemoteControlUIModule::OnToggleExposeProperty(const FRCExposesPropertyArgs
 		return;
 	}
 
-	if (TSharedPtr<SRemoteControlPanel> Panel = WeakActivePanel.Pin())
+	if (TSharedPtr<SRemoteControlPanel> Panel = GetPanelForProperty(InPropertyArgs))
 	{
 		Panel->ToggleProperty(InPropertyArgs);
 	}
@@ -601,7 +661,7 @@ FRemoteControlUIModule::EPropertyExposeStatus FRemoteControlUIModule::GetPropert
 {
 	if (InPropertyArgs.IsValid())
 	{
-		if (TSharedPtr<SRemoteControlPanel> Panel = WeakActivePanel.Pin())
+		if (TSharedPtr<SRemoteControlPanel> Panel = GetPanelForProperty(InPropertyArgs))
 		{
 			return Panel->IsExposed(InPropertyArgs) ? EPropertyExposeStatus::Exposed : EPropertyExposeStatus::Unexposed;
 		}
@@ -626,7 +686,7 @@ FSlateIcon FRemoteControlUIModule::OnGetOverrideMaterialsIcon(const FRCExposesPr
 
 bool FRemoteControlUIModule::IsStaticOrSkeletalMaterialProperty(const FRCExposesPropertyArgs InPropertyArgs) const
 {
-	if (TSharedPtr<SRemoteControlPanel> Panel = WeakActivePanel.Pin()) // Check whether the panel is active.
+	if (TSharedPtr<SRemoteControlPanel> Panel = GetPanelForProperty(InPropertyArgs)) // Check whether the panel is active.
 	{
 		if (Panel->GetPreset() && InPropertyArgs.IsValid()) // Ensure that we have a valid preset and handle.
 		{
@@ -824,9 +884,9 @@ void FRemoteControlUIModule::UnregisterSettings()
 	}
 }
 
-void FRemoteControlUIModule::OnSettingsModified(UObject*, FPropertyChangedEvent&)
+void FRemoteControlUIModule::OnSettingsModified(UObject*, FPropertyChangedEvent& InPropertyChangeEvent)
 {
-	if (TSharedPtr<SRemoteControlPanel> Panel = WeakActivePanel.Pin())
+	if (TSharedPtr<SRemoteControlPanel> Panel = GetPanelForPropertyChangeEvent(InPropertyChangeEvent))
 	{
 		if (TSharedPtr<SRCPanelExposedEntitiesList> EntityList = Panel->GetEntityList())
 		{
@@ -1039,7 +1099,7 @@ void FRemoteControlUIModule::RefreshPanels()
 		SharedDetailsPanel->ForceRefresh();
 	}
 
-	if (TSharedPtr<SRemoteControlPanel> Panel = WeakActivePanel.Pin())
+	if (TSharedPtr<SRemoteControlPanel> Panel = GetPanelForObject(nullptr))
 	{
 		Panel->Refresh();
 	}
@@ -1060,6 +1120,23 @@ TSharedPtr<SRCPanelTreeNode> FRemoteControlUIModule::GenerateEntityWidget(const
 	return nullptr;
 }
 
+void FRemoteControlUIModule::UnregisterRemoteControlPanel(SRemoteControlPanel* Panel)
+{
+	if (!Panel)
+	{
+		return;
+	}
+
+	for (int32 PanelIndex = 0; PanelIndex < RegisteredRemoteControlPanels.Num(); ++PanelIndex)
+	{
+		if (RegisteredRemoteControlPanels[PanelIndex].HasSameObject(Panel))
+		{
+			RegisteredRemoteControlPanels.RemoveAt(PanelIndex);
+			return;
+		}
+	}
+}
+
 IMPLEMENT_MODULE(FRemoteControlUIModule, RemoteControlUI);
 
 #undef LOCTEXT_NAMESPACE /*RemoteControlUI*/
diff --git a/Engine/Plugins/VirtualProduction/RemoteControl/Source/RemoteControlUI/Private/RemoteControlUIModule.h b/Engine/Plugins/VirtualProduction/RemoteControl/Source/RemoteControlUI/Private/RemoteControlUIModule.h
--- a/Engine/Plugins/VirtualProduction/RemoteControl/Source/RemoteControlUI/Private/RemoteControlUIModule.h
+++ b/Engine/Plugins/VirtualProduction/RemoteControl/Source/RemoteControlUI/Private/RemoteControlUIModule.h
@@ -53,7 +53,6 @@ public:
 	virtual uint32 GetRemoteControlAssetCategory() const override;
 	virtual void RegisterWidgetFactoryForType(UScriptStruct* RemoteControlEntityType, const FOnGenerateRCWidget& OnGenerateRCWidgetDelegate) override;
 	virtual void UnregisterWidgetFactoryForType(UScriptStruct* RemoteControlEntityType) override;
-
 	//~ End IRemoteControlUIModule interface
 
 	/**
@@ -63,6 +62,8 @@ public:
 	 */
 	TSharedRef<SRemoteControlPanel> CreateRemoteControlPanel(URemoteControlPreset* Preset, const TSharedPtr<IToolkitHost>& ToolkitHost = TSharedPtr<IToolkitHost>());
 
+	void UnregisterRemoteControlPanel(SRemoteControlPanel* Panel);
+
 	/**
 	 * Get the map of entity metadata entry customizations.
 	 */
@@ -117,6 +118,14 @@ private:
 	/** Handle creating the row extensions.  */
 	void HandleCreatePropertyRowExtension(const FOnGenerateGlobalRowExtensionArgs& InArgs, TArray<FPropertyRowExtensionButton>& OutExtensions);
 
+	/*
+	 * Returns the Panel widget associated with an property by, in the case of embedded presets, examining the
+	 * world that the property's owning object belongs to. In the case of assets, just returns the Panel ref in this module.
+	 */
+	TSharedPtr<SRemoteControlPanel> GetPanelForObject(const UObject* Object) const;
+	TSharedPtr<SRemoteControlPanel> GetPanelForProperty(const FRCExposesPropertyArgs& InPropertyArgs) const;
+	TSharedPtr<SRemoteControlPanel> GetPanelForPropertyChangeEvent(const FPropertyChangedEvent& InPropertyChangeEvent) const;
+
 	/** Handle getting the icon displayed in the row extension. */
 	FSlateIcon OnGetExposedIcon(const FRCExposesPropertyArgs& InArgs) const;
 
@@ -210,4 +219,6 @@ private:
 
 	/** Holds the advanced asset type category bit registered by Remote Control. */
 	EAssetTypeCategories::Type RemoteControlAssetCategoryBit;
+
+	TArray<TWeakPtr<SRemoteControlPanel>> RegisteredRemoteControlPanels;
 };
diff --git a/Engine/Plugins/VirtualProduction/RemoteControl/Source/RemoteControlUI/Private/UI/SRCPanelExposedActor.cpp b/Engine/Plugins/VirtualProduction/RemoteControl/Source/RemoteControlUI/Private/UI/SRCPanelExposedActor.cpp
--- a/Engine/Plugins/VirtualProduction/RemoteControl/Source/RemoteControlUI/Private/UI/SRCPanelExposedActor.cpp
+++ b/Engine/Plugins/VirtualProduction/RemoteControl/Source/RemoteControlUI/Private/UI/SRCPanelExposedActor.cpp
@@ -143,7 +143,8 @@ TSharedRef<SWidget> SRCPanelExposedActor::CreateEmbeddedPresetActorPicker()
 {
 	FSceneOutlinerModule& SceneOutlinerModule = FModuleManager::LoadModuleChecked<FSceneOutlinerModule>("SceneOutliner");
 	FSceneOutlinerInitializationOptions InitOptions;
-	UWorld* PresetWorld = URemoteControlPreset::GetPresetWorld(Preset.Get());
+	constexpr bool bAllowPIE = false;
+	UWorld* PresetWorld = URemoteControlPreset::GetWorld(Preset.Get(), bAllowPIE);
 
 	TSharedRef<SWidget> ActorPicker =
 		SceneOutlinerModule.CreateActorPicker(
diff --git a/Engine/Plugins/VirtualProduction/RemoteControl/Source/RemoteControlUI/Private/UI/SRCPanelExposedEntity.cpp b/Engine/Plugins/VirtualProduction/RemoteControl/Source/RemoteControlUI/Private/UI/SRCPanelExposedEntity.cpp
--- a/Engine/Plugins/VirtualProduction/RemoteControl/Source/RemoteControlUI/Private/UI/SRCPanelExposedEntity.cpp
+++ b/Engine/Plugins/VirtualProduction/RemoteControl/Source/RemoteControlUI/Private/UI/SRCPanelExposedEntity.cpp
@@ -166,11 +166,13 @@ TSharedRef<SWidget> SRCPanelExposedEntity::CreateRebindAllPropertiesForActorMenu
 		Options.Filters->AddFilterPredicate<FActorTreeItem>(FActorTreeItem::FFilterPredicate::CreateSP(this, &SRCPanelExposedEntity::IsActorSelectable));
 	}
 
+	constexpr bool bAllowPIE = false;
+
 	return SNew(SBox)
 		.MaxDesiredHeight(400.0f)
 		.WidthOverride(300.0f)
 		[
-			SceneOutlinerModule.CreateActorPicker(Options, FOnActorPicked::CreateSP(this, &SRCPanelExposedEntity::OnActorSelectedForRebindAllProperties), URemoteControlPreset::GetPresetWorld(Preset.Get()))
+			SceneOutlinerModule.CreateActorPicker(Options, FOnActorPicked::CreateSP(this, &SRCPanelExposedEntity::OnActorSelectedForRebindAllProperties), URemoteControlPreset::GetWorld(Preset.Get(), bAllowPIE))
 		];
 }
 
@@ -209,7 +211,8 @@ TSharedRef<SWidget> SRCPanelExposedEntity::CreateRebindMenuContent()
 	FSceneOutlinerInitializationOptions Options;
 	Options.Filters = MakeShared<FSceneOutlinerFilters>();
 	Options.Filters->AddFilterPredicate<FActorTreeItem>(FActorTreeItem::FFilterPredicate::CreateSP(this, &SRCPanelExposedEntity::IsActorSelectable));
-	UWorld* PresetWorld = URemoteControlPreset::GetPresetWorld(Preset.Get());
+	constexpr bool bAllowPIE = false;
+	UWorld* PresetWorld = URemoteControlPreset::GetWorld(Preset.Get(), bAllowPIE);
 
 	return SNew(SBox)
 		.MaxDesiredHeight(400.0f)
diff --git a/Engine/Plugins/VirtualProduction/RemoteControl/Source/RemoteControlUI/Private/UI/SRCPanelFunctionPicker.cpp b/Engine/Plugins/VirtualProduction/RemoteControl/Source/RemoteControlUI/Private/UI/SRCPanelFunctionPicker.cpp
--- a/Engine/Plugins/VirtualProduction/RemoteControl/Source/RemoteControlUI/Private/UI/SRCPanelFunctionPicker.cpp
+++ b/Engine/Plugins/VirtualProduction/RemoteControl/Source/RemoteControlUI/Private/UI/SRCPanelFunctionPicker.cpp
@@ -283,7 +283,8 @@ void SRCPanelFunctionPicker::Refresh()
 	if (Class->IsChildOf(AActor::StaticClass()) && GEditor)
 	{
 		TArray<AActor*> ActorList;
-		UGameplayStatics::GetAllActorsOfClass(GetPresetWorld(), Class, ActorList);
+		constexpr bool bIgnorePIE = false;
+		UGameplayStatics::GetAllActorsOfClass(GetWorld(bIgnorePIE), Class, ActorList);
 
 		for (AActor* Actor : ActorList)
 		{
@@ -304,16 +305,16 @@ void SRCPanelFunctionPicker::Refresh()
 	}
 }
 
-UWorld* SRCPanelFunctionPicker::GetPresetWorld(bool bIgnorePIE) const
+UWorld* SRCPanelFunctionPicker::GetWorld(bool bIgnorePIE) const
 {
 	TSharedPtr<SRemoteControlPanel> RCP = RemoteControlPanel.Pin();
 
 	if (RCP.IsValid())
 	{
-		return URemoteControlPreset::GetPresetWorld(RCP->GetPreset(), !bIgnorePIE);
+		return URemoteControlPreset::GetWorld(RCP->GetPreset(), !bIgnorePIE);
 	}
 
-	return URemoteControlPreset::GetPresetWorld(nullptr, !bIgnorePIE);
+	return URemoteControlPreset::GetWorld(nullptr, !bIgnorePIE);
 }
 
 #undef LOCTEXT_NAMESPACE /* RemoteControlPanel */
\ No newline at end of file
diff --git a/Engine/Plugins/VirtualProduction/RemoteControl/Source/RemoteControlUI/Private/UI/SRCPanelFunctionPicker.h b/Engine/Plugins/VirtualProduction/RemoteControl/Source/RemoteControlUI/Private/UI/SRCPanelFunctionPicker.h
--- a/Engine/Plugins/VirtualProduction/RemoteControl/Source/RemoteControlUI/Private/UI/SRCPanelFunctionPicker.h
+++ b/Engine/Plugins/VirtualProduction/RemoteControl/Source/RemoteControlUI/Private/UI/SRCPanelFunctionPicker.h
@@ -58,5 +58,5 @@ private:
 	/** Keep a weak reference to the RCPanel for access to toolkit host. */
 	TWeakPtr<SRemoteControlPanel> RemoteControlPanel;
 	/** Checks the remote control panel to get a reference to the preset world. */
-	UWorld* GetPresetWorld(bool bIgnorePIE = false) const;
+	UWorld* GetWorld(bool bIgnorePIE = false) const;
 };
\ No newline at end of file
diff --git a/Engine/Plugins/VirtualProduction/RemoteControl/Source/RemoteControlUI/Private/UI/SRemoteControlPanel.cpp b/Engine/Plugins/VirtualProduction/RemoteControl/Source/RemoteControlUI/Private/UI/SRemoteControlPanel.cpp
--- a/Engine/Plugins/VirtualProduction/RemoteControl/Source/RemoteControlUI/Private/UI/SRemoteControlPanel.cpp
+++ b/Engine/Plugins/VirtualProduction/RemoteControl/Source/RemoteControlUI/Private/UI/SRemoteControlPanel.cpp
@@ -2,6 +2,8 @@
 
 #include "SRemoteControlPanel.h"
 
+#include "IRemoteControlUIModule.h"
+#include "RemoteControlUIModule.h"
 #include "Action/SRCActionPanel.h"
 #include "ActorEditorUtils.h"
 #include "Behaviour/SRCBehaviourPanel.h"
@@ -751,6 +753,12 @@ SRemoteControlPanel::~SRemoteControlPanel()
 	// Remove protocol bindings
 	IRemoteControlProtocolWidgetsModule& ProtocolWidgetsModule = FModuleManager::LoadModuleChecked<IRemoteControlProtocolWidgetsModule>("RemoteControlProtocolWidgets");
 	ProtocolWidgetsModule.ResetProtocolBindingList();	
+
+	// Unregister with UI module
+	if (FModuleManager::Get().IsModuleLoaded("RemoteControlUIModule"))
+	{
+		FRemoteControlUIModule::Get().UnregisterRemoteControlPanel(this);
+	}	
 }
 
 void SRemoteControlPanel::Shutdown()
@@ -1039,7 +1047,8 @@ TSharedRef<SWidget> SRemoteControlPanel::CreateExposeButton()
 					{
 						FSceneOutlinerModule& SceneOutlinerModule = FModuleManager::LoadModuleChecked<FSceneOutlinerModule>("SceneOutliner");
 						FSceneOutlinerInitializationOptions InitOptions;
-						UWorld* PresetWorld = URemoteControlPreset::GetPresetWorld(Preset.Get());
+						constexpr bool bAllowPIE = false;
+						UWorld* PresetWorld = URemoteControlPreset::GetWorld(Preset.Get(), bAllowPIE);
 
 						SubMenuBuilder.AddWidget(
 							SceneOutlinerModule.CreateActorPicker(
@@ -1129,7 +1138,9 @@ TSharedRef<SWidget> SRemoteControlPanel::CreateExposeByClassWidget()
 	TSharedRef<SWidget> Widget = FModuleManager::LoadModuleChecked<FClassViewerModule>("ClassViewer").CreateClassViewer(Options, FOnClassPicked::CreateLambda(
 		[this](UClass* ChosenClass)
 		{
-			if (UWorld* World = URemoteControlPreset::GetPresetWorld(Preset.Get()))
+			constexpr bool bAllowPIE = false;
+
+			if (UWorld* World = URemoteControlPreset::GetWorld(Preset.Get(), bAllowPIE))
 			{
 				for (TActorIterator<AActor> It(World, ChosenClass, EActorIteratorFlags::SkipPendingKill); It; ++It)
 				{
@@ -1158,8 +1169,9 @@ TSharedRef<SWidget> SRemoteControlPanel::CreateExposeByClassWidget()
 void SRemoteControlPanel::CacheLevelClasses()
 {
 	CachedClassesInLevel.Empty();
+	constexpr bool bAllowPIE = false;
 
-	if (UWorld* World = URemoteControlPreset::GetPresetWorld(Preset.Get()))
+	if (UWorld* World = URemoteControlPreset::GetWorld(Preset.Get(), bAllowPIE))
 	{
 		for (TActorIterator<AActor> It(World, AActor::StaticClass(), EActorIteratorFlags::SkipPendingKill); It; ++It)
 		{
diff --git a/Engine/Plugins/VirtualProduction/RemoteControlWebInterface/Source/RemoteControlWebInterface/Private/RCWebInterfaceLibrary.cpp b/Engine/Plugins/VirtualProduction/RemoteControlWebInterface/Source/RemoteControlWebInterface/Private/RCWebInterfaceLibrary.cpp
--- a/Engine/Plugins/VirtualProduction/RemoteControlWebInterface/Source/RemoteControlWebInterface/Private/RCWebInterfaceLibrary.cpp
+++ b/Engine/Plugins/VirtualProduction/RemoteControlWebInterface/Source/RemoteControlWebInterface/Private/RCWebInterfaceLibrary.cpp
@@ -52,7 +52,8 @@ TMap<FString, AActor*> URCWebInterfaceBlueprintLibrary::FindMatchingActorsToRebi
 				}
 				else if (UClass* SupportedBindingClass = RemoteControlProperty->GetSupportedBindingClass())
 				{
-					UWorld* World = RCPreset->GetPresetWorld(false);
+					constexpr bool bAllowPIE = false;
+					UWorld* World = RCPreset->GetWorld(bAllowPIE);
 					if (SupportedBindingClass->IsChildOf(AActor::StaticClass()))
 					{
 						TArray<AActor*> MatchingActors;
