Misc direct FTCHARToUTF8_Convert usage replaced with StringCast or FPlatformString.

diff --git a/Engine/Plugins/VirtualProduction/RemoteControl/Source/WebRemoteControl/Private/WebRemoteControlUtils.cpp b/Engine/Plugins/VirtualProduction/RemoteControl/Source/WebRemoteControl/Private/WebRemoteControlUtils.cpp
--- a/Engine/Plugins/VirtualProduction/RemoteControl/Source/WebRemoteControl/Private/WebRemoteControlUtils.cpp
+++ b/Engine/Plugins/VirtualProduction/RemoteControl/Source/WebRemoteControl/Private/WebRemoteControlUtils.cpp
@@ -14,15 +14,15 @@ void WebRemoteControlUtils::ConvertToTCHAR(TConstArrayView<uint8> InUTF8Payload,
 void WebRemoteControlUtils::ConvertToUTF8(TConstArrayView<uint8> InTCHARPayload, TArray<uint8>& OutUTF8Payload)
 {
 	int32 StartIndex = OutUTF8Payload.Num();
-	OutUTF8Payload.AddUninitialized(FTCHARToUTF8_Convert::ConvertedLength((TCHAR*)InTCHARPayload.GetData(), InTCHARPayload.Num() / sizeof(TCHAR)) * sizeof(ANSICHAR));
-	FTCHARToUTF8_Convert::Convert((ANSICHAR*)(OutUTF8Payload.GetData() + StartIndex), (OutUTF8Payload.Num() - StartIndex) / sizeof(ANSICHAR), (TCHAR*)InTCHARPayload.GetData(), InTCHARPayload.Num() / sizeof(TCHAR));
+	OutUTF8Payload.AddUninitialized(FPlatformString::ConvertedLength<UTF8CHAR>((TCHAR*)InTCHARPayload.GetData(), InTCHARPayload.Num() / sizeof(TCHAR)) * sizeof(ANSICHAR));
+	FPlatformString::Convert((UTF8CHAR*)(OutUTF8Payload.GetData() + StartIndex), (OutUTF8Payload.Num() - StartIndex) / sizeof(ANSICHAR), (TCHAR*)InTCHARPayload.GetData(), InTCHARPayload.Num() / sizeof(TCHAR));
 }
 
 void WebRemoteControlUtils::ConvertToUTF8(const FString& InString, TArray<uint8>& OutUTF8Payload)
 {
 	int32 StartIndex = OutUTF8Payload.Num();
-	OutUTF8Payload.AddUninitialized(FTCHARToUTF8_Convert::ConvertedLength(*InString, InString.Len()) * sizeof(ANSICHAR));
-	FTCHARToUTF8_Convert::Convert((ANSICHAR*)(OutUTF8Payload.GetData() + StartIndex), (OutUTF8Payload.Num() - StartIndex) / sizeof(ANSICHAR), *InString, InString.Len());
+	OutUTF8Payload.AddUninitialized(FPlatformString::ConvertedLength<UTF8CHAR>(*InString, InString.Len()) * sizeof(ANSICHAR));
+	FPlatformString::Convert((UTF8CHAR*)(OutUTF8Payload.GetData() + StartIndex), (OutUTF8Payload.Num() - StartIndex) / sizeof(ANSICHAR), *InString, InString.Len());
 }
 
 TSharedRef<IStructSerializerBackend> WebRemoteControlUtils::CreateJsonSerializerBackend(FMemoryWriter& Writer)
