Fixing Memstomp, Malloc

diff --git a/Engine/Plugins/VirtualProduction/RemoteControl/Source/RemoteControlUI/Private/UI/SRCPanelExposedField.cpp b/Engine/Plugins/VirtualProduction/RemoteControl/Source/RemoteControlUI/Private/UI/SRCPanelExposedField.cpp
--- a/Engine/Plugins/VirtualProduction/RemoteControl/Source/RemoteControlUI/Private/UI/SRCPanelExposedField.cpp
+++ b/Engine/Plugins/VirtualProduction/RemoteControl/Source/RemoteControlUI/Private/UI/SRCPanelExposedField.cpp
@@ -182,25 +182,35 @@ TSharedRef<SWidget> SRCPanelExposedField::ConstructWidget()
 					];
 
 					TSharedPtr<IPropertyHandle> PropertyHandle = Node->CreatePropertyHandle();
-					if (PropertyHandle->GetParentHandle()->GetPropertyDisplayName().ToString() == FString("Transform")
-						|| PropertyHandle->GetOuterBaseClass() == UMaterialInstanceDynamic::StaticClass())
+					if (PropertyHandle.IsValid()
+						&& (PropertyHandle && (PropertyHandle->GetParentHandle()->GetPropertyDisplayName().ToString() == FString("Transform")
+						|| PropertyHandle->GetOuterBaseClass() == UMaterialInstanceDynamic::StaticClass())))
 					{
 						// Set up a Zeroed DefaultValue, in case an ExposedEntity doesn't have a native Default. Needed for certain ResetToDefault cases.
 						void* ValuePtr;
 						Node->CreatePropertyHandle()->GetValueData(ValuePtr);
-						DefaultValue = MakeShared<uint8>(Node->CreatePropertyHandle()->GetProperty()->GetSize());
+						DefaultValue.Reset(new uint8[PropertyHandle->GetProperty()->GetSize()]);
 						Node->CreatePropertyHandle()->GetProperty()->CopyCompleteValue(DefaultValue.Get(), ValuePtr);
 						Node->CreatePropertyHandle()->GetProperty()->ClearValue(DefaultValue.Get());
-						
-						auto IsVisible = [this, Node]()
+
+						TWeakPtr<SRCPanelExposedField> WeakPtr = StaticCastWeakPtr<SRCPanelExposedField>(AsWeak());
+						auto IsVisible = [WeakPtr, Node]()
 						{
-							TSharedPtr<IPropertyHandle> PropertyHandle = Node->CreatePropertyHandle();
-							void* DataPtr;
-							PropertyHandle->GetValueData(DataPtr);
-							FProperty* NodeProperty = PropertyHandle->GetProperty();
-							
-							bool bVisible = !NodeProperty->Identical(DefaultValue.Get(), DataPtr);
-							return bVisible ? EVisibility::Visible : EVisibility::Hidden;
+							if (!WeakPtr.IsValid())
+							{
+								return EVisibility::Hidden;
+							}
+
+							if (TSharedPtr<IPropertyHandle> PropertyHandle = Node->CreatePropertyHandle())
+							{
+								void* DataPtr;
+								PropertyHandle->GetValueData(DataPtr);
+								FProperty* NodeProperty = PropertyHandle->GetProperty();
+								
+								bool bVisible = !NodeProperty->Identical(WeakPtr.Pin()->DefaultValue.Get(), DataPtr);
+								return bVisible ? EVisibility::Visible : EVisibility::Hidden;
+							}
+							return EVisibility::Hidden;
 						};
 						
 						ResetButtonWidget = SNew(SBox)
diff --git a/Engine/Plugins/VirtualProduction/RemoteControl/Source/RemoteControlUI/Private/UI/SRCPanelExposedField.h b/Engine/Plugins/VirtualProduction/RemoteControl/Source/RemoteControlUI/Private/UI/SRCPanelExposedField.h
--- a/Engine/Plugins/VirtualProduction/RemoteControl/Source/RemoteControlUI/Private/UI/SRCPanelExposedField.h
+++ b/Engine/Plugins/VirtualProduction/RemoteControl/Source/RemoteControlUI/Private/UI/SRCPanelExposedField.h
@@ -92,7 +92,7 @@ private:
 	/** Holds the panel's cached widgets. */
 	TWeakPtr<FRCPanelWidgetRegistry> WidgetRegistry;
 	/** Holds the zeroed Default Value of the ExposedField */
-	TSharedPtr<uint8> DefaultValue;
+	TUniquePtr<uint8> DefaultValue;
 	/** Holds the shared reference of reset button for this field. */
 	TSharedPtr<SWidget> ResetButtonWidget;
 };
