Remote Control Panel - Reset Arrow Field, Memstomp Hotfix

Fixes issues and improves little bits of the Reset Arrow Code being used.
- Change of uint8 to uint8[]
- Remove redundent cnoditional
- Remove not needed Type Cast

diff --git a/Engine/Plugins/VirtualProduction/RemoteControl/Source/RemoteControlUI/Private/UI/SRCPanelExposedField.cpp b/Engine/Plugins/VirtualProduction/RemoteControl/Source/RemoteControlUI/Private/UI/SRCPanelExposedField.cpp
--- a/Engine/Plugins/VirtualProduction/RemoteControl/Source/RemoteControlUI/Private/UI/SRCPanelExposedField.cpp
+++ b/Engine/Plugins/VirtualProduction/RemoteControl/Source/RemoteControlUI/Private/UI/SRCPanelExposedField.cpp
@@ -182,9 +182,8 @@ TSharedRef<SWidget> SRCPanelExposedField::ConstructWidget()
 					];
 
 					TSharedPtr<IPropertyHandle> PropertyHandle = Node->CreatePropertyHandle();
-					if (PropertyHandle.IsValid()
-						&& (PropertyHandle && (PropertyHandle->GetParentHandle()->GetPropertyDisplayName().ToString() == FString("Transform")
-						|| PropertyHandle->GetOuterBaseClass() == UMaterialInstanceDynamic::StaticClass())))
+					if (PropertyHandle && (PropertyHandle->GetParentHandle()->GetPropertyDisplayName().ToString() == FString("Transform")
+						|| PropertyHandle->GetOuterBaseClass() == UMaterialInstanceDynamic::StaticClass()))
 					{
 						// Set up a Zeroed DefaultValue, in case an ExposedEntity doesn't have a native Default. Needed for certain ResetToDefault cases.
 						void* ValuePtr;
@@ -193,21 +192,18 @@ TSharedRef<SWidget> SRCPanelExposedField::ConstructWidget()
 						Node->CreatePropertyHandle()->GetProperty()->CopyCompleteValue(DefaultValue.Get(), ValuePtr);
 						Node->CreatePropertyHandle()->GetProperty()->ClearValue(DefaultValue.Get());
 
-						TWeakPtr<SRCPanelExposedField> WeakPtr = StaticCastWeakPtr<SRCPanelExposedField>(AsWeak());
-						auto IsVisible = [WeakPtr, Node]()
+						TWeakPtr<SRCPanelExposedField> WeakFieldPtr = SharedThis(this);
+						auto IsVisible = [WeakFieldPtr, Node]()
 						{
-							if (!WeakPtr.IsValid())
-							{
-								return EVisibility::Hidden;
-							}
-
-							if (TSharedPtr<IPropertyHandle> PropertyHandle = Node->CreatePropertyHandle())
+							TSharedPtr<SRCPanelExposedField> FieldPtr = WeakFieldPtr.Pin();
+							TSharedPtr<IPropertyHandle> PropertyHandle = Node->CreatePropertyHandle();
+							if (FieldPtr && PropertyHandle)
 							{
 								void* DataPtr;
 								PropertyHandle->GetValueData(DataPtr);
 								FProperty* NodeProperty = PropertyHandle->GetProperty();
 								
-								bool bVisible = !NodeProperty->Identical(WeakPtr.Pin()->DefaultValue.Get(), DataPtr);
+								bool bVisible = !NodeProperty->Identical(FieldPtr->DefaultValue.Get(), DataPtr);
 								return bVisible ? EVisibility::Visible : EVisibility::Hidden;
 							}
 							return EVisibility::Hidden;
diff --git a/Engine/Plugins/VirtualProduction/RemoteControl/Source/RemoteControlUI/Private/UI/SRCPanelExposedField.h b/Engine/Plugins/VirtualProduction/RemoteControl/Source/RemoteControlUI/Private/UI/SRCPanelExposedField.h
--- a/Engine/Plugins/VirtualProduction/RemoteControl/Source/RemoteControlUI/Private/UI/SRCPanelExposedField.h
+++ b/Engine/Plugins/VirtualProduction/RemoteControl/Source/RemoteControlUI/Private/UI/SRCPanelExposedField.h
@@ -92,7 +92,7 @@ private:
 	/** Holds the panel's cached widgets. */
 	TWeakPtr<FRCPanelWidgetRegistry> WidgetRegistry;
 	/** Holds the zeroed Default Value of the ExposedField */
-	TUniquePtr<uint8> DefaultValue;
+	TUniquePtr<uint8[]> DefaultValue;
 	/** Holds the shared reference of reset button for this field. */
 	TSharedPtr<SWidget> ResetButtonWidget;
 };
