Conditional Behaviour (v2)

diff --git a/Engine/Plugins/VirtualProduction/RemoteControl/Config/BaseRemoteControl.ini b/Engine/Plugins/VirtualProduction/RemoteControl/Config/BaseRemoteControl.ini
--- a/Engine/Plugins/VirtualProduction/RemoteControl/Config/BaseRemoteControl.ini
+++ b/Engine/Plugins/VirtualProduction/RemoteControl/Config/BaseRemoteControl.ini
@@ -12,6 +12,6 @@
 +SupportedControllerTypes=EPropertyBagPropertyType::String
 +SupportedControllerTypes=EPropertyBagPropertyType::Text
 +SupportedControllerTypes=EPropertyBagPropertyType::Struct
-+SupportedStructTypes="FVector"
-+SupportedStructTypes="FColor"
-+SupportedStructTypes="FRotator"
\ No newline at end of file
++SupportedControllerStructTypes=Vector
++SupportedControllerStructTypes=Color
++SupportedControllerStructTypes=Rotator
\ No newline at end of file
diff --git a/Engine/Plugins/VirtualProduction/RemoteControl/Source/RemoteControl/Private/RCVirtualProperty.cpp b/Engine/Plugins/VirtualProduction/RemoteControl/Source/RemoteControl/Private/RCVirtualProperty.cpp
--- a/Engine/Plugins/VirtualProduction/RemoteControl/Source/RemoteControl/Private/RCVirtualProperty.cpp
+++ b/Engine/Plugins/VirtualProduction/RemoteControl/Source/RemoteControl/Private/RCVirtualProperty.cpp
@@ -140,7 +140,47 @@ bool URCVirtualPropertyBase::DeserializeFromBackend(IStructDeserializerBackend&
 	return bSuccess;
 }
 
-bool URCVirtualPropertyBase::IsValueEqual(URCVirtualPropertyBase* InVirtualProperty)
+bool URCVirtualPropertyBase::IsNumericType() const
+{
+	const FProperty* Property = GetProperty();
+
+	return Property && Property->IsA(FNumericProperty::StaticClass());
+}
+
+bool URCVirtualPropertyBase::IsVectorType() const
+{
+	const FProperty* Property = GetProperty();
+	if (const FStructProperty* StructProperty = CastField<FStructProperty>(Property))
+	{
+		return StructProperty->Struct == TBaseStructure<FVector>::Get();
+	}
+
+	return false;
+}
+
+bool URCVirtualPropertyBase::IsColorType() const
+{
+	const FProperty* Property = GetProperty();
+	if (const FStructProperty* StructProperty = CastField<FStructProperty>(Property))
+	{
+		return StructProperty->Struct == TBaseStructure<FColor>::Get();
+	}
+
+	return false;
+}
+
+bool URCVirtualPropertyBase::IsRotatorType() const
+{
+	const FProperty* Property = GetProperty();
+	if (const FStructProperty* StructProperty = CastField<FStructProperty>(Property))
+	{
+		return StructProperty->Struct == TBaseStructure<FRotator>::Get();
+	}
+
+	return false;
+}
+
+bool URCVirtualPropertyBase::IsValueEqual(URCVirtualPropertyBase* InVirtualProperty) const
 {
 	const FProperty* ThisProperty = GetProperty();
 	const uint8* ThisContainerPtr = GetContainerPtr();
@@ -194,7 +234,7 @@ bool URCVirtualPropertyBase::CopyCompleteValue(const FProperty* InTargetProperty
 struct FRCVirtualPropertyCastHelpers
 {
 	template<typename TPropertyType, typename TValueType>
-	static bool GetPrimitiveValue(URCVirtualPropertyBase* InVirtualProperty, TValueType& OutValue)
+	static bool GetPrimitiveValue(const URCVirtualPropertyBase* InVirtualProperty, TValueType& OutValue)
 	{
 		const uint8* ContainerPtr = InVirtualProperty->GetContainerPtr();
 		if (!ContainerPtr)
@@ -234,7 +274,7 @@ struct FRCVirtualPropertyCastHelpers
 	}
 
 	template<typename TValueType>
-	static bool GetStructValue(URCVirtualPropertyBase* InVirtualProperty, UScriptStruct* InScriptStruct, TValueType* OutValuePtr)
+	static bool GetStructValue(const URCVirtualPropertyBase* InVirtualProperty, UScriptStruct* InScriptStruct, TValueType* OutValuePtr)
 	{
 		const uint8* ValuePtr = InVirtualProperty->GetValuePtr();
 		if (!ValuePtr)
@@ -278,77 +318,77 @@ struct FRCVirtualPropertyCastHelpers
 	}
 };
 
-bool URCVirtualPropertyBase::GetValueBool(bool& OutBoolValue)
+bool URCVirtualPropertyBase::GetValueBool(bool& OutBoolValue) const
 {
 	return FRCVirtualPropertyCastHelpers::GetPrimitiveValue<FBoolProperty, bool>(this, OutBoolValue);
 }
 
-bool URCVirtualPropertyBase::GetValueInt8(int8& OutInt8)
+bool URCVirtualPropertyBase::GetValueInt8(int8& OutInt8) const
 {
 	return FRCVirtualPropertyCastHelpers::GetPrimitiveValue<FInt8Property, int8>(this, OutInt8);
 }
 
-bool URCVirtualPropertyBase::GetValueByte(uint8& OutByte)
+bool URCVirtualPropertyBase::GetValueByte(uint8& OutByte) const
 {
 	return FRCVirtualPropertyCastHelpers::GetPrimitiveValue<FByteProperty, uint8&>(this, OutByte);
 }
 
-bool URCVirtualPropertyBase::GetValueInt16(int16& OutInt16)
+bool URCVirtualPropertyBase::GetValueInt16(int16& OutInt16) const
 {
 	return FRCVirtualPropertyCastHelpers::GetPrimitiveValue<FInt16Property, int16&>(this, OutInt16);
 }
 
-bool URCVirtualPropertyBase::GetValueUint16(uint16& OutUInt16)
+bool URCVirtualPropertyBase::GetValueUint16(uint16& OutUInt16) const
 {
 	return FRCVirtualPropertyCastHelpers::GetPrimitiveValue<FUInt16Property, uint16&>(this, OutUInt16);
 }
 
-bool URCVirtualPropertyBase::GetValueInt32(int32& OutInt32)
+bool URCVirtualPropertyBase::GetValueInt32(int32& OutInt32) const
 {
 	return FRCVirtualPropertyCastHelpers::GetPrimitiveValue<FIntProperty, int32&>(this, OutInt32);
 }
 
-bool URCVirtualPropertyBase::GetValueUInt32(uint32& OutUInt32)
+bool URCVirtualPropertyBase::GetValueUInt32(uint32& OutUInt32) const
 {
 	return FRCVirtualPropertyCastHelpers::GetPrimitiveValue<FUInt32Property, uint32&>(this, OutUInt32);
 }
 
-bool URCVirtualPropertyBase::GetValueInt64(int64& OuyInt64)
+bool URCVirtualPropertyBase::GetValueInt64(int64& OuyInt64) const
 {
 	return FRCVirtualPropertyCastHelpers::GetPrimitiveValue<FInt64Property, int64>(this, OuyInt64);
 }
 
-bool URCVirtualPropertyBase::GetValueUint64(uint64& OuyUInt64)
+bool URCVirtualPropertyBase::GetValueUint64(uint64& OuyUInt64) const
 {
 return FRCVirtualPropertyCastHelpers::GetPrimitiveValue<FUInt64Property, uint64>(this, OuyUInt64);
 }
 
-bool URCVirtualPropertyBase::GetValueFloat(float& OutFloat)
+bool URCVirtualPropertyBase::GetValueFloat(float& OutFloat) const
 {
 	return FRCVirtualPropertyCastHelpers::GetPrimitiveValue<FFloatProperty, float>(this, OutFloat);
 }
 
-bool URCVirtualPropertyBase::GetValueDouble(double& OutDouble)
+bool URCVirtualPropertyBase::GetValueDouble(double& OutDouble) const
 {
 	return FRCVirtualPropertyCastHelpers::GetPrimitiveValue<FDoubleProperty, double>(this, OutDouble);
 }
 
-bool URCVirtualPropertyBase::GetValueString(FString& OutStringValue)
+bool URCVirtualPropertyBase::GetValueString(FString& OutStringValue) const
 {
 	return FRCVirtualPropertyCastHelpers::GetPrimitiveValue<FStrProperty, FString>(this, OutStringValue);
 }
 
-bool URCVirtualPropertyBase::GetValueName(FName& OutNameValue)
+bool URCVirtualPropertyBase::GetValueName(FName& OutNameValue) const
 {
 	return FRCVirtualPropertyCastHelpers::GetPrimitiveValue<FNameProperty, FName>(this, OutNameValue);
 }
 
-bool URCVirtualPropertyBase::GetValueText(FText& OutTextValue)
+bool URCVirtualPropertyBase::GetValueText(FText& OutTextValue) const
 {
 	return FRCVirtualPropertyCastHelpers::GetPrimitiveValue<FTextProperty, FText>(this, OutTextValue);
 }
 
-bool URCVirtualPropertyBase::GetValueNumericInteger(int64& OutInt64Value)
+bool URCVirtualPropertyBase::GetValueNumericInteger(int64& OutInt64Value) const
 {
 	const uint8* ContainerPtr = GetContainerPtr();
 	if (!ContainerPtr)
@@ -371,21 +411,109 @@ bool URCVirtualPropertyBase::GetValueNumericInteger(int64& OutInt64Value)
 	return false;
 }
 
-bool URCVirtualPropertyBase::GetValueVector(FVector& OutVector)
+bool URCVirtualPropertyBase::GetValueVector(FVector& OutVector) const
 {
 	return FRCVirtualPropertyCastHelpers::GetStructValue<FVector>(this, TBaseStructure<FVector>::Get(), &OutVector);
 }
 
-bool URCVirtualPropertyBase::GetValueRotator(FRotator& OutRotator)
+bool URCVirtualPropertyBase::GetValueRotator(FRotator& OutRotator) const
 {
 	return FRCVirtualPropertyCastHelpers::GetStructValue<FRotator>(this, TBaseStructure<FRotator>::Get(), &OutRotator);
 }
 
-bool URCVirtualPropertyBase::GetValueColor(FColor& OutColor)
+bool URCVirtualPropertyBase::GetValueColor(FColor& OutColor) const
 {
 	return FRCVirtualPropertyCastHelpers::GetStructValue<FColor>(this, TBaseStructure<FColor>::Get(), &OutColor);
 }
 
+FString URCVirtualPropertyBase::GetDisplayValueAsString() const
+{
+	const FProperty* Property = GetProperty();
+
+	if (const FBoolProperty* BoolProperty = CastField<FBoolProperty>(Property))
+	{
+		bool Value;
+		GetValueBool(Value);
+
+		return Value ? TEXT("true") : TEXT("false");
+	}
+	else if (const FByteProperty* ByteProperty = CastField<FByteProperty>(Property))
+	{
+		uint8 Value;
+		GetValueByte(Value);
+
+		return FString::Printf(TEXT("%d"), Value);
+	}
+	else if (const FTextProperty* TextProperty = CastField<FTextProperty>(Property))
+	{
+		FText Value;
+		GetValueText(Value);
+
+		return Value.ToString();
+	}
+	else if (const FStrProperty* StringProperty = CastField<FStrProperty>(Property))
+	{
+		FString Value;
+		GetValueString(Value);
+
+		return Value;
+	}
+	else if (const FNameProperty* NameProperty = CastField<FNameProperty>(Property))
+	{
+		FName Value;
+		GetValueName(Value);
+
+		return Value.ToString();
+	}
+	else if (const FIntProperty* IntProperty = CastField<FIntProperty>(Property))
+	{
+		int32 Value;
+		GetValueInt32(Value);
+
+		return FString::Printf(TEXT("%d"), Value);
+	}
+	else if (const FFloatProperty* FloatProperty = CastField<FFloatProperty>(Property))
+	{
+		float Value;
+		GetValueFloat(Value);
+
+		return FString::Printf(TEXT("%f"), Value);
+	}
+	else if (const FDoubleProperty* DoubleProperty = CastField<FDoubleProperty>(Property))
+	{
+		double Value;
+		GetValueDouble(Value);
+
+		return FString::Printf(TEXT("%lf"), Value);
+	}
+	else if (IsVectorType())
+	{
+		FVector Vector;
+		GetValueVector(Vector);
+
+		return Vector.ToString();
+	}
+	else if (IsRotatorType())
+	{
+		FRotator Rotator;
+		GetValueRotator(Rotator);
+
+		return Rotator.ToString();
+	}
+	else if (IsColorType())
+	{
+		FColor Color;
+		GetValueColor(Color);
+
+		return Color.ToString();
+	}
+
+	// Unimplemented type!
+	ensureAlwaysMsgf(false, TEXT("Unimplemented type %s passed to GetDisplayValueAsString"), *Property->GetName());
+		 
+	return "Unimplemented type";
+}
+
 bool URCVirtualPropertyBase::SetValueBool(bool InBoolValue)
 {
 	return FRCVirtualPropertyCastHelpers::SetPrimitiveValue<FBoolProperty, bool>(this, InBoolValue);
diff --git a/Engine/Plugins/VirtualProduction/RemoteControl/Source/RemoteControl/Private/RCVirtualPropertyComparators.cpp b/Engine/Plugins/VirtualProduction/RemoteControl/Source/RemoteControl/Private/RCVirtualPropertyComparators.cpp
new file mode 100644
--- /dev/null
+++ b/Engine/Plugins/VirtualProduction/RemoteControl/Source/RemoteControl/Private/RCVirtualPropertyComparators.cpp
@@ -0,0 +1,66 @@
+// Copyright Epic Games, Inc. All Rights Reserved.
+
+#include "RCVirtualProperty.h"
+
+/** Generates functions that Compare this object with a given Virtual Property, using a given Comparator */
+#define IMPLEMENT_RCLOGIC_COMPARATOR( CompareFunctionName, Comparator)  \
+\
+bool URCVirtualPropertyBase::CompareFunctionName(URCVirtualPropertyBase* InVirtualProperty) const\
+{ \
+	const FProperty* Property = GetProperty(); \
+	                                                                      \
+/** Note: Currently numeric Controller types are limited to int32 and float, so only these are implemented*/ \
+	if (const FNumericProperty* NumericProperty = CastField<FNumericProperty>(Property)) \
+	{ \
+        /** Numeric Integers (potentially covers int32, uint32, int8, int16, uint16, int32, uint32, int64, uint64) */ \
+		if (NumericProperty->IsInteger()) \
+		{ \
+            /** int32 */ \
+			if (const FIntProperty* IntProperty = CastField<FIntProperty>(Property)) \
+			{ \
+				int32 LHSValue, RHSValue; \
+				this->GetValueInt32(LHSValue); \
+				InVirtualProperty->GetValueInt32(RHSValue); \
+				                                                                          \
+				/** Comparison */                                              \
+				return LHSValue Comparator RHSValue; \
+			} \
+		} \
+		else /** Numeric Floating Point(potentially covers float and double) */ \
+		{ \
+            /** float */ \
+			if (const FFloatProperty* FloatProperty = CastField<FFloatProperty>(Property)) \
+			{ \
+				float LHSValue, RHSValue; \
+				this->GetValueFloat(LHSValue); \
+				InVirtualProperty->GetValueFloat(RHSValue); \
+				                                                                         \
+				/** Comparison */                                             \
+				return LHSValue Comparator RHSValue; \
+			} \
+            /** double */ \
+			if (const FDoubleProperty* DoubleProperty = CastField<FDoubleProperty>(Property)) \
+			{ \
+				double LHSValue, RHSValue; \
+				this->GetValueDouble(LHSValue); \
+				InVirtualProperty->GetValueDouble(RHSValue); \
+				                                                                         \
+				/** Comparison */                                             \
+				return LHSValue Comparator RHSValue; \
+			} \
+		} \
+	} \
+	else \
+	{ \
+		ensureAlwaysMsgf(false, TEXT("Only numeric types are supported for value comparison")); \
+	} \
+	ensureAlwaysMsgf(false, TEXT("Unimplemented Numeric Type %s passed for value comparison!/nSee above for currently supported numeric types"), *Property->GetName()); \
+	return false; \
+} \
+
+IMPLEMENT_RCLOGIC_COMPARATOR(IsValueGreaterThan, >)
+IMPLEMENT_RCLOGIC_COMPARATOR(IsValueGreaterThanOrEqualTo, >= )
+IMPLEMENT_RCLOGIC_COMPARATOR(IsValueLesserThan, < )
+IMPLEMENT_RCLOGIC_COMPARATOR(IsValueLesserThanOrEqualTo, <= )
+
+#undef IMPLEMENT_RCLOGIC_COMPARATOR
\ No newline at end of file
diff --git a/Engine/Plugins/VirtualProduction/RemoteControl/Source/RemoteControl/Public/RCVirtualProperty.h b/Engine/Plugins/VirtualProduction/RemoteControl/Source/RemoteControl/Public/RCVirtualProperty.h
--- a/Engine/Plugins/VirtualProduction/RemoteControl/Source/RemoteControl/Public/RCVirtualProperty.h
+++ b/Engine/Plugins/VirtualProduction/RemoteControl/Source/RemoteControl/Public/RCVirtualProperty.h
@@ -73,8 +73,32 @@ public:
 	 */
 	bool DeserializeFromBackend(IStructDeserializerBackend& InBackend);
 
+	/** Whether this Virtual Property is a numeric type, i.e. an Integral or Floating point type */
+	bool IsNumericType() const;
+
+	/** Whether this Virtual Property represents an FVector*/
+	bool IsVectorType() const;
+
+	/** Whether this Virtual Property represents an FColor*/
+	bool IsColorType() const;
+
+	/** Whether this Virtual Property represents an FRotator*/
+	bool IsRotatorType() const;
+
 	/** Compare this virtual property value with given property value */
-	bool IsValueEqual(URCVirtualPropertyBase* InVirtualProperty);
+	bool IsValueEqual(URCVirtualPropertyBase* InVirtualProperty) const;
+
+	/** Operator > comparison of self with a given virtual property*/
+	bool IsValueGreaterThan(URCVirtualPropertyBase* InVirtualProperty) const;
+
+	/** Operator >= comparison of self with a given virtual property*/
+	bool IsValueGreaterThanOrEqualTo(URCVirtualPropertyBase* InVirtualProperty) const;
+
+	/** Operator < comparison of self with a given virtual property*/
+	bool IsValueLesserThan(URCVirtualPropertyBase* InVirtualProperty) const;
+
+	/** Operator <= comparison of self with a given virtual property*/
+	bool IsValueLesserThanOrEqualTo(URCVirtualPropertyBase* InVirtualProperty) const;
 
 	/** Copy this virtual property's data onto a given FProperty
 	* 
@@ -85,75 +109,81 @@ public:
 	
 	/** Get Bool value from Virtual Property */
 	UFUNCTION(BlueprintCallable, Category = "Remote Control Behaviour")
-	bool GetValueBool(bool& OutBoolValue);
+	bool GetValueBool(bool& OutBoolValue) const;
 
 	/** Get Int8 value from Virtual Property */
 	UFUNCTION()
-	bool GetValueInt8(int8& OutInt8);
+	bool GetValueInt8(int8& OutInt8) const;
 
 	/** Get Byte value from Virtual Property */
 	UFUNCTION(BlueprintCallable, Category = "Remote Control Behaviour")
-	bool GetValueByte(uint8& OutByte);
+	bool GetValueByte(uint8& OutByte) const;
 
 	/** Get Int16 value from Virtual Property */
 	UFUNCTION()
-	bool GetValueInt16(int16& OutInt16);
+	bool GetValueInt16(int16& OutInt16) const;
 
 	/** Get Uint16 value from Virtual Property */
 	UFUNCTION()
-	bool GetValueUint16(uint16& OutUInt16);
+	bool GetValueUint16(uint16& OutUInt16) const;
 
 	/** Get Int32 value from Virtual Property */
 	UFUNCTION(BlueprintCallable, Category = "Remote Control Behaviour")
-	bool GetValueInt32(int32& OutInt32);
+	bool GetValueInt32(int32& OutInt32) const;
 
 	/** Get Uint32 value from Virtual Property */
 	UFUNCTION()
-	bool GetValueUInt32(uint32& OutUInt32);
+	bool GetValueUInt32(uint32& OutUInt32) const;
 
 	/** Get Int64 value from Virtual Property */
 	UFUNCTION(BlueprintCallable, Category = "Remote Control Behaviour")
-	bool GetValueInt64(int64& OuyInt64);
+	bool GetValueInt64(int64& OuyInt64) const;
 
 	/** Get Uint64 value from Virtual Property */
 	UFUNCTION()
-	bool GetValueUint64(uint64& OuyUInt64);
+	bool GetValueUint64(uint64& OuyUInt64) const;
 
 	/** Get Float value from Virtual Property */
 	UFUNCTION(BlueprintCallable, Category = "Remote Control Behaviour")
-	bool GetValueFloat(float& OutFloat);
+	bool GetValueFloat(float& OutFloat) const;
 
 	/** Get Double value from Virtual Property */
 	UFUNCTION(BlueprintCallable, Category = "Remote Control Behaviour")
-	bool GetValueDouble(double& OutDouble);
+	bool GetValueDouble(double& OutDouble) const;
 
 	/** Get String value from Virtual Property */
 	UFUNCTION(BlueprintCallable, Category = "Remote Control Behaviour")
-	bool GetValueString(FString& OutStringValue);
+	bool GetValueString(FString& OutStringValue) const;
 
 	/** Get Name value from Virtual Property */
 	UFUNCTION(BlueprintCallable, Category = "Remote Control Behaviour")
-	bool GetValueName(FName& OutNameValue);
+	bool GetValueName(FName& OutNameValue) const;
 
 	/** Get Text value from Virtual Property */
 	UFUNCTION(BlueprintCallable, Category = "Remote Control Behaviour")
-	bool GetValueText(FText& OutTextValue);
+	bool GetValueText(FText& OutTextValue) const;
 
 	/** Get Numeric value from Virtual Property */
 	UFUNCTION(BlueprintCallable, Category = "Remote Control Behaviour")
-	bool GetValueNumericInteger(int64& OutInt64Value);
+	bool GetValueNumericInteger(int64& OutInt64Value) const;
 
 	/** Get Vector value from Virtual Property */
 	UFUNCTION(BlueprintCallable, Category = "Remote Control Behaviour")
-	bool GetValueVector(FVector& OutVector);
+	bool GetValueVector(FVector& OutVector) const;
 
 	/** Get Rotator value from Virtual Property */
 	UFUNCTION(BlueprintCallable, Category = "Remote Control Behaviour")
-	bool GetValueRotator(FRotator& OutRotator);
+	bool GetValueRotator(FRotator& OutRotator) const;
 
 	/** Get Color value from Virtual Property */
 	UFUNCTION(BlueprintCallable, Category = "Remote Control Behaviour")
-	bool GetValueColor(FColor& OutColor);
+	bool GetValueColor(FColor& OutColor) const;	
+
+	/** Infers correct type internally, fetches value from memory and returns the value as a string 
+	* that can be immediately used for dispaly (without needing to create a generic readonly property widget)
+	*/
+	UFUNCTION(BlueprintCallable, Category = "Remote Control Behaviour")
+	FString GetDisplayValueAsString() const;
 
 	/** Set Bool value from Virtual Property */
 	UFUNCTION(BlueprintCallable, Category = "Remote Control Behaviour")
@@ -227,7 +257,7 @@ public:
 	UFUNCTION(BlueprintCallable, Category = "Remote Control Behaviour")
 	bool SetValueColor(const FColor& InColor);
 	
-	/** Set FProperty Name */
+	/** Get FProperty Name */
 	UFUNCTION(BlueprintPure, Category = "Remote Control Behaviour")
 	FName GetPropertyName() const;
 
diff --git a/Engine/Plugins/VirtualProduction/RemoteControl/Source/RemoteControlLogic/Private/Action/RCActionContainer.cpp b/Engine/Plugins/VirtualProduction/RemoteControl/Source/RemoteControlLogic/Private/Action/RCActionContainer.cpp
--- a/Engine/Plugins/VirtualProduction/RemoteControl/Source/RemoteControlLogic/Private/Action/RCActionContainer.cpp
+++ b/Engine/Plugins/VirtualProduction/RemoteControl/Source/RemoteControlLogic/Private/Action/RCActionContainer.cpp
@@ -17,18 +17,53 @@ void URCActionContainer::ExecuteActions()
 	}
 }
 
-URCPropertyAction* URCActionContainer::AddAction(const TSharedPtr<FRemoteControlProperty> InRemoteControlProperty)
+TRCActionUniquenessTest URCActionContainer::GetDefaultActionUniquenessTest(const TSharedRef<const FRemoteControlField> InRemoteControlField)
 {
-	if (!InRemoteControlProperty.IsValid())
+	const FGuid FieldId = InRemoteControlField->GetId();
+
+	return [FieldId](const TSet<TObjectPtr<URCAction>>& InActions)
+	{
+		for (const URCAction* Action : InActions)
+		{
+			if (Action->ExposedFieldId == FieldId)
+			{
+				return false;
+			}
+		}
+
+		return true;
+	};
+}
+
+URCAction* URCActionContainer::AddAction(const TSharedRef<const FRemoteControlField> InRemoteControlField)
+{
+	return AddAction(GetDefaultActionUniquenessTest(InRemoteControlField), InRemoteControlField);
+}
+
+
+URCAction* URCActionContainer::AddAction(TRCActionUniquenessTest IsUnique, const TSharedRef<const FRemoteControlField> InRemoteControlField)
+{
+	if (!IsUnique(Actions))
 	{
 		return nullptr;
 	}
-	
-	if (FindActionByField(InRemoteControlProperty))
+
+	URCAction* NewAction = nullptr;
+
+	if (InRemoteControlField->FieldType == EExposedFieldType::Property)
 	{
-		return nullptr;
+		NewAction = AddPropertyAction(StaticCastSharedRef<const FRemoteControlProperty>(InRemoteControlField));
+	}
+	else if (InRemoteControlField->FieldType == EExposedFieldType::Function)
+	{
+		NewAction = AddFunctionAction(StaticCastSharedRef<const FRemoteControlFunction>(InRemoteControlField));
 	}
 
+	return NewAction;
+}
+
+URCPropertyAction* URCActionContainer::AddPropertyAction(const TSharedRef<const FRemoteControlProperty> InRemoteControlProperty)
+{
 	// Create new Property
 	URCPropertyAction* NewPropertyAction = NewObject<URCPropertyAction>(this);
 	NewPropertyAction->PresetWeakPtr = PresetWeakPtr;
@@ -47,18 +82,8 @@ URCPropertyAction* URCActionContainer::AddAction(const TSharedPtr<FRemoteControl
 	return NewPropertyAction;
 }
 
-URCFunctionAction* URCActionContainer::AddAction(const TSharedPtr<FRemoteControlFunction> InRemoteControlFunction)
+URCFunctionAction* URCActionContainer::AddFunctionAction(const TSharedRef<const FRemoteControlFunction> InRemoteControlFunction)
 {
-	if (!InRemoteControlFunction.IsValid())
-	{
-		return nullptr;
-	}
-
-	if (FindActionByField(InRemoteControlFunction))
-	{
-		return nullptr;
-	}
-	
 	// Create new Function Action
 	URCFunctionAction* NewFunctionAction = NewObject<URCFunctionAction>(this);
 	NewFunctionAction->PresetWeakPtr = PresetWeakPtr;
@@ -70,7 +95,7 @@ URCFunctionAction* URCActionContainer::AddAction(const TSharedPtr<FRemoteControl
 	return NewFunctionAction;
 }
 
-URCAction* URCActionContainer::FindActionByFieldId(const FGuid InId)
+URCAction* URCActionContainer::FindActionByFieldId(const FGuid InId) const
 {
 	for (URCAction* Action :  Actions)
 	{
@@ -83,17 +108,9 @@ URCAction* URCActionContainer::FindActionByFieldId(const FGuid InId)
 	return nullptr;
 }
 
-URCAction* URCActionContainer::FindActionByField(const TSharedPtr<FRemoteControlField> InRemoteControlField)
+URCAction* URCActionContainer::FindActionByField(const TSharedRef<const FRemoteControlField> InRemoteControlField) const
 {
-	for (URCAction* Action :  Actions)
-	{
-		if (Action->ExposedFieldId == InRemoteControlField->GetId())
-		{
-			return Action;
-		}
-	}
-
-	return nullptr;
+	return FindActionByFieldId(InRemoteControlField->GetId());
 }
 
 int32 URCActionContainer::RemoveAction(const FGuid InExposedFieldId)
diff --git a/Engine/Plugins/VirtualProduction/RemoteControl/Source/RemoteControlLogic/Private/Action/RCPropertyAction.cpp b/Engine/Plugins/VirtualProduction/RemoteControl/Source/RemoteControlLogic/Private/Action/RCPropertyAction.cpp
--- a/Engine/Plugins/VirtualProduction/RemoteControl/Source/RemoteControlLogic/Private/Action/RCPropertyAction.cpp
+++ b/Engine/Plugins/VirtualProduction/RemoteControl/Source/RemoteControlLogic/Private/Action/RCPropertyAction.cpp
@@ -57,3 +57,16 @@ void URCPropertyAction::Execute() const
 	
 	Super::Execute();
 }
+
+FName URCAction::GetExposedFieldLabel() const
+{
+	if (const URemoteControlPreset* Preset = PresetWeakPtr.Get())
+	{
+		if (TSharedPtr<const FRemoteControlField> RemoteControlField = Preset->GetExposedEntity<FRemoteControlField>(ExposedFieldId).Pin())
+		{
+			return RemoteControlField->GetLabel();
+		}
+	}
+
+	return NAME_None;
+}
\ No newline at end of file
diff --git a/Engine/Plugins/VirtualProduction/RemoteControl/Source/RemoteControlLogic/Private/Behaviour/Builtin/Conditional/RCBehaviourConditional.cpp b/Engine/Plugins/VirtualProduction/RemoteControl/Source/RemoteControlLogic/Private/Behaviour/Builtin/Conditional/RCBehaviourConditional.cpp
new file mode 100644
--- /dev/null
+++ b/Engine/Plugins/VirtualProduction/RemoteControl/Source/RemoteControlLogic/Private/Behaviour/Builtin/Conditional/RCBehaviourConditional.cpp
@@ -0,0 +1,186 @@
+// Copyright Epic Games, Inc. All Rights Reserved.
+
+#include "Behaviour/Builtin/Conditional/RCBehaviourConditional.h"
+
+#include "Action/RCAction.h"
+#include "Action/RCFunctionAction.h"
+#include "Action/RCPropertyAction.h"
+#include "Behaviour/Builtin/Conditional/RCBehaviourConditionalNode.h"
+#include "RCVirtualProperty.h"
+#include "RemoteControlField.h"
+#include "Controller/RCController.h"
+
+URCBehaviourConditional::URCBehaviourConditional()
+{
+}
+
+void URCBehaviourConditional::Initialize()
+{
+	Super::Initialize();
+}
+
+void URCBehaviourConditional::Execute()
+{
+	const URCBehaviourNode* BehaviourNode = GetBehaviourNode();
+	check(BehaviourNode);
+
+	if (BehaviourNode->GetClass() != URCBehaviourConditionalNode::StaticClass())
+	{
+		return; // Allow custom Blueprints to drive their own behaviour entirely
+	}
+
+	// Execute before the logic
+	BehaviourNode->PreExecute(this);
+
+	bool bConditionPass = false;
+
+	/* Consider the following sequence of conditions for a hypothetical Controller "Tricode" with value "Tri2"
+	*   =Tri1    <FALSE> (Action 1)       ...skip...
+	*   =Tri2   <TRUE>    (Action 2)   ...execute...
+	*   =Tri2   <TRUE>    (Action 3)   ...execute...
+	*   =Tri3   <FALSE>  (Action 4)       ...skip...
+	*   Else                                          ...skip...
+	* 
+	* For such multiple equality rows we want to know if at least one of them succeeded. 
+	* The flag bHasEqualitySuccess is used for determining whether Else should be executed 
+	*/
+	bool bHasEqualitySuccess = false;
+
+	for (TObjectPtr<URCAction> Action : ActionContainer->Actions)
+	{
+		FRCBehaviourCondition* Condition = Conditions.Find(Action);
+		if (!Condition)
+		{
+			ensureMsgf(false, TEXT("Unable to find condition for Action"));
+			continue;
+		}
+
+		const ERCBehaviourConditionType ConditionType = Condition->ConditionType;		
+
+		URCController* RCController = ControllerWeakPtr.Get();
+		if (RCController)
+		{
+			switch (ConditionType)
+			{
+			case ERCBehaviourConditionType::IsEqual:
+				bConditionPass = RCController->IsValueEqual(Condition->Comparand);
+				bHasEqualitySuccess |= bConditionPass;
+				break;			
+
+			case ERCBehaviourConditionType::IsGreaterThan:
+				bConditionPass = RCController->IsValueGreaterThan(Condition->Comparand);
+				break;
+
+			case ERCBehaviourConditionType::IsGreaterThanOrEqualTo:
+				bConditionPass = RCController->IsValueGreaterThanOrEqualTo(Condition->Comparand);
+				break;
+
+			case ERCBehaviourConditionType::IsLesserThan:
+				bConditionPass = RCController->IsValueLesserThan(Condition->Comparand);
+				break;
+
+			case ERCBehaviourConditionType::IsLesserThanOrEqualTo:
+				bConditionPass = RCController->IsValueLesserThanOrEqualTo(Condition->Comparand);
+				break;
+
+			case ERCBehaviourConditionType::Else:
+				// If the previous condition failed and no prior equality condition succeeded (among multiple =rows above) then execute Else!
+				bConditionPass = !bConditionPass && !bHasEqualitySuccess;
+				break;
+
+			default:				
+				ensureAlwaysMsgf(false, TEXT("Unimplemented comparator!"));
+			}
+
+			if (ConditionType != ERCBehaviourConditionType::IsEqual)
+			{
+				bHasEqualitySuccess = false; // Reset flag; either we reached an else clause (resolved above) or we moved from equality rows to a different comparison type
+			}
+		}
+
+		if (bConditionPass)
+		{
+			Action->Execute();
+			BehaviourNode->OnPassed(this);
+		}
+	}
+}
+
+void URCBehaviourConditional::OnActionAdded(URCAction* Action, const ERCBehaviourConditionType InConditionType, const TObjectPtr<URCVirtualPropertySelfContainer> InComparand)
+{
+	FRCBehaviourCondition Condition(InConditionType, InComparand);
+
+	Conditions.Add(Action, MoveTemp(Condition));	
+}
+
+URCAction* URCBehaviourConditional::AddAction(const TSharedRef<const FRemoteControlField> InRemoteControlField, const ERCBehaviourConditionType InConditionType, const TObjectPtr<URCVirtualPropertySelfContainer> InComparand)
+{
+	TRCActionUniquenessTest UniquenessTest = [this, InRemoteControlField, InConditionType, InComparand](const TSet<TObjectPtr<URCAction>>& Actions)
+	{
+		const FGuid FieldId = InRemoteControlField->GetId();
+
+		for (const URCAction* Action : Actions)
+		{
+			if (const FRCBehaviourCondition* Condition = this->Conditions.Find(Action))
+			{
+				const ERCBehaviourConditionType ActionCondition = Condition->ConditionType;
+
+				TObjectPtr<URCVirtualPropertySelfContainer> ActionComparand = Condition->Comparand;
+
+				if (Action->ExposedFieldId == FieldId && ActionCondition == InConditionType && ActionComparand->IsValueEqual(InComparand))
+				{
+					return false; // Not Unique!
+				}
+			}
+		}
+
+		return true;
+	};
+
+	return ActionContainer->AddAction(UniquenessTest, InRemoteControlField);
+}
+
+FText URCBehaviourConditional::GetConditionTypeAsText(ERCBehaviourConditionType ConditionType) const
+{
+	FText ConditionDisplayText;
+
+	switch (ConditionType)
+	{
+		case ERCBehaviourConditionType::IsEqual:
+		{
+			ConditionDisplayText = FText::FromName("=");
+			break;
+		}
+		case ERCBehaviourConditionType::IsGreaterThan:
+		{
+			ConditionDisplayText = FText::FromName(">");
+			break;
+		}
+		case ERCBehaviourConditionType::IsLesserThan:
+		{
+			ConditionDisplayText = FText::FromName("<");
+			break;
+		}
+		case ERCBehaviourConditionType::IsGreaterThanOrEqualTo:
+		{
+			ConditionDisplayText = FText::FromName(">=");
+			break;
+		}
+		case ERCBehaviourConditionType::IsLesserThanOrEqualTo:
+		{
+			ConditionDisplayText = FText::FromName("<=");
+			break;
+		}
+		case ERCBehaviourConditionType::Else:
+		{
+			ConditionDisplayText = FText::FromName("Else");
+			break;
+		}
+		default:
+		{
+			ensureMsgf(false, TEXT("Unknown condition type"));
+		}
+	}
+
+	return ConditionDisplayText;
+}
\ No newline at end of file
diff --git a/Engine/Plugins/VirtualProduction/RemoteControl/Source/RemoteControlLogic/Private/Behaviour/Builtin/Conditional/RCBehaviourConditionalNode.cpp b/Engine/Plugins/VirtualProduction/RemoteControl/Source/RemoteControlLogic/Private/Behaviour/Builtin/Conditional/RCBehaviourConditionalNode.cpp
new file mode 100644
--- /dev/null
+++ b/Engine/Plugins/VirtualProduction/RemoteControl/Source/RemoteControlLogic/Private/Behaviour/Builtin/Conditional/RCBehaviourConditionalNode.cpp
@@ -0,0 +1,47 @@
+// Copyright Epic Games, Inc. All Rights Reserved.
+
+#include "Behaviour/Builtin/Conditional/RCBehaviourConditionalNode.h"
+
+#include "Behaviour/Builtin/Conditional/RCBehaviourConditional.h"
+#include "Controller/RCController.h"
+#include "RCVirtualProperty.h"
+
+#define LOCTEXT_NAMESPACE "RemoteControlBehaviours"
+
+URCBehaviourConditionalNode::URCBehaviourConditionalNode()
+{
+	DisplayName = LOCTEXT("BehaviorNameConditional", "Conditional");
+	BehaviorDescription = LOCTEXT("BehaviorDescConditional", "If value matches a set condition, apply action list.");
+}
+
+bool URCBehaviourConditionalNode::Execute_Implementation(URCBehaviour* InBehaviour) const
+{
+	if (GetClass() != URCBehaviourConditionalNode::GetClass())
+	{
+		return true; // This is a custom blueprint node
+	}
+
+	// Unlike other Behaviours, Conditional Behaviour needs to evaluate each Action individually
+	// This logic is currently housed inside URCBehaviourConditional for convenience and flow should never come here
+	ensureAlwaysMsgf(false, TEXT("Unexpected flow for Conditional Behaviour! Control should never come here."));
+
+	return false;
+}
+
+bool URCBehaviourConditionalNode::IsSupported_Implementation(URCBehaviour* InBehaviour) const
+{
+	// Conditional is supported in general for all types that are configured in BaseRemoteControl.ini
+	// Additional filtering per condition type (eg: numeric comparators for numeric fields only) is performed at UI level based on the actively selected Comparator type
+	return true; 
+}
+
+UClass* URCBehaviourConditionalNode::GetBehaviourClass() const
+{
+	return URCBehaviourConditional::StaticClass();
+}
+
+void URCBehaviourConditionalNode::OnPassed_Implementation(URCBehaviour* InBehaviour) const
+{
+}
+
+#undef LOCTEXT_NAMESPACE
\ No newline at end of file
diff --git a/Engine/Plugins/VirtualProduction/RemoteControl/Source/RemoteControlLogic/Private/Behaviour/RCBehaviour.cpp b/Engine/Plugins/VirtualProduction/RemoteControl/Source/RemoteControlLogic/Private/Behaviour/RCBehaviour.cpp
--- a/Engine/Plugins/VirtualProduction/RemoteControl/Source/RemoteControlLogic/Private/Behaviour/RCBehaviour.cpp
+++ b/Engine/Plugins/VirtualProduction/RemoteControl/Source/RemoteControlLogic/Private/Behaviour/RCBehaviour.cpp
@@ -2,10 +2,12 @@
 
 #include "Behaviour/RCBehaviour.h"
 
-#include "Engine/Blueprint.h"
-
-#include "Action/RCActionContainer.h"
+#include "Action/RCAction.h"
+#include "Action/RCFunctionAction.h"
+#include "Action/RCPropertyAction.h"
 #include "Behaviour/RCBehaviourNode.h"
+#include "Engine/Blueprint.h"
+#include "RemoteControlField.h"
 
 URCBehaviour::URCBehaviour()
 {
@@ -30,11 +32,26 @@ void URCBehaviour::Execute()
 	}
 }
 
+URCAction* URCBehaviour::AddAction(const TSharedRef<const FRemoteControlField> InRemoteControlField)
+{
+	return ActionContainer->AddAction(InRemoteControlField);
+}
+
 int32 URCBehaviour::GetNumActions() const
 {
 	return ActionContainer->Actions.Num();
 }
 
+bool URCBehaviour::CanHaveActionForField(const TSharedPtr<FRemoteControlField> InRemoteControlField) const
+{
+	if (ActionContainer->FindActionByFieldId(InRemoteControlField->GetId()))
+	{
+		return false; // already exists!
+	}
+
+	return true;
+}
+
 UClass* URCBehaviour::GetOverrideBehaviourBlueprintClass() const
 {
 	// Blueprint extensions stores as a separate blueprints files that is why it should check can be path resolved as asset 
diff --git a/Engine/Plugins/VirtualProduction/RemoteControl/Source/RemoteControlLogic/Private/Tests/RemoteControlLogicTests.cpp b/Engine/Plugins/VirtualProduction/RemoteControl/Source/RemoteControlLogic/Private/Tests/RemoteControlLogicTests.cpp
--- a/Engine/Plugins/VirtualProduction/RemoteControl/Source/RemoteControlLogic/Private/Tests/RemoteControlLogicTests.cpp
+++ b/Engine/Plugins/VirtualProduction/RemoteControl/Source/RemoteControlLogic/Private/Tests/RemoteControlLogicTests.cpp
@@ -36,12 +36,12 @@ bool FRemoteControlLogicTest::RunTest(const FString& Parameters)
 	
 	// 2. Expose Fields
 	// 2.1 Expose Properties
-	const TSharedPtr<FRemoteControlProperty> RCProp1 = Preset->ExposeProperty(TestObject.Get(), FRCFieldPathInfo{GET_TEST_PROP(Color)->GetName()}).Pin();
+	const TSharedRef<FRemoteControlProperty> RCProp1 = Preset->ExposeProperty(TestObject.Get(), FRCFieldPathInfo{GET_TEST_PROP(Color)->GetName()}).Pin().ToSharedRef();
  
 	// 2.2 Expose Functions
 	// UObject* Object, UFunction* Function, FRemoteControlPresetExposeArgs Args
 	UFunction* TestIntFunction = TestObject->GetClass()->FindFunctionByName(GET_FUNCTION_NAME_CHECKED(URemoteControlLogicTestData, TestIntFunction));
-	const TSharedPtr<FRemoteControlFunction> RCFunc1 = Preset->ExposeFunction(TestObject.Get(), TestIntFunction).Pin();
+	const TSharedRef<FRemoteControlFunction> RCFunc1 = Preset->ExposeFunction(TestObject.Get(), TestIntFunction).Pin().ToSharedRef();
  
 	// 3. Add Controllers
 	// 3.1.1 Create Float Property
@@ -85,7 +85,7 @@ bool FRemoteControlLogicTest::RunTest(const FString& Parameters)
  
 	// 5 Add actions
 	// 5.1 Add Float Controller Actions
-	URCPropertyAction* FloatControllerBehaviourAction = FloatControllerBehaviour->ActionContainer->AddAction(RCProp1);
+	URCPropertyAction* FloatControllerBehaviourAction = Cast<URCPropertyAction> (FloatControllerBehaviour->ActionContainer->AddAction(RCProp1));
 
 	FColor OutColorValue = FColor();
 	TestTrue(TEXT("Should Get Color"), FloatControllerBehaviourAction->PropertySelfContainer->GetValueColor(OutColorValue) == true);
@@ -96,8 +96,8 @@ bool FRemoteControlLogicTest::RunTest(const FString& Parameters)
 	TestTrue(TEXT("Should Set Color"), FloatControllerBehaviourAction->PropertySelfContainer->SetValueColor(SecColorValue));
 	
 	// 5.4 Add Function Controller Actions
-	URCFunctionAction* FloatControllerBehaviourAction1 = FloatControllerBehaviour->ActionContainer->AddAction(RCFunc1);
-	URCPropertyAction* StrControllerBehaviourAction = StrControllerBehaviour->ActionContainer->AddAction(RCProp1);
+	URCFunctionAction* FloatControllerBehaviourAction1 = Cast<URCFunctionAction>(FloatControllerBehaviour->ActionContainer->AddAction(RCFunc1));
+	URCPropertyAction* StrControllerBehaviourAction = Cast<URCPropertyAction> (StrControllerBehaviour->ActionContainer->AddAction(RCProp1));
 	FColor StringControllerColorValue(7,8,9,10);
 	TestTrue(TEXT("Should Set Color"), StrControllerBehaviourAction->PropertySelfContainer->SetValueColor(StringControllerColorValue));
 	
diff --git a/Engine/Plugins/VirtualProduction/RemoteControl/Source/RemoteControlLogic/Public/Action/RCAction.h b/Engine/Plugins/VirtualProduction/RemoteControl/Source/RemoteControlLogic/Public/Action/RCAction.h
--- a/Engine/Plugins/VirtualProduction/RemoteControl/Source/RemoteControlLogic/Public/Action/RCAction.h
+++ b/Engine/Plugins/VirtualProduction/RemoteControl/Source/RemoteControlLogic/Public/Action/RCAction.h
@@ -19,6 +19,8 @@ public:
 	/** Execute action */
 	virtual void Execute() const {}
 
+	FName GetExposedFieldLabel() const;
+
 public:
 	/** Exposed Property or Function field Id*/
 	UPROPERTY()
@@ -31,4 +33,4 @@ public:
 	/** Reference to preset */
 	UPROPERTY()
 	TWeakObjectPtr<URemoteControlPreset> PresetWeakPtr;
-};
+};
\ No newline at end of file
diff --git a/Engine/Plugins/VirtualProduction/RemoteControl/Source/RemoteControlLogic/Public/Action/RCActionContainer.h b/Engine/Plugins/VirtualProduction/RemoteControl/Source/RemoteControlLogic/Public/Action/RCActionContainer.h
--- a/Engine/Plugins/VirtualProduction/RemoteControl/Source/RemoteControlLogic/Public/Action/RCActionContainer.h
+++ b/Engine/Plugins/VirtualProduction/RemoteControl/Source/RemoteControlLogic/Public/Action/RCActionContainer.h
@@ -16,6 +16,8 @@ struct FRemoteControlField;
 struct FRemoteControlFunction;
 struct FRemoteControlProperty;
 
+typedef TFunction<bool(const TSet<TObjectPtr<URCAction>>& Actions)> TRCActionUniquenessTest;
+
 /**
  * Container for created actions
  */
@@ -28,17 +30,19 @@ public:
 	/** Execute container actions */
 	void ExecuteActions();
 
+	TRCActionUniquenessTest GetDefaultActionUniquenessTest(const TSharedRef<const FRemoteControlField> InRemoteControlField);
+
 	/** Add remote control property action  */
-	URCPropertyAction* AddAction(const TSharedPtr<FRemoteControlProperty> InRemoteControlProperty);
+	URCAction* AddAction(const TSharedRef<const FRemoteControlField> InRemoteControlField);
 
-	/** Add remote control property function */
-	URCFunctionAction* AddAction(const TSharedPtr<FRemoteControlFunction> InRemoteControlFunction);
+	/** Add remote control property action  */
+	URCAction* AddAction(TRCActionUniquenessTest IsUnique, const TSharedRef<const FRemoteControlField> InRemoteControlField);
 
 	/** Find Action by given exposed filed id */
-	URCAction* FindActionByFieldId(const FGuid InFieldId);
+	URCAction* FindActionByFieldId(const FGuid InFieldId) const;
 
 	/** Find Action by given remote control field */
-	URCAction* FindActionByField(const TSharedPtr<FRemoteControlField> InRemoteControlField);
+	URCAction* FindActionByField(const TSharedRef<const FRemoteControlField> InRemoteControlField) const;
 	
 	/** Remove action by exposed field Id */
 	virtual int32 RemoveAction(const FGuid InExposedFieldId);
@@ -49,7 +53,6 @@ public:
 	/** Empty action set */
 	void EmptyActions();
 
-public:
 	/** Set of actions */
 	UPROPERTY()
 	TSet<TObjectPtr<URCAction>> Actions;
@@ -61,4 +64,11 @@ public:
 	/** Reference to Preset */
 	UPROPERTY()
 	TWeakObjectPtr<URemoteControlPreset> PresetWeakPtr;
+
+private:
+	/** Add remote control property action  */
+	URCPropertyAction* AddPropertyAction(const TSharedRef<const FRemoteControlProperty> InRemoteControlProperty);
+
+	/** Add remote control property function */
+	URCFunctionAction* AddFunctionAction(const TSharedRef<const FRemoteControlFunction> InRemoteControlFunction);
 };
diff --git a/Engine/Plugins/VirtualProduction/RemoteControl/Source/RemoteControlLogic/Public/Behaviour/Builtin/Conditional/RCBehaviourConditional.h b/Engine/Plugins/VirtualProduction/RemoteControl/Source/RemoteControlLogic/Public/Behaviour/Builtin/Conditional/RCBehaviourConditional.h
new file mode 100644
--- /dev/null
+++ b/Engine/Plugins/VirtualProduction/RemoteControl/Source/RemoteControlLogic/Public/Behaviour/Builtin/Conditional/RCBehaviourConditional.h
@@ -0,0 +1,93 @@
+// Copyright Epic Games, Inc. All Rights Reserved.
+
+#pragma once
+
+#include "Action/RCActionContainer.h"
+#include "Behaviour/RCBehaviour.h"
+
+#include "RCBehaviourConditional.generated.h"
+
+class URCVirtualPropertySelfContainer;
+class URCAction;
+
+/**
+ * Condition Typess
+ */
+UENUM()
+enum class ERCBehaviourConditionType : uint8
+{
+	IsEqual,
+	IsGreaterThan,
+	IsLesserThan,
+	IsGreaterThanOrEqualTo,
+	IsLesserThanOrEqualTo,
+	Else
+};
+
+/**
+ * Struct representing a single Condition
+ * 
+ * Serialization and GC support are automatically provided by standard container types, as all relevant fields have been marked as UPROPERTY()
+ */
+USTRUCT()
+struct FRCBehaviourCondition
+{
+	GENERATED_BODY()
+
+	FRCBehaviourCondition() {}
+	FRCBehaviourCondition(const ERCBehaviourConditionType InConditionType, const TObjectPtr<URCVirtualPropertySelfContainer> InComparand)
+		: ConditionType(InConditionType), Comparand(InComparand) {}
+
+	/** The type of condition to be used for comparison*/
+	UPROPERTY()
+	ERCBehaviourConditionType ConditionType;
+
+	/** The value with which to compare */
+	UPROPERTY()
+	TObjectPtr<URCVirtualPropertySelfContainer> Comparand;
+};
+
+/**
+ * [Conditional Behaviour]
+ * 
+ * Supports unique conditions per action that allow powerful Logic & Rigging of multiple usecases
+ * including Tricode based multi-mapping (selector behaviour), simulation of Hide On Empty behaviours, etc.
+ */
+UCLASS()
+class REMOTECONTROLLOGIC_API URCBehaviourConditional : public URCBehaviour
+{
+	GENERATED_BODY()
+
+public:
+	URCBehaviourConditional();
+	
+	//~ Begin URCBehaviour interface
+	virtual void Initialize() override;
+
+	/** Execute the behaviour */
+	virtual void Execute() override;	
+
+	//~ End URCBehaviour interface
+
+	/** ~ OnActionAdded ~
+	* Invoked after a new condition action has just been created.
+	* Used to link the newly created action with its condition data */
+	void OnActionAdded(URCAction* Action, const ERCBehaviourConditionType InConditionType, const TObjectPtr<URCVirtualPropertySelfContainer> InComparand);
+
+	/** Add a Logic action using a remote control field as input */
+	URCAction* AddAction(const TSharedRef<const FRemoteControlField> InRemoteControlField, const ERCBehaviourConditionType InConditionType, const TObjectPtr<URCVirtualPropertySelfContainer> InComparand);
+
+	/** Whether we can create an action pertaining to a given remote control field for the current behaviour */
+	virtual bool CanHaveActionForField(const TSharedPtr<FRemoteControlField> InRemoteControlField) const override
+	{
+		return true; // Conditional Behaviour can support multiple Action rows for a single Remote Control Field
+	}
+
+	/** User-friendly text representation of a condition enum. Used to display comparator info in the Actions table */
+	FText GetConditionTypeAsText(ERCBehaviourConditionType ConditionType) const;
+	
+	/** Data storage for Actions and related Conditions; stored as a mapping of Action object and associated condition data
+	*  Each Action is associated with a unique condition (for the Conditional Behaviour) */
+	UPROPERTY()
+	TMap<URCAction*, FRCBehaviourCondition> Conditions;
+};
diff --git a/Engine/Plugins/VirtualProduction/RemoteControl/Source/RemoteControlLogic/Public/Behaviour/Builtin/Conditional/RCBehaviourConditionalNode.h b/Engine/Plugins/VirtualProduction/RemoteControl/Source/RemoteControlLogic/Public/Behaviour/Builtin/Conditional/RCBehaviourConditionalNode.h
new file mode 100644
--- /dev/null
+++ b/Engine/Plugins/VirtualProduction/RemoteControl/Source/RemoteControlLogic/Public/Behaviour/Builtin/Conditional/RCBehaviourConditionalNode.h
@@ -0,0 +1,32 @@
+﻿// Copyright Epic Games, Inc. All Rights Reserved.
+
+#pragma once
+
+#include "Behaviour/RCBehaviourNode.h"
+#include "RCBehaviourConditionalNode.generated.h"
+
+/**
+ * Compares given property value with controller property value
+ */
+UCLASS()
+class REMOTECONTROLLOGIC_API URCBehaviourConditionalNode : public URCBehaviourNode
+{
+	GENERATED_BODY()
+
+public:
+	
+	URCBehaviourConditionalNode();
+
+	//~ Begin URCBehaviourNode interface
+	UFUNCTION(BlueprintNativeEvent, BlueprintCallable, Category = "Remote Control Behaviour")
+	bool Execute(URCBehaviour* InBehaviour) const override;
+
+	UFUNCTION(BlueprintNativeEvent, BlueprintCallable, Category = "Remote Control Behaviour")
+	bool IsSupported(URCBehaviour* InBehaviour) const override;
+
+	UFUNCTION(BlueprintNativeEvent, BlueprintCallable, Category = "Remote Control Behaviour")
+	void OnPassed(URCBehaviour* InBehaviour) const;
+
+	UClass* GetBehaviourClass() const override;
+	//~ End URCBehaviourNode interface
+};
diff --git a/Engine/Plugins/VirtualProduction/RemoteControl/Source/RemoteControlLogic/Public/Behaviour/RCBehaviour.h b/Engine/Plugins/VirtualProduction/RemoteControl/Source/RemoteControlLogic/Public/Behaviour/RCBehaviour.h
--- a/Engine/Plugins/VirtualProduction/RemoteControl/Source/RemoteControlLogic/Public/Behaviour/RCBehaviour.h
+++ b/Engine/Plugins/VirtualProduction/RemoteControl/Source/RemoteControlLogic/Public/Behaviour/RCBehaviour.h
@@ -2,14 +2,19 @@
 
 #pragma once
 
+#include "Action/RCActionContainer.h"
 #include "CoreMinimal.h"
 #include "UObject/Object.h"
 #include "Templates/SubclassOf.h"
 #include "RCBehaviour.generated.h"
 
+struct FRemoteControlField;
 class URCBehaviourNode;
 class URCController;
+class URCFunctionAction;
+class URCPropertyAction;
 class URemoteControlPreset;
+class URCAction;
 class URCActionContainer;
 class URCVirtualProperty;
 
@@ -36,9 +41,15 @@ public:
 	/** Execute the behaviour */
 	virtual void Execute();
 
+	/** Add a Logic action using a remote control field as input */
+	URCAction* AddAction(const TSharedRef<const FRemoteControlField> InRemoteControlField);
+
 	/** Get number of action associated with behaviour */
 	int32 GetNumActions() const;
 
+	/** Whether we can create an action pertaining to a given remote control field for the current behaviour */
+	virtual bool CanHaveActionForField(const TSharedPtr<FRemoteControlField> InRemoteControlField) const;
+
 	/**
 	 * Return blueprint class associated with behaviour if exists
 	 */
@@ -64,7 +75,7 @@ public:
 	const FText& GetBehaviorDescription();
 #endif
 
-private:
+protected:
 	/**
 	 * It created the node if it called first time
 	 * If BehaviourNodeClass changes it creates new instance
diff --git a/Engine/Plugins/VirtualProduction/RemoteControl/Source/RemoteControlUI/Private/RemoteControlUIModule.cpp b/Engine/Plugins/VirtualProduction/RemoteControl/Source/RemoteControlUI/Private/RemoteControlUIModule.cpp
--- a/Engine/Plugins/VirtualProduction/RemoteControl/Source/RemoteControlUI/Private/RemoteControlUIModule.cpp
+++ b/Engine/Plugins/VirtualProduction/RemoteControl/Source/RemoteControlUI/Private/RemoteControlUIModule.cpp
@@ -28,6 +28,7 @@
 #include "Textures/SlateIcon.h"
 #include "UI/Action/SRCActionPanel.h"
 #include "UI/Behaviour/SRCBehaviourPanel.h"
+#include "UI/Behaviour/Builtin/Conditional/SRCBehaviourConditional.h"
 #include "UI/Customizations/FPassphraseCustomization.h"
 #include "UI/Customizations/RemoteControlEntityCustomization.h"
 #include "UI/SRCPanelExposedField.h"
@@ -274,6 +275,7 @@ void FRemoteControlUIModule::ShutdownModule()
 	SRemoteControlPanel::Shutdown();
 	SRCActionPanel::Shutdown();
 	SRCBehaviourPanel::Shutdown();
+	SRCBehaviourConditional::Shutdown();
 }
 
 FDelegateHandle FRemoteControlUIModule::AddPropertyFilter(FOnDisplayExposeIcon OnDisplayExposeIcon)
diff --git a/Engine/Plugins/VirtualProduction/RemoteControl/Source/RemoteControlUI/Private/UI/Action/Conditional/RCActionConditionalModel.cpp b/Engine/Plugins/VirtualProduction/RemoteControl/Source/RemoteControlUI/Private/UI/Action/Conditional/RCActionConditionalModel.cpp
new file mode 100644
--- /dev/null
+++ b/Engine/Plugins/VirtualProduction/RemoteControl/Source/RemoteControlUI/Private/UI/Action/Conditional/RCActionConditionalModel.cpp
@@ -0,0 +1,175 @@
+// Copyright Epic Games, Inc. All Rights Reserved.
+
+#include "RCActionConditionalModel.h"
+
+#include "Behaviour/Builtin/Conditional/RCBehaviourConditional.h"
+#include "Controller/RCController.h"
+#include "IDetailTreeNode.h"
+#include "IPropertyRowGenerator.h"
+#include "RCVirtualProperty.h"
+#include "RCVirtualPropertyContainer.h"
+#include "RemoteControlField.h"
+#include "RemoteControlPreset.h"
+#include "Styling/RemoteControlStyles.h"
+#include "UI/Behaviour/Builtin/Conditional/RCBehaviourConditionalModel.h"
+#include "UI/RCUIHelpers.h"
+#include "UI/RemoteControlPanelStyle.h"
+#include "UI/SRCPanelDragHandle.h"
+#include "UI/SRCPanelExposedEntity.h"
+#include "Widgets/Layout/SBox.h"
+#include "Widgets/Views/SHeaderRow.h"
+#include "Widgets/Views/SListView.h"
+#include "Widgets/Views/STableRow.h"
+
+#define LOCTEXT_NAMESPACE "RCActionModel"
+
+namespace UE::RCActionPanelConditionalList
+{
+	namespace Columns
+	{
+		const FName VariableColor = TEXT("VariableColor");
+		const FName DragDropHandle = TEXT("DragDropHandle");
+		const FName Condition = TEXT("Condition");
+		const FName Description = TEXT("Description");
+		const FName Value = TEXT("Value");
+	}
+
+	class SActionItemListRow : public SMultiColumnTableRow<TSharedRef<FRCActionModel>>
+	{
+	public:
+		void Construct(const FTableRowArgs& InArgs, const TSharedRef<STableViewBase>& OwnerTableView, TSharedRef<FRCActionConditionalModel> InActionItem)
+		{
+			ActionItem = InActionItem;
+			FSuperRowType::Construct(InArgs, OwnerTableView);
+		}
+
+		TSharedRef<SWidget> GenerateWidgetForColumn(const FName& ColumnName) override
+		{
+			if (!ensure(ActionItem.IsValid()))
+				return SNullWidget::NullWidget;
+
+			if (ColumnName == UE::RCActionPanelConditionalList::Columns::VariableColor)
+			{
+				return ActionItem->GetVariableColorWidget();
+			}
+			else if (ColumnName == UE::RCActionPanelConditionalList::Columns::Condition)
+			{
+				return ActionItem->GetConditionWidget();
+			}
+			else if (ColumnName == UE::RCActionPanelConditionalList::Columns::Description)
+			{
+				return ActionItem->GetNameWidget();
+			}
+			else if (ColumnName == UE::RCActionPanelConditionalList::Columns::Value)
+			{
+				return ActionItem->GetWidget();
+			}
+
+			// @todo: Implement Drag-Drop-handle column with Actions reordering support
+
+			return SNullWidget::NullWidget;
+		}
+
+	private:
+		TSharedPtr<FRCActionConditionalModel> ActionItem;
+	};
+}
+
+TSharedPtr<SHeaderRow> FRCActionConditionalModel::GetHeaderRow()
+{
+	const FRCPanelStyle* RCPanelStyle = &FRemoteControlPanelStyle::Get()->GetWidgetStyle<FRCPanelStyle>("RemoteControlPanel.MinorPanel");
+
+	return SNew(SHeaderRow)
+		.Style(&RCPanelStyle->HeaderRowStyle)
+
+		+ SHeaderRow::Column(UE::RCActionPanelConditionalList::Columns::VariableColor)
+		.DefaultLabel(LOCTEXT("RCActionVariableColorColumnHeader", ""))
+		.FixedWidth(5.f)
+		.HeaderContentPadding(RCPanelStyle->HeaderRowPadding)
+
+		+ SHeaderRow::Column(UE::RCActionPanelConditionalList::Columns::DragDropHandle)
+		.DefaultLabel(LOCTEXT("RCActionDragDropHandleColumnHeader", ""))
+		.FixedWidth(25.f)
+		.HeaderContentPadding(RCPanelStyle->HeaderRowPadding)
+
+		+ SHeaderRow::Column(UE::RCActionPanelConditionalList::Columns::Condition)
+		.DefaultLabel(LOCTEXT("RCActionConditionColumnHeader", "Condition"))
+		.FillWidth(0.2f)
+		.HeaderContentPadding(RCPanelStyle->HeaderRowPadding)
+
+		+ SHeaderRow::Column(UE::RCActionPanelConditionalList::Columns::Description)
+		.DefaultLabel(LOCTEXT("RCActionDescColumnHeader", "Description"))
+		.FillWidth(0.4f)
+		.HeaderContentPadding(RCPanelStyle->HeaderRowPadding)
+
+		+ SHeaderRow::Column(UE::RCActionPanelConditionalList::Columns::Value)
+		.DefaultLabel(LOCTEXT("RCActionValueColumnHeader", "Value"))
+		.FillWidth(0.4f)
+		.HeaderContentPadding(RCPanelStyle->HeaderRowPadding);
+}
+
+FRCActionConditionalModel::FRCActionConditionalModel(URCAction* InAction, const TSharedPtr<class FRCBehaviourModel> InBehaviourItem, const TSharedPtr<SRemoteControlPanel> InRemoteControlPanel)
+	: FRCActionModel(InAction, InBehaviourItem, InRemoteControlPanel)
+{
+}
+
+TSharedPtr<FRCActionConditionalModel> FRCActionConditionalModel::GetModelByActionType(URCAction* InAction, const TSharedPtr<class FRCBehaviourModel> InBehaviourItem, const TSharedPtr<SRemoteControlPanel> InRemoteControlPanel)
+{
+	if (URCPropertyAction* PropertyAction = Cast<URCPropertyAction>(InAction))
+	{
+		return MakeShared<FRCPropertyActionConditionalModel>(PropertyAction, InBehaviourItem, InRemoteControlPanel);
+	}
+	else if (URCFunctionAction* FunctionAction = Cast<URCFunctionAction>(InAction))
+	{
+		return MakeShared<FRCFunctionActionConditionalModel>(FunctionAction, InBehaviourItem, InRemoteControlPanel);
+	}
+
+	return nullptr;
+}
+
+TSharedRef<SWidget> FRCActionConditionalModel::GetConditionWidget() const
+{
+	if (TSharedPtr<FRCBehaviourConditionalModel> BehaviourItem = StaticCastSharedPtr<FRCBehaviourConditionalModel>(BehaviourItemWeakPtr.Pin()))
+	{
+		if (URCBehaviourConditional* Behaviour = Cast<URCBehaviourConditional>(BehaviourItem->GetBehaviour()))
+		{
+			if (FRCBehaviourCondition* Condition = Behaviour->Conditions.Find(GetAction()))
+			{
+				const FText ConditionText = Behaviour->GetConditionTypeAsText(Condition->ConditionType);
+
+				FText ConditionComparandText;
+
+				if (Condition->ConditionType == ERCBehaviourConditionType::Else)
+				{
+					ConditionComparandText = ConditionText; // Else doesn't need to display the comparand
+				}
+				else
+				{
+					const FText ComparandValueText = FText::FromName(*Condition->Comparand->GetDisplayValueAsString());
+
+					ConditionComparandText = FText::Format(FText::FromName("{0} {1}"), ConditionText, ComparandValueText);
+				}
+
+				return SNew(SBox)
+					.Padding(FMargin(6.f))
+					[
+						SNew(STextBlock).Text(ConditionComparandText)
+					];
+			}
+		}
+	}
+	
+
+	return SNullWidget::NullWidget;
+}
+
+TSharedRef<ITableRow> FRCActionConditionalModel::OnGenerateWidgetForList(TSharedPtr<FRCActionConditionalModel> InItem, const TSharedRef<STableViewBase>& OwnerTable)
+{
+	typedef UE::RCActionPanelConditionalList::SActionItemListRow ActionRowType;
+
+	return SNew(ActionRowType, OwnerTable, InItem.ToSharedRef())
+		.Style(&RCPanelStyle->TableRowStyle)
+		.Padding(FMargin(3.f));
+}
+
+#undef LOCTEXT_NAMESPACE
\ No newline at end of file
diff --git a/Engine/Plugins/VirtualProduction/RemoteControl/Source/RemoteControlUI/Private/UI/Action/Conditional/RCActionConditionalModel.h b/Engine/Plugins/VirtualProduction/RemoteControl/Source/RemoteControlUI/Private/UI/Action/Conditional/RCActionConditionalModel.h
new file mode 100644
--- /dev/null
+++ b/Engine/Plugins/VirtualProduction/RemoteControl/Source/RemoteControlUI/Private/UI/Action/Conditional/RCActionConditionalModel.h
@@ -0,0 +1,80 @@
+// Copyright Epic Games, Inc. All Rights Reserved.
+
+#pragma once
+
+#include "Action/RCFunctionAction.h"
+#include "Action/RCPropertyAction.h"
+#include "UI/Action/RCActionModel.h"
+
+/* 
+* ~ FRCActionConditionalModel ~
+*
+* UI model for representing an Action of the Conditional Behaviour.
+* Contains a Condition widget for representing the comparand/comparator and target property/value widgets
+* 
+*                              <Condition>                   <Description>         <Value>
+* 
+*  Color coding | Comparator+Comparand | Target Property | Value To Set
+*/
+class FRCActionConditionalModel : public FRCActionModel
+{
+public:
+	FRCActionConditionalModel(URCAction* InAction, const TSharedPtr<class FRCBehaviourModel> InBehaviourItem, const TSharedPtr<SRemoteControlPanel> InRemoteControlPanel);	
+
+	/** Widget representing the Condition (Comparator|Comparand) used
+	* to trigger the corresponding action
+	*/
+	virtual TSharedRef<SWidget> GetConditionWidget() const;
+
+	/** OnGenerateRow delegate for the Actions List View*/
+	TSharedRef<ITableRow> OnGenerateWidgetForList(TSharedPtr<FRCActionConditionalModel> InItem, const TSharedRef<STableViewBase>& OwnerTable);
+
+	/** Get the Header Row appropriate for this particular Action model */
+	static TSharedPtr<SHeaderRow> GetHeaderRow();
+
+	/** Chooses the appropriate Action model for the current class and field type*/
+	static TSharedPtr<FRCActionConditionalModel> GetModelByActionType(URCAction* InAction, const TSharedPtr<class FRCBehaviourModel> InBehaviourItem, const TSharedPtr<SRemoteControlPanel> InRemoteControlPanel);
+};
+
+/*
+* ~ FRCPropertyActionConditionalModel ~
+*
+* UI model for Property based Conditional Actions
+*/
+class FRCPropertyActionConditionalModel : public FRCActionConditionalModel, public FRCPropertyActionType
+{
+public:
+	FRCPropertyActionConditionalModel(URCPropertyAction* InPropertyAction, const TSharedPtr<class FRCBehaviourModel> InBehaviourItem, const TSharedPtr<SRemoteControlPanel> InRemoteControlPanel)
+	: FRCActionConditionalModel(InPropertyAction, InBehaviourItem, InRemoteControlPanel),
+		FRCPropertyActionType(InPropertyAction) { }
+
+	virtual TSharedRef<SWidget> GetWidget() const override
+	{
+		return GetPropertyWidget();
+	}
+
+	/** Color code for this Action*/
+	virtual FLinearColor GetActionTypeColor() const override
+	{
+		return GetPropertyTypeColor();
+	}
+};
+
+/*
+* ~ FRCFunctionActionConditionalModel ~
+*
+* UI model for Function based Conditional Actions
+*/
+class FRCFunctionActionConditionalModel : public FRCActionConditionalModel, public FRCFunctionActionType
+{
+public:
+	FRCFunctionActionConditionalModel(URCFunctionAction* InFunctionAction, const TSharedPtr<class FRCBehaviourModel> InBehaviourItem, const TSharedPtr<SRemoteControlPanel> InRemoteControlPanel)
+		: FRCActionConditionalModel(InFunctionAction, InBehaviourItem, InRemoteControlPanel)
+		, FRCFunctionActionType(InFunctionAction) { }
+
+	/** Color code for this Action*/
+	virtual FLinearColor GetActionTypeColor() const override
+	{
+		return GetFunctionTypeColor();
+	}
+};
\ No newline at end of file
diff --git a/Engine/Plugins/VirtualProduction/RemoteControl/Source/RemoteControlUI/Private/UI/Action/RCActionModel.cpp b/Engine/Plugins/VirtualProduction/RemoteControl/Source/RemoteControlUI/Private/UI/Action/RCActionModel.cpp
--- a/Engine/Plugins/VirtualProduction/RemoteControl/Source/RemoteControlUI/Private/UI/Action/RCActionModel.cpp
+++ b/Engine/Plugins/VirtualProduction/RemoteControl/Source/RemoteControlUI/Private/UI/Action/RCActionModel.cpp
@@ -2,8 +2,6 @@
 
 #include "RCActionModel.h"
 
-#include "Action/RCFunctionAction.h"
-#include "Action/RCPropertyAction.h"
 #include "Controller/RCController.h"
 #include "IDetailTreeNode.h"
 #include "IPropertyRowGenerator.h"
@@ -11,14 +9,114 @@
 #include "RCVirtualPropertyContainer.h"
 #include "RemoteControlField.h"
 #include "RemoteControlPreset.h"
+#include "Styling/RemoteControlStyles.h"
 #include "UI/RCUIHelpers.h"
 #include "UI/RemoteControlPanelStyle.h"
 #include "UI/SRCPanelDragHandle.h"
 #include "UI/SRCPanelExposedEntity.h"
+#include "Widgets/Layout/SBox.h"
+#include "Widgets/Views/SHeaderRow.h"
+#include "Widgets/Views/SListView.h"
+#include "Widgets/Views/STableRow.h"
 
-FRCActionModel::FRCActionModel(URCAction* InAction)
-	: ActionWeakPtr(InAction)
+#define LOCTEXT_NAMESPACE "RCActionModel"
+
+namespace UE::RCActionPanelList
 {
+	namespace Columns
+	{
+		const FName VariableColor = TEXT("VariableColor");
+		const FName DragDropHandle = TEXT("DragDropHandle");
+		const FName Description = TEXT("Description");
+		const FName Value = TEXT("Value");
+	}
+
+	class SActionItemListRow : public SMultiColumnTableRow<TSharedRef<FRCActionModel>>
+	{
+	public:
+		void Construct(const FTableRowArgs& InArgs, const TSharedRef<STableViewBase>& OwnerTableView, TSharedRef<FRCActionModel> InActionItem)
+		{
+			ActionItem = InActionItem;
+			FSuperRowType::Construct(InArgs, OwnerTableView);
+		}
+
+		TSharedRef<SWidget> GenerateWidgetForColumn(const FName& ColumnName) override
+		{
+			if (!ensure(ActionItem.IsValid()))
+				return SNullWidget::NullWidget;
+
+			if (ColumnName == UE::RCActionPanelList::Columns::VariableColor)
+			{
+				return ActionItem->GetVariableColorWidget();
+			}
+			else if (ColumnName == UE::RCActionPanelList::Columns::Description)
+			{
+				return ActionItem->GetNameWidget();
+			}
+			else if (ColumnName == UE::RCActionPanelList::Columns::Value)
+			{
+				return ActionItem->GetWidget();
+			}
+
+			// @todo: Implement Drag-Drop-handle column with Actions reordering support
+
+			return SNullWidget::NullWidget;
+		}
+
+	private:
+		TSharedPtr<FRCActionModel> ActionItem;
+	};
+}
+
+
+FRCActionModel::FRCActionModel(URCAction* InAction, const TSharedPtr<class FRCBehaviourModel> InBehaviourItem, const TSharedPtr<SRemoteControlPanel> InRemoteControlPanel)
+	: FRCLogicModeBase(InRemoteControlPanel),
+	ActionWeakPtr(InAction)
+{
+	BehaviourItemWeakPtr = InBehaviourItem;
+
+	RCPanelStyle = &FRemoteControlPanelStyle::Get()->GetWidgetStyle<FRCPanelStyle>("RemoteControlPanel.MinorPanel");
+}
+
+TSharedRef<SWidget> FRCActionModel::GetNameWidget() const
+{
+	if(URCAction* Action = GetAction())
+	{
+		if (URemoteControlPreset* Preset = GetPreset())
+		{
+			if (const TSharedPtr<FRemoteControlField> RemoteControlField = Preset->GetExposedEntity<FRemoteControlField>(Action->ExposedFieldId).Pin())
+			{
+				return SNew(SBox)
+					.Padding(FMargin(6.f))
+					[
+						SNew(STextBlock).Text(FText::FromName(RemoteControlField->GetLabel()))
+					];
+			}
+		}
+	}
+
+	return SNullWidget::NullWidget;
+}
+
+TSharedRef<SWidget> FRCActionModel::GetWidget() const
+{
+	return SNullWidget::NullWidget;
+}
+
+TSharedRef<SWidget> FRCActionModel::GetVariableColorWidget() const
+{
+	const FLinearColor TypeColor = GetActionTypeColor();
+
+	// Variable Color Bar
+	return SNew(SBox)
+		.HeightOverride(5.f)
+		[
+			SNew(SBorder)
+			.Visibility(EVisibility::HitTestInvisible)
+			.BorderImage(FAppStyle::Get().GetBrush("NumericEntrySpinBox.NarrowDecorator"))
+			.BorderBackgroundColor(TypeColor)
+			.Padding(FMargin(5.0f, 0.0f, 0.0f, 0.f))
+		];
 }
 
 URCAction* FRCActionModel::GetAction() const
@@ -26,10 +124,61 @@ URCAction* FRCActionModel::GetAction() const
 	return ActionWeakPtr.Get();
 }
 
-FRCPropertyActionModel::FRCPropertyActionModel(URCPropertyAction* InPropertyAction)
-	: FRCActionModel(InPropertyAction)
-	, PropertyActionWeakPtr(InPropertyAction)
+TSharedRef<ITableRow> FRCActionModel::OnGenerateWidgetForList(TSharedPtr<FRCActionModel> InItem, const TSharedRef<STableViewBase>& OwnerTable)
+{
+	typedef UE::RCActionPanelList::SActionItemListRow ActionRowType;
+
+	return SNew(ActionRowType, OwnerTable, InItem.ToSharedRef())
+		.Style(&RCPanelStyle->TableRowStyle)
+		.Padding(FMargin(3.f));
+}
+
+TSharedPtr<SHeaderRow> FRCActionModel::GetHeaderRow()
+{
+	const FRCPanelStyle* RCPanelStyle = &FRemoteControlPanelStyle::Get()->GetWidgetStyle<FRCPanelStyle>("RemoteControlPanel.MinorPanel");
+
+	return SNew(SHeaderRow)
+		.Style(&RCPanelStyle->HeaderRowStyle)
+
+		+ SHeaderRow::Column(UE::RCActionPanelList::Columns::VariableColor)
+		.DefaultLabel(LOCTEXT("RCActionVariableColorColumnHeader", ""))
+		.FixedWidth(5.f)
+		.HeaderContentPadding(RCPanelStyle->HeaderRowPadding)
+
+		+ SHeaderRow::Column(UE::RCActionPanelList::Columns::DragDropHandle)
+		.DefaultLabel(LOCTEXT("RCActionDragDropHandleColumnHeader", ""))
+		.FixedWidth(25.f)
+		.HeaderContentPadding(RCPanelStyle->HeaderRowPadding)
+
+		+ SHeaderRow::Column(UE::RCActionPanelList::Columns::Description)
+		.DefaultLabel(LOCTEXT("RCActionDescColumnHeader", "Description"))
+		.FillWidth(0.5f)
+		.HeaderContentPadding(RCPanelStyle->HeaderRowPadding)
+
+		+ SHeaderRow::Column(UE::RCActionPanelList::Columns::Value)
+		.DefaultLabel(LOCTEXT("RCActionValueColumnHeader", "Value"))
+		.FillWidth(0.5f)
+		.HeaderContentPadding(RCPanelStyle->HeaderRowPadding);
+}
+
+TSharedPtr<FRCActionModel> FRCActionModel::GetModelByActionType(URCAction* InAction, const TSharedPtr<class FRCBehaviourModel> InBehaviourItem, const TSharedPtr<SRemoteControlPanel> InRemoteControlPanel)
 {
+	if (URCPropertyAction* PropertyAction = Cast<URCPropertyAction>(InAction))
+	{
+		return MakeShared<FRCPropertyActionModel>(PropertyAction, InBehaviourItem, InRemoteControlPanel);
+	}
+	else if (URCFunctionAction* FunctionAction = Cast<URCFunctionAction>(InAction))
+	{
+		return MakeShared<FRCFunctionActionModel>(FunctionAction, InBehaviourItem, InRemoteControlPanel);
+	}
+	else
+		return nullptr;
+}
+
+FRCPropertyActionType::FRCPropertyActionType(URCPropertyAction* InPropertyAction)
+{
+	PropertyActionWeakPtr = InPropertyAction;
+
 	FPropertyRowGeneratorArgs Args;
 	Args.bShouldShowHiddenProperties = true;
 	PropertyRowGenerator = FModuleManager::GetModuleChecked<FPropertyEditorModule>("PropertyEditor").CreatePropertyRowGenerator(Args);
@@ -54,19 +203,26 @@ FRCPropertyActionModel::FRCPropertyActionModel(URCPropertyAction* InPropertyActi
 	}
 }
 
-const FName& FRCPropertyActionModel::GetPropertyName() const
+const FName& FRCPropertyActionType::GetPropertyName() const
 {
 	return PropertyActionWeakPtr.Get()->PropertySelfContainer->PropertyName;
 }
 
-TSharedRef<SWidget> FRCPropertyActionModel::GetWidget() const
+TSharedRef<SWidget> FRCPropertyActionType::GetPropertyWidget() const
 {
+	if (!ensure(PropertyActionWeakPtr.IsValid()))
+	{
+		return SNullWidget::NullWidget;
+	}
+
+	const URCPropertyAction* PropertyAction = PropertyActionWeakPtr.Get();	
+
 	if (DetailTreeNodeWeakPtr.IsValid())
 	{
 		const FNodeWidgets NodeWidgets = DetailTreeNodeWeakPtr.Pin()->CreateNodeWidgets();
 
 		const TSharedRef<SHorizontalBox> FieldWidget = SNew(SHorizontalBox);
-	
+
 		if (NodeWidgets.ValueWidget)
 		{
 			FieldWidget->AddSlot()
@@ -87,70 +243,24 @@ TSharedRef<SWidget> FRCPropertyActionModel::GetWidget() const
 				];
 		}
 
-		TSharedRef<SBorder> BorderWidget = SNew(SBorder)
-			.Padding(0.0f)
-			.BorderImage(FRemoteControlPanelStyle::Get()->GetBrush("RemoteControlPanel.ExposedFieldBorder"));
-
-		FLinearColor TypeColor;
-		if (PropertyActionWeakPtr.IsValid())
-		{
-			TypeColor = UE::RCUIHelpers::GetFieldClassTypeColor(PropertyActionWeakPtr->PropertySelfContainer->GetProperty());
-		}
-	
-		return SNew(SHorizontalBox)
-			.Clipping(EWidgetClipping::OnDemand)
-
-			// Variable Color Bar
-			+ SHorizontalBox::Slot()
-			.VAlign(VAlign_Fill)
-			.AutoWidth()
-			.Padding(FMargin(3.f))
-			[
-				SNew(SBorder)
-				.Visibility(EVisibility::HitTestInvisible)
-				.BorderImage(FAppStyle::Get().GetBrush("NumericEntrySpinBox.NarrowDecorator"))
-				.BorderBackgroundColor(TypeColor)
-				.Padding(FMargin(5.0f, 0.0f, 0.0f, 0.0f))
-			]
-			// Field name
-			+SHorizontalBox::Slot()
-			.VAlign(VAlign_Center)
-			.AutoWidth()
-			[
-				NodeWidgets.NameWidget.ToSharedRef()
-			]
-			// Value Widget
-			+ SHorizontalBox::Slot()
-			.VAlign(VAlign_Center)
-			.AutoWidth()
-			[
-				FieldWidget
-			];
+		return FieldWidget;
 	}
 	
-	return FRCActionModel::GetWidget();
+	return SNullWidget::NullWidget;
 }
 
-FRCFunctionActionModel::FRCFunctionActionModel(URCFunctionAction* InFunctionAction)
-	: FRCActionModel(InFunctionAction)
-	, FunctionActionWeakPtr(InFunctionAction)
+FLinearColor FRCPropertyActionType::GetPropertyTypeColor() const
 {
-}
-
-TSharedRef<SWidget> FRCFunctionActionModel::GetWidget() const
-{
-	if (FunctionActionWeakPtr.IsValid())
+	if (!ensure(PropertyActionWeakPtr.IsValid()))
 	{
-		if (URemoteControlPreset* Preset = FunctionActionWeakPtr->PresetWeakPtr.Get())
-		{
-			if (const TSharedPtr<FRemoteControlFunction> RemoteControlField = Preset->GetExposedEntity<FRemoteControlFunction>(FunctionActionWeakPtr->ExposedFieldId).Pin())
-			{
-				const FString FinalNameStr = RemoteControlField->GetLabel().ToString();
-
-				return SNew(STextBlock).Text(FText::FromString(FinalNameStr));
-			}
-		}
+		return FLinearColor::White;
 	}
-	
-	return FRCActionModel::GetWidget();
-}
\ No newline at end of file
+
+	const URCPropertyAction* PropertyAction = PropertyActionWeakPtr.Get();
+
+	const FLinearColor TypeColor = UE::RCUIHelpers::GetFieldClassTypeColor(PropertyAction->PropertySelfContainer->GetProperty());
+
+	return TypeColor;
+}
+
+#undef LOCTEXT_NAMESPACE
\ No newline at end of file
diff --git a/Engine/Plugins/VirtualProduction/RemoteControl/Source/RemoteControlUI/Private/UI/Action/RCActionModel.h b/Engine/Plugins/VirtualProduction/RemoteControl/Source/RemoteControlUI/Private/UI/Action/RCActionModel.h
--- a/Engine/Plugins/VirtualProduction/RemoteControl/Source/RemoteControlUI/Private/UI/Action/RCActionModel.h
+++ b/Engine/Plugins/VirtualProduction/RemoteControl/Source/RemoteControlUI/Private/UI/Action/RCActionModel.h
@@ -2,11 +2,20 @@
 
 #pragma once
 
+#include "Action/RCFunctionAction.h"
+#include "Action/RCPropertyAction.h"
 #include "UI/BaseLogicUI/RCLogicModeBase.h"
+
 #include "UObject/WeakFieldPtr.h"
 
+class FRCBehaviourModel;
+struct FRCPanelStyle;
 class IDetailTreeNode;
+class ITableRow;
 class IPropertyRowGenerator;
+class SHeaderRow;
+class SHorizontalBox;
+class STableViewBase;
 class SWidget;
 class URCAction;
 class URCPropertyAction;
@@ -22,35 +31,64 @@ class URCFunctionAction;
 class FRCActionModel : public FRCLogicModeBase
 {
 public:
-	FRCActionModel(URCAction* InAction);
+	FRCActionModel(URCAction* InAction, const TSharedPtr<class FRCBehaviourModel> InBehaviourItem, const TSharedPtr<SRemoteControlPanel> InRemoteControlPanel);
 	
 	/** Fetch the Action (data model) associated with us*/
 	URCAction* GetAction() const;
 
-private:
+	/** Widget representing Action Name field */
+	virtual TSharedRef<SWidget> GetNameWidget() const;
+
+	/** Widget representing the Value field */
+	virtual TSharedRef<SWidget> GetWidget() const override;
+
+	/** Widget representing Color Coding for the Action Type*/
+	virtual TSharedRef<SWidget> GetVariableColorWidget() const;
+
+	/** Color code for this Action. Customized per action type*/
+	virtual FLinearColor GetActionTypeColor() const = 0;
+
+	/** OnGenerateRow delegate for the Actions List View*/
+	TSharedRef<ITableRow> OnGenerateWidgetForList(TSharedPtr<FRCActionModel> InItem, const TSharedRef<STableViewBase>& OwnerTable);
+
+	/** Get the Header Row appropriate for this particular Action model */
+	static TSharedPtr<SHeaderRow> GetHeaderRow();
+
+	/** Chooses the appropriate Action model for the current class and field type*/
+	static TSharedPtr<FRCActionModel> GetModelByActionType(URCAction* InAction, const TSharedPtr<class FRCBehaviourModel> InBehaviourItem, const TSharedPtr<SRemoteControlPanel> InRemoteControlPanel);
+
+protected:
 	/** The Action (data model) associated with us*/
 	TWeakObjectPtr<URCAction> ActionWeakPtr;
+
+	/** Panel Style reference. */
+	const FRCPanelStyle* RCPanelStyle;
+
+	/** The Behaviour (UI model) associated with us*/
+	TWeakPtr<FRCBehaviourModel> BehaviourItemWeakPtr;
 };
 
-/* 
-* ~ FRCPropertyActionModel ~
+/*
+* ~ FRCPropertyActionType ~
 *
-* UI model for Property based Actions
+* Reusable class containing Property action related data and functions
 */
-class FRCPropertyActionModel : public FRCActionModel
+class FRCPropertyActionType
 {
 public:
-	FRCPropertyActionModel(URCPropertyAction* InPropertyAction);
+	FRCPropertyActionType(URCPropertyAction* InPropertyAction);
 
 	/** Property Name associated with this Action */
 	const FName& GetPropertyName() const;
 
-	/** The widget to be rendered for this Property 
-	* Used to represent a single row when added to the Actions Panel List
+	/** The widget to be rendered for this Property
+	* Represents the input field which the user will use to set a value for this Action
 	*/
-	virtual TSharedRef<SWidget> GetWidget() const override;
+	TSharedRef<SWidget> GetPropertyWidget() const;
 
-private:
+	FLinearColor GetPropertyTypeColor() const;
+
+protected:
 	/** The Property Action (data model) associated with us*/
 	TWeakObjectPtr<URCPropertyAction> PropertyActionWeakPtr;
 
@@ -58,25 +96,70 @@ private:
 	TSharedPtr<IPropertyRowGenerator> PropertyRowGenerator;
 
 	/** Used to create a generic Value Widget for the property row widget*/
-	TWeakPtr<IDetailTreeNode> DetailTreeNodeWeakPtr;
+	TWeakPtr<IDetailTreeNode> DetailTreeNodeWeakPtr;	
 };
 
 /*
-* ~ FRCFunctionActionModel ~
+* ~ FRCFunctionActionType ~
 *
-* UI model for Function based Actions
+* Reusable class containing Function action related data and functions
 */
-class FRCFunctionActionModel : public FRCActionModel
+class FRCFunctionActionType
 {
 public:
-	FRCFunctionActionModel(URCFunctionAction* InAction);
+	FRCFunctionActionType(URCFunctionAction* InFunctionAction) : FunctionActionWeakPtr(InFunctionAction) {}
 
-	/** The widget to be rendered for this Function
-	* Used to represent a single row when added to the Actions Panel List
-	*/
-	virtual TSharedRef<SWidget> GetWidget() const override;
+	/** Color code for function actions in the Actions table */
+	FLinearColor GetFunctionTypeColor() const
+	{
+		// @todo: Confirm color to be used for this with VP team.
+		return FLinearColor::Blue; 
+	}
 
-private:
+protected:
 	/** The Function Action (data model) associated with us*/
 	TWeakObjectPtr<URCFunctionAction> FunctionActionWeakPtr;
 };
+
+/*
+* ~ FRCPropertyActionModel ~
+*
+* UI model for Property based Actions
+*/
+class FRCPropertyActionModel : public FRCActionModel, public FRCPropertyActionType
+{
+public:
+	FRCPropertyActionModel(URCPropertyAction* InPropertyAction, const TSharedPtr<class FRCBehaviourModel> InBehaviourItem, const TSharedPtr<SRemoteControlPanel> InRemoteControlPanel)
+		: FRCActionModel(InPropertyAction, InBehaviourItem, InRemoteControlPanel),
+		FRCPropertyActionType(InPropertyAction)	{ }
+
+	virtual TSharedRef<SWidget> GetWidget() const override
+	{
+		return GetPropertyWidget();
+	}
+
+	/** Color code for this Action*/
+	virtual FLinearColor GetActionTypeColor() const override
+	{
+		return GetPropertyTypeColor();
+	}
+};
+
+/*
+* ~ FRCFunctionActionModel ~
+*
+* UI model for Function based Actions
+*/
+class FRCFunctionActionModel : public FRCActionModel, public FRCFunctionActionType
+{
+public:
+	FRCFunctionActionModel(URCFunctionAction* InFunctionAction, const TSharedPtr<class FRCBehaviourModel> InBehaviourItem, const TSharedPtr<SRemoteControlPanel> InRemoteControlPanel)
+		: FRCActionModel(InFunctionAction, InBehaviourItem, InRemoteControlPanel)
+		, FRCFunctionActionType(InFunctionAction) { }
+
+	/** Color code for this Action*/
+	virtual FLinearColor GetActionTypeColor() const override
+	{
+		return GetFunctionTypeColor();
+	}
+};
diff --git a/Engine/Plugins/VirtualProduction/RemoteControl/Source/RemoteControlUI/Private/UI/Action/SRCActionPanel.cpp b/Engine/Plugins/VirtualProduction/RemoteControl/Source/RemoteControlUI/Private/UI/Action/SRCActionPanel.cpp
--- a/Engine/Plugins/VirtualProduction/RemoteControl/Source/RemoteControlUI/Private/UI/Action/SRCActionPanel.cpp
+++ b/Engine/Plugins/VirtualProduction/RemoteControl/Source/RemoteControlUI/Private/UI/Action/SRCActionPanel.cpp
@@ -84,11 +84,13 @@ void SRCActionPanel::UpdateWrappedWidget(TSharedPtr<FRCBehaviourModel> InBehavio
 {
 	if (InBehaviourItem.IsValid())
 	{
+		ActionPanelList = InBehaviourItem->GetActionsListWidget(SharedThis(this));
+
 		// Action Dock Panel
 		TSharedPtr<SRCMinorPanel> ActionDockPanel = SNew(SRCMinorPanel)
 			.HeaderLabel(LOCTEXT("ActionsLabel", "Actions"))
 			[
-				SAssignNew(ActionPanelList, SRCActionPanelList, SharedThis(this), InBehaviourItem)
+				ActionPanelList.ToSharedRef()
 			];
 
 		// Add New Action Button
@@ -255,7 +257,7 @@ TSharedRef<SWidget> SRCActionPanel::GetActionMenuContentWidget()
 				if (const URCBehaviour* Behaviour = SelectedBehaviourItemWeakPtr.Pin()->GetBehaviour())
 				{
 					// Skip if we already have an Action created for this exposed entity
-					if (Behaviour->ActionContainer->FindActionByFieldId(RemoteControlField->GetId()))
+					if(!Behaviour->CanHaveActionForField(RemoteControlField))
 					{
 						continue;
 					}
@@ -275,31 +277,32 @@ TSharedRef<SWidget> SRCActionPanel::GetActionMenuContentWidget()
 	return MenuBuilder.MakeWidget();
 }
 
-void SRCActionPanel::OnAddActionClicked(TSharedPtr<FRemoteControlField> InRemoteControlField)
+URCAction* SRCActionPanel::AddAction(const TSharedRef<const FRemoteControlField> InRemoteControlField)
 {
-	if (!SelectedBehaviourItemWeakPtr.IsValid())
-	{
-		return;
-	}
-
-	if (const URCBehaviour* Behaviour = SelectedBehaviourItemWeakPtr.Pin()->GetBehaviour())
+	if (const TSharedPtr<FRCBehaviourModel> BehaviourItem = SelectedBehaviourItemWeakPtr.Pin())
 	{
-		URCAction* NewAction = nullptr;
-
-		if (InRemoteControlField->FieldType == EExposedFieldType::Property)
-		{
-			NewAction = Behaviour->ActionContainer->AddAction(StaticCastSharedPtr<FRemoteControlProperty>(InRemoteControlField));
-		}
-		else if (InRemoteControlField->FieldType == EExposedFieldType::Function)
-		{
-			NewAction = Behaviour->ActionContainer->AddAction(StaticCastSharedPtr<FRemoteControlFunction>(InRemoteControlField));
-		}
+		URCAction* NewAction = BehaviourItem->AddAction(InRemoteControlField);
 
+		// Broadcast new Action to other panels
 		if (const TSharedPtr<SRemoteControlPanel> RemoteControlPanel = GetRemoteControlPanel())
 		{
 			RemoteControlPanel->OnActionAdded.Broadcast(NewAction);
 		}
+
+		return NewAction;
 	}
+
+	return nullptr;
+}
+
+void SRCActionPanel::OnAddActionClicked(TSharedPtr<FRemoteControlField> InRemoteControlField)
+{
+	if (!SelectedBehaviourItemWeakPtr.IsValid() || !InRemoteControlField)
+	{
+		return;
+	}
+
+	AddAction(InRemoteControlField.ToSharedRef());
 }
 
 FReply SRCActionPanel::OnClickEmptyButton()
@@ -329,7 +332,7 @@ FReply SRCActionPanel::OnAddAllFields()
 
 	if (URemoteControlPreset* Preset = GetPreset())
 	{
-		if (const URCBehaviour* Behaviour = SelectedBehaviourItemWeakPtr.Pin()->GetBehaviour())
+		if (const TSharedPtr<FRCBehaviourModel> BehaviourItem = SelectedBehaviourItemWeakPtr.Pin())
 		{
 			const TArray<TWeakPtr<FRemoteControlField>>& RemoteControlFields = Preset->GetExposedEntities<FRemoteControlField>();
 
@@ -338,24 +341,7 @@ FReply SRCActionPanel::OnAddAllFields()
 			{
 				if (const TSharedPtr<FRemoteControlField> RemoteControlField = RemoteControlFieldWeakPtr.Pin())
 				{
-					URCAction* NewAction = nullptr;
-
-					// Property Action
-					if (RemoteControlField->FieldType == EExposedFieldType::Property)
-					{
-						NewAction = Behaviour->ActionContainer->AddAction(StaticCastSharedPtr<FRemoteControlProperty>(RemoteControlField));
-					}
-					// Function Action
-					else if (RemoteControlField->FieldType == EExposedFieldType::Function)
-					{
-						NewAction = Behaviour->ActionContainer->AddAction(StaticCastSharedPtr<FRemoteControlFunction>(RemoteControlField));
-					}
-
-					// Broadcast New Action
-					if (const TSharedPtr<SRemoteControlPanel> RemoteControlPanel = GetRemoteControlPanel())
-					{
-						RemoteControlPanel->OnActionAdded.Broadcast(NewAction);
-					}
+					AddAction(RemoteControlField.ToSharedRef());
 				}
 			}
 		}
diff --git a/Engine/Plugins/VirtualProduction/RemoteControl/Source/RemoteControlUI/Private/UI/Action/SRCActionPanel.h b/Engine/Plugins/VirtualProduction/RemoteControl/Source/RemoteControlUI/Private/UI/Action/SRCActionPanel.h
--- a/Engine/Plugins/VirtualProduction/RemoteControl/Source/RemoteControlUI/Private/UI/Action/SRCActionPanel.h
+++ b/Engine/Plugins/VirtualProduction/RemoteControl/Source/RemoteControlUI/Private/UI/Action/SRCActionPanel.h
@@ -8,7 +8,10 @@ struct FRCPanelStyle;
 struct FRemoteControlField;
 class FRCBehaviourModel;
 class SBox;
+class SRCLogicPanelListBase;
 class SRemoteControlPanel;
+class URCAction;
+
 /*
 * ~ SRCActionPanel ~
 *
@@ -36,6 +39,9 @@ public:
 	/** Delete Item UI command implementation for this panel */
 	virtual void DeleteSelectedPanelItem() override;
 
+	/** Adds an Action for the currently active Behaviour and broadcasts to parent panels */
+	URCAction* AddAction(const TSharedRef<const FRemoteControlField> InRemoteControlField);
+
 protected:
 
 	/** Warns user before deleting all items in a panel. */
@@ -82,7 +88,7 @@ private:
 	TSharedPtr<SBox> WrappedBoxWidget;
 
 	/** Widget representing List of Actions */
-	TSharedPtr<class SRCActionPanelList> ActionPanelList;
+	TSharedPtr<SRCLogicPanelListBase> ActionPanelList;
 
 	/** Helper widget for behavior details. */
 	static TSharedPtr<SBox> NoneSelectedWidget;
diff --git a/Engine/Plugins/VirtualProduction/RemoteControl/Source/RemoteControlUI/Private/UI/Action/SRCActionPanelList.cpp b/Engine/Plugins/VirtualProduction/RemoteControl/Source/RemoteControlUI/Private/UI/Action/SRCActionPanelList.cpp
--- a/Engine/Plugins/VirtualProduction/RemoteControl/Source/RemoteControlUI/Private/UI/Action/SRCActionPanelList.cpp
+++ b/Engine/Plugins/VirtualProduction/RemoteControl/Source/RemoteControlUI/Private/UI/Action/SRCActionPanelList.cpp
@@ -7,12 +7,16 @@
 #include "Action/RCFunctionAction.h"
 #include "Action/RCPropertyAction.h"
 #include "RCActionModel.h"
+#include "RemoteControlPreset.h"
 #include "SlateOptMacros.h"
 #include "SRCActionPanel.h"
+#include "SDropTarget.h"
 #include "Styling/RemoteControlStyles.h"
+#include "UI/Action/Conditional/RCActionConditionalModel.h"
 #include "UI/Behaviour/RCBehaviourModel.h"
 #include "UI/RCUIHelpers.h"
 #include "UI/RemoteControlPanelStyle.h"
+#include "UI/SRCPanelExposedField.h"
 #include "UI/SRemoteControlPanel.h"
 #include "Widgets/Views/SListView.h"
 #include "Widgets/Views/STableRow.h"
@@ -20,16 +24,8 @@
 
 #define LOCTEXT_NAMESPACE "SRCActionPanelList"
 
-namespace FRemoteControlActionColumns
-{
-	const FName DragDropHandle = TEXT("DragDropHandle");
-	const FName Description = TEXT("Description");
-	const FName Value = TEXT("Value");
-}
-
-BEGIN_SLATE_FUNCTION_BUILD_OPTIMIZATION
-
-void SRCActionPanelList::Construct(const FArguments& InArgs, const TSharedRef<SRCActionPanel> InActionPanel, TSharedPtr<FRCBehaviourModel> InBehaviourItem)
+template <class ActionType>
+void SRCActionPanelList<ActionType>::Construct(const FArguments& InArgs, const TSharedRef<SRCActionPanel> InActionPanel, TSharedPtr<FRCBehaviourModel> InBehaviourItem)
 {
 	SRCLogicPanelListBase::Construct(SRCLogicPanelListBase::FArguments());
 
@@ -38,33 +34,23 @@ void SRCActionPanelList::Construct(const FArguments& InArgs, const TSharedRef<SR
 	
 	RCPanelStyle = &FRemoteControlPanelStyle::Get()->GetWidgetStyle<FRCPanelStyle>("RemoteControlPanel.MinorPanel");
 
-	ListView = SNew(SListView<TSharedPtr<FRCActionModel>>)
-		.ListItemsSource( &ActionItems )
-		.OnGenerateRow(this, &SRCActionPanelList::OnGenerateWidgetForList )
+	ListView = SNew(SListView<TSharedPtr<ActionType>>)
+		.ListItemsSource(&ActionItems)
+		.OnGenerateRow(this, &SRCActionPanelList::OnGenerateWidgetForList)
 		.ListViewStyle(&RCPanelStyle->TableViewStyle)
-		.HeaderRow(
-			SNew(SHeaderRow)
-			.Style(&RCPanelStyle->HeaderRowStyle)
-
-			+ SHeaderRow::Column(FRemoteControlActionColumns::DragDropHandle)
-			.DefaultLabel(LOCTEXT("RCActionDragDropHandleColumnHeader", ""))
-			.FixedWidth(25.f)
-			.HeaderContentPadding(RCPanelStyle->HeaderRowPadding)
-
-			+ SHeaderRow::Column(FRemoteControlActionColumns::Description)
-			.DefaultLabel(LOCTEXT("RCActionDescColumnHeader", "Description"))
-			.FillWidth(0.5f)
-			.HeaderContentPadding(RCPanelStyle->HeaderRowPadding)
-
-			+ SHeaderRow::Column(FRemoteControlActionColumns::Value)
-			.DefaultLabel(LOCTEXT("RCActionValueColumnHeader", "Value"))
-			.FillWidth(0.5f)
-			.HeaderContentPadding(RCPanelStyle->HeaderRowPadding)
-		);
-	
+		.HeaderRow(ActionType::GetHeaderRow());
+
 	ChildSlot
 	[
-		ListView.ToSharedRef()
+		SNew(SDropTarget)
+		.VerticalImage(FRemoteControlPanelStyle::Get()->GetBrush("RemoteControlPanel.VerticalDash"))
+		.HorizontalImage(FRemoteControlPanelStyle::Get()->GetBrush("RemoteControlPanel.HorizontalDash"))
+		.OnDropped_Lambda([this](const FGeometry& InGeometry, const FDragDropEvent& InDragDropEvent) { return SRCActionPanelList::OnExposedFieldDrop(InDragDropEvent.GetOperation()); })
+		.OnAllowDrop(this, &SRCActionPanelList::OnAllowDrop)
+		.OnIsRecognized(this, &SRCActionPanelList::OnAllowDrop)
+		[
+			ListView.ToSharedRef()
+		]
 	];
 
 	// Add delegates
@@ -78,63 +64,103 @@ void SRCActionPanelList::Construct(const FArguments& InArgs, const TSharedRef<SR
 	Reset();
 }
 
-END_SLATE_FUNCTION_BUILD_OPTIMIZATION
-
-bool SRCActionPanelList::IsEmpty() const
+template <class ActionType>
+bool SRCActionPanelList<ActionType>::IsEmpty() const
 {
 	return ActionItems.IsEmpty();
 }
 
-int32 SRCActionPanelList::Num() const
+template <class ActionType>
+int32 SRCActionPanelList<ActionType>::Num() const
 {
 	return ActionItems.Num();
 }
 
-void SRCActionPanelList::Reset()
+template <class ActionType>
+void SRCActionPanelList<ActionType>::Reset()
 {
-	ActionItems.Empty();
-
-	if (TSharedPtr<FRCBehaviourModel> BehaviourItem = BehaviourItemWeakPtr.Pin())
+	if (TSharedPtr<SRCActionPanel> ActionPanel = GetActionPanel())
 	{
-		if (URCBehaviour* Behaviour = Cast<URCBehaviour>(BehaviourItem->GetBehaviour()))
+		if (TSharedPtr<SRemoteControlPanel> RemoteControlPanel = ActionPanel->GetRemoteControlPanel())
 		{
-			for (URCAction* Action : Behaviour->ActionContainer->Actions)
+			ActionItems.Empty();
+
+			if (TSharedPtr<FRCBehaviourModel> BehaviourItem = BehaviourItemWeakPtr.Pin())
 			{
-				if (URCPropertyAction* PropertyAction = Cast<URCPropertyAction>(Action))
+				if (URCBehaviour* Behaviour = Cast<URCBehaviour>(BehaviourItem->GetBehaviour()))
 				{
-					ActionItems.Add(MakeShared<FRCPropertyActionModel>(PropertyAction));
+					for (URCAction* Action : Behaviour->ActionContainer->Actions)
+					{
+						TSharedPtr<ActionType> ActionItem = ActionType::GetModelByActionType(Action, BehaviourItem, RemoteControlPanel);
+
+						ActionItems.Add(ActionItem);
+					}
 				}
-				else if (URCFunctionAction* FunctionAction = Cast<URCFunctionAction>(Action))
+			}
+
+			ListView->RebuildList();
+		}
+	}
+}
+
+template <class ActionType>
+FReply SRCActionPanelList<ActionType>::OnExposedFieldDrop(TSharedPtr<FDragDropOperation> DragDropOperation)
+{
+	if (DragDropOperation && DragDropOperation->IsOfType<FExposedEntityDragDrop>())
+	{
+		if (TSharedPtr<FExposedEntityDragDrop> DragDropOp = StaticCastSharedPtr<FExposedEntityDragDrop>(DragDropOperation))
+		{
+			if (const URemoteControlPreset* Preset = GetPreset())
+			{
+				FGuid ExposedEntityId = DragDropOp->GetId();
+				if (TSharedPtr<const FRemoteControlField> RemoteControlField = Preset->GetExposedEntity<FRemoteControlField>(ExposedEntityId).Pin())
 				{
-					ActionItems.Add(MakeShared<FRCFunctionActionModel>(FunctionAction));
+					if (const TSharedPtr<SRCActionPanel> ActionPanel = GetActionPanel())
+					{
+						ActionPanel->AddAction(RemoteControlField.ToSharedRef());
+					}
 				}
 			}
 		}
 	}
 
-	ListView->RebuildList();
+	return FReply::Handled();
 }
 
-TSharedRef<ITableRow> SRCActionPanelList::OnGenerateWidgetForList(TSharedPtr<FRCActionModel> InItem, const TSharedRef<STableViewBase>& OwnerTable)
+template <class ActionType>
+bool SRCActionPanelList<ActionType>::OnAllowDrop(TSharedPtr<FDragDropOperation> DragDropOperation)
 {
+	return true;
+}
+
+template <class ActionType>
+TSharedRef<ITableRow> SRCActionPanelList<ActionType>::OnGenerateWidgetForList(TSharedPtr<ActionType> InItem, const TSharedRef<STableViewBase>& OwnerTable)
+{
+	if (ensure(InItem))
+	{
+		return InItem->OnGenerateWidgetForList(InItem, OwnerTable);
+	}
+
 	return SNew(STableRow<TSharedPtr<FString>>, OwnerTable)
-		.Style(&RCPanelStyle->TableRowStyle)
 		[
-			InItem->GetWidget()
+			SNullWidget::NullWidget
 		];
 }
 
-void SRCActionPanelList::OnActionAdded(URCAction* InAction)
+template <class ActionType>
+void SRCActionPanelList<ActionType>::OnActionAdded(URCAction* InAction)
 {
 	Reset();
 }
 
-void SRCActionPanelList::OnEmptyActions()
+template <class ActionType>
+void SRCActionPanelList<ActionType>::OnEmptyActions()
 {
 	Reset();
 }
 
-URemoteControlPreset* SRCActionPanelList::GetPreset()
+template <class ActionType>
+URemoteControlPreset* SRCActionPanelList<ActionType>::GetPreset()
 {
 	if (ActionPanelWeakPtr.IsValid())
 	{
@@ -144,13 +170,14 @@ URemoteControlPreset* SRCActionPanelList::GetPreset()
 	return nullptr;
 }
 
-int32 SRCActionPanelList::RemoveModel(const TSharedPtr<FRCLogicModeBase> InModel)
+template <class ActionType>
+int32 SRCActionPanelList<ActionType>::RemoveModel(const TSharedPtr<FRCLogicModeBase> InModel)
 {
 	if (const TSharedPtr<FRCBehaviourModel> BehaviourModel = BehaviourItemWeakPtr.Pin())
 	{
 		if (const URCBehaviour* Behaviour = BehaviourModel->GetBehaviour())
 		{
-			if(const TSharedPtr<FRCActionModel> SelectedAction = StaticCastSharedPtr<FRCActionModel>(InModel))
+			if(const TSharedPtr<ActionType> SelectedAction = StaticCastSharedPtr<ActionType>(InModel))
 			{
 				// Remove Model from Data Container
 				const int32 RemoveCount = Behaviour->ActionContainer->RemoveAction(SelectedAction->GetAction());
@@ -163,14 +190,19 @@ int32 SRCActionPanelList::RemoveModel(const TSharedPtr<FRCLogicModeBase> InModel
 	return 0;
 }
 
-bool SRCActionPanelList::IsListFocused() const
+template <class ActionType>
+bool SRCActionPanelList<ActionType>::IsListFocused() const
 {
 	return ListView->HasAnyUserFocus().IsSet();
 }
 
-void SRCActionPanelList::DeleteSelectedPanelItem()
+template <class ActionType>
+void SRCActionPanelList<ActionType>::DeleteSelectedPanelItem()
 {
-	DeleteItemFromLogicPanel<FRCActionModel>(ActionItems, ListView->GetSelectedItems());
+	DeleteItemFromLogicPanel<ActionType>(ActionItems, ListView->GetSelectedItems());
 }
 
-#undef LOCTEXT_NAMESPACE
\ No newline at end of file
+#undef LOCTEXT_NAMESPACE
+
+template class SRCActionPanelList<FRCActionModel>;
+template class SRCActionPanelList<FRCActionConditionalModel>;
diff --git a/Engine/Plugins/VirtualProduction/RemoteControl/Source/RemoteControlUI/Private/UI/Action/SRCActionPanelList.h b/Engine/Plugins/VirtualProduction/RemoteControl/Source/RemoteControlUI/Private/UI/Action/SRCActionPanelList.h
--- a/Engine/Plugins/VirtualProduction/RemoteControl/Source/RemoteControlUI/Private/UI/Action/SRCActionPanelList.h
+++ b/Engine/Plugins/VirtualProduction/RemoteControl/Source/RemoteControlUI/Private/UI/Action/SRCActionPanelList.h
@@ -6,7 +6,7 @@
 
 struct FRCPanelStyle;
 class FRCBehaviourModel;
-class FRCActionModel;
+class ActionType;
 class ITableRow;
 class STableViewBase;
 class SRCActionPanel;
@@ -20,10 +20,11 @@ template <typename ItemType> class SListView;
 * UI Widget for Actions List
 * Used as part of the RC Logic Actions Panel.
 */
+template <typename ActionType>
 class REMOTECONTROLUI_API SRCActionPanelList : public SRCLogicPanelListBase
 {
 public:
-	SLATE_BEGIN_ARGS(SRCActionPanelList)
+	SLATE_BEGIN_ARGS(SRCActionPanelList<ActionType>)
 		{
 		}
 
@@ -39,15 +40,27 @@ public:
 	virtual int32 Num() const override;
 
 	/** Whether the Actions List View currently has focus.*/
-	bool IsListFocused() const;
+	virtual bool IsListFocused() const override;
 
 	/** Deletes currently selected items from the list view*/
-	void DeleteSelectedPanelItem();
+	virtual void DeleteSelectedPanelItem() override;
+
+	/** Fetches the parent Action panel*/
+	TSharedPtr<SRCActionPanel> GetActionPanel()
+	{
+		return ActionPanelWeakPtr.Pin();
+	}
+
+	/** Fetches the Behaviour (UI model) associated with us */
+	TSharedPtr<FRCBehaviourModel> GetBehaviourItem()
+	{
+		return BehaviourItemWeakPtr.Pin();
+	}
 
 private:
 
 	/** OnGenerateRow delegate for the Actions List View*/
-	TSharedRef<ITableRow> OnGenerateWidgetForList( TSharedPtr<FRCActionModel> InItem, const TSharedRef<STableViewBase>& OwnerTable );
+	TSharedRef<ITableRow> OnGenerateWidgetForList( TSharedPtr<ActionType> InItem, const TSharedRef<STableViewBase>& OwnerTable );
 
 	/** Responds to the selection of a newly created action. Resets UI state*/
 	void OnActionAdded(URCAction* InAction);
@@ -67,6 +80,12 @@ private:
 	/** Removes the given Action UI model item from the list of UI models*/
 	virtual int32 RemoveModel(const TSharedPtr<FRCLogicModeBase> InModel) override;
 
+	/*Drag and drop event for creating an Action from an exposed field */
+	FReply OnExposedFieldDrop(TSharedPtr<FDragDropOperation> DragDropOperation);
+
+	/** Whether drag and drop is possible from the current exposed property to the Actions table */
+	bool OnAllowDrop(TSharedPtr<FDragDropOperation> DragDropOperation);
+
 private:
 
 	/** The parent Action Panel widget*/
@@ -76,10 +95,10 @@ private:
 	TWeakPtr<FRCBehaviourModel> BehaviourItemWeakPtr;
 
 	/** List of Actions (UI model) active in this widget */
-	TArray<TSharedPtr<FRCActionModel>> ActionItems;
+	TArray<TSharedPtr<ActionType>> ActionItems;
 
 	/** List View widget for representing our Actions List*/
-	TSharedPtr<SListView<TSharedPtr<FRCActionModel>>> ListView;
+	TSharedPtr<SListView<TSharedPtr<ActionType>>> ListView;
 
 	/** Panel Style reference. */
 	const FRCPanelStyle* RCPanelStyle;
diff --git a/Engine/Plugins/VirtualProduction/RemoteControl/Source/RemoteControlUI/Private/UI/BaseLogicUI/SRCLogicPanelListBase.cpp b/Engine/Plugins/VirtualProduction/RemoteControl/Source/RemoteControlUI/Private/UI/BaseLogicUI/SRCLogicPanelListBase.cpp
--- a/Engine/Plugins/VirtualProduction/RemoteControl/Source/RemoteControlUI/Private/UI/BaseLogicUI/SRCLogicPanelListBase.cpp
+++ b/Engine/Plugins/VirtualProduction/RemoteControl/Source/RemoteControlUI/Private/UI/BaseLogicUI/SRCLogicPanelListBase.cpp
@@ -4,6 +4,7 @@
 
 #include "RCLogicModeBase.h"
 #include "SlateOptMacros.h"
+#include "UI/Action/Conditional/RCActionConditionalModel.h"
 #include "UI/Action/RCActionModel.h"
 #include "UI/Behaviour/RCBehaviourModel.h"
 #include "UI/Controller/RCControllerModel.h"
@@ -53,4 +54,5 @@ void SRCLogicPanelListBase::DeleteItemFromLogicPanel(TArray<TSharedPtr<T>>& Item
 
 template void SRCLogicPanelListBase::DeleteItemFromLogicPanel(TArray<TSharedPtr<FRCControllerModel>>& ItemsSource, const TArray<TSharedPtr<FRCControllerModel>>& SelectedItems);
 template void SRCLogicPanelListBase::DeleteItemFromLogicPanel(TArray<TSharedPtr<FRCBehaviourModel>>& ItemsSource, const TArray<TSharedPtr<FRCBehaviourModel>>& SelectedItems);
-template void SRCLogicPanelListBase::DeleteItemFromLogicPanel(TArray<TSharedPtr<FRCActionModel>>& ItemsSource, const TArray<TSharedPtr<FRCActionModel>>& SelectedItems);
\ No newline at end of file
+template void SRCLogicPanelListBase::DeleteItemFromLogicPanel(TArray<TSharedPtr<FRCActionModel>>& ItemsSource, const TArray<TSharedPtr<FRCActionModel>>& SelectedItems);
+template void SRCLogicPanelListBase::DeleteItemFromLogicPanel(TArray<TSharedPtr<FRCActionConditionalModel>>& ItemsSource, const TArray<TSharedPtr<FRCActionConditionalModel>>& SelectedItems);
\ No newline at end of file
diff --git a/Engine/Plugins/VirtualProduction/RemoteControl/Source/RemoteControlUI/Private/UI/BaseLogicUI/SRCLogicPanelListBase.h b/Engine/Plugins/VirtualProduction/RemoteControl/Source/RemoteControlUI/Private/UI/BaseLogicUI/SRCLogicPanelListBase.h
--- a/Engine/Plugins/VirtualProduction/RemoteControl/Source/RemoteControlUI/Private/UI/BaseLogicUI/SRCLogicPanelListBase.h
+++ b/Engine/Plugins/VirtualProduction/RemoteControl/Source/RemoteControlUI/Private/UI/BaseLogicUI/SRCLogicPanelListBase.h
@@ -32,6 +32,12 @@ public:
 	/** Returns number of items in the list. */
 	virtual int32 Num() const = 0;
 
+	/** Whether the List View currently has focus.*/
+	virtual bool IsListFocused() const = 0;
+
+	/** Deletes currently selected items from the list view*/
+	virtual void DeleteSelectedPanelItem()  = 0;
+
 private:
 
 	/** Refreshes the list from the latest state of the model*/
diff --git a/Engine/Plugins/VirtualProduction/RemoteControl/Source/RemoteControlUI/Private/UI/Behaviour/Builtin/Conditional/RCBehaviourConditionalModel.cpp b/Engine/Plugins/VirtualProduction/RemoteControl/Source/RemoteControlUI/Private/UI/Behaviour/Builtin/Conditional/RCBehaviourConditionalModel.cpp
new file mode 100644
--- /dev/null
+++ b/Engine/Plugins/VirtualProduction/RemoteControl/Source/RemoteControlUI/Private/UI/Behaviour/Builtin/Conditional/RCBehaviourConditionalModel.cpp
@@ -0,0 +1,141 @@
+// Copyright Epic Games, Inc. All Rights Reserved.
+
+#include "RCBehaviourConditionalModel.h"
+
+#include "Behaviour/Builtin/Conditional/RCBehaviourConditional.h"
+#include "Controller/RCController.h"
+#include "Modules/ModuleManager.h"
+#include "IDetailTreeNode.h"
+#include "IPropertyRowGenerator.h"
+#include "PropertyEditorModule.h"
+#include "RCVirtualProperty.h"
+#include "SRCBehaviourConditional.h"
+#include "UI/Action/Conditional/RCActionConditionalModel.h"
+#include "UI/Action/SRCActionPanel.h"
+#include "UI/Action/SRCActionPanelList.h"
+#include "Widgets/Layout/SBox.h"
+
+#define LOCTEXT_NAMESPACE "FRCBehaviourConditionalModel"
+
+FRCBehaviourConditionalModel::FRCBehaviourConditionalModel(URCBehaviourConditional* ConditionalBehaviour)
+	: FRCBehaviourModel(ConditionalBehaviour)
+	, ConditionalBehaviourWeakPtr(ConditionalBehaviour)
+{
+	FPropertyRowGeneratorArgs Args;
+	Args.bShouldShowHiddenProperties = true;
+	PropertyRowGenerator = FModuleManager::GetModuleChecked<FPropertyEditorModule>("PropertyEditor").CreatePropertyRowGenerator(Args);
+
+	CreateComparandInputField();
+}
+
+URCAction* FRCBehaviourConditionalModel::AddAction(const TSharedRef<const FRemoteControlField> InRemoteControlField)
+{
+	URCAction* NewAction = nullptr;
+
+	if (URCBehaviourConditional* ConditionalBehaviour = Cast<URCBehaviourConditional>(GetBehaviour()))
+	{
+		NewAction = ConditionalBehaviour->AddAction(InRemoteControlField, Condition, Comparand);
+
+		OnActionAdded(NewAction);
+	}
+
+	return NewAction;
+}
+
+void FRCBehaviourConditionalModel::CreateComparandInputField()
+{
+	if (URCBehaviourConditional* ConditionalBehaviour = Cast<URCBehaviourConditional>(GetBehaviour()))
+	{
+		if (URCController* Controller = ConditionalBehaviour->ControllerWeakPtr.Get())
+		{
+			// Virtual Property for the Comparand
+			Comparand = NewObject<URCVirtualPropertySelfContainer>(ConditionalBehaviour);
+			Comparand->DuplicateProperty(FName("Comparand"), Controller->GetProperty());
+
+			// UI widget (via Property Generator)
+			PropertyRowGenerator->SetStructure(Comparand->CreateStructOnScope());
+			for (const TSharedRef<IDetailTreeNode>& CategoryNode : PropertyRowGenerator->GetRootTreeNodes())
+			{
+				TArray<TSharedRef<IDetailTreeNode>> Children;
+				CategoryNode->GetChildren(Children);
+				for (const TSharedRef<IDetailTreeNode>& Child : Children)
+				{
+					DetailTreeNodeWeakPtr = Child;
+					break;
+				}
+			}
+		}
+	}
+}
+
+TSharedRef<SWidget> FRCBehaviourConditionalModel::GetBehaviourDetailsWidget()
+{
+	return SAssignNew(BehaviourDetailsWidget, SRCBehaviourConditional, SharedThis(this));
+}
+
+void FRCBehaviourConditionalModel::OnActionAdded(URCAction* Action)
+{
+	if (URCBehaviourConditional* ConditionalBehaviour = Cast<URCBehaviourConditional>(GetBehaviour()))
+	{
+		ConditionalBehaviour->OnActionAdded(Action, Condition, Comparand);
+
+		// Create a new Comparand for the user (the previous input field is already associated with the newly created action)
+		CreateComparandInputField();
+
+		BehaviourDetailsWidget->RefreshPropertyWidget();
+	}
+}
+
+TSharedRef<SWidget> FRCBehaviourConditionalModel::GetComparandFieldWidget() const
+{
+	if (const TSharedPtr<IDetailTreeNode> DetailTreeNode = DetailTreeNodeWeakPtr.Pin())
+	{
+		const FNodeWidgets NodeWidgets = DetailTreeNode->CreateNodeWidgets();
+
+		const TSharedRef<SHorizontalBox> FieldWidget = SNew(SHorizontalBox);
+
+		if (NodeWidgets.ValueWidget)
+		{
+			FieldWidget->AddSlot()
+				.Padding(FMargin(3.0f, 2.0f))
+				.HAlign(HAlign_Right)
+				.AutoWidth()
+				[
+					NodeWidgets.ValueWidget.ToSharedRef()
+				];
+		}
+		else if (NodeWidgets.WholeRowWidget)
+		{
+			FieldWidget->AddSlot()
+				.Padding(FMargin(3.0f, 2.0f))
+				.AutoWidth()
+				[
+					NodeWidgets.WholeRowWidget.ToSharedRef()
+				];
+		}
+
+		return SNew(SHorizontalBox)
+			.Clipping(EWidgetClipping::OnDemand)
+			// Property Value Widget
+			+ SHorizontalBox::Slot()
+			.VAlign(VAlign_Center)
+			.AutoWidth()
+			[
+				FieldWidget
+			];
+	}
+
+	return SNullWidget::NullWidget;
+}
+
+void FRCBehaviourConditionalModel::SetSelectedConditionType(const ERCBehaviourConditionType InCondition)
+{
+	Condition = InCondition;
+}
+
+TSharedPtr<SRCLogicPanelListBase> FRCBehaviourConditionalModel::GetActionsListWidget(TSharedRef<SRCActionPanel> InActionPanel)
+{
+	return SNew(SRCActionPanelList<FRCActionConditionalModel>, InActionPanel, SharedThis(this));
+}
+
+#undef LOCTEXT_NAMESPACE
\ No newline at end of file
diff --git a/Engine/Plugins/VirtualProduction/RemoteControl/Source/RemoteControlUI/Private/UI/Behaviour/Builtin/Conditional/RCBehaviourConditionalModel.h b/Engine/Plugins/VirtualProduction/RemoteControl/Source/RemoteControlUI/Private/UI/Behaviour/Builtin/Conditional/RCBehaviourConditionalModel.h
new file mode 100644
--- /dev/null
+++ b/Engine/Plugins/VirtualProduction/RemoteControl/Source/RemoteControlUI/Private/UI/Behaviour/Builtin/Conditional/RCBehaviourConditionalModel.h
@@ -0,0 +1,68 @@
+// Copyright Epic Games, Inc. All Rights Reserved.
+
+#pragma once
+#include "UI/Behaviour/RCBehaviourModel.h"
+
+enum class ERCBehaviourConditionType : uint8;
+class IDetailTreeNode;
+class IPropertyRowGenerator;
+class SRCActionPanel;
+class SRCBehaviourConditional;
+class SRCLogicPanelListBase;
+class SWidget;
+class URCBehaviourConditional;
+class URCVirtualPropertySelfContainer;
+
+/*
+* ~ FRCBehaviourConditionalModel ~
+*
+* UI model for Conditional Behaviour
+
+*/
+class FRCBehaviourConditionalModel : public FRCBehaviourModel
+{
+public:
+	FRCBehaviourConditionalModel(URCBehaviourConditional* IsEqualBehaviour);
+
+	/** Add a Logic Action using a remote control field as input */
+	virtual URCAction* AddAction(const TSharedRef<const FRemoteControlField> InRemoteControlField);
+
+	/** Builds a Behaviour specific details widget enabling the user to define Conditions*/
+	virtual TSharedRef<SWidget> GetBehaviourDetailsWidget() override;
+
+	/** Builds the Comparand input field's widget which enables the user to enter values to compare with*/
+	TSharedRef<SWidget> GetComparandFieldWidget() const;
+
+	/** Updates the condition type selected by the user for Conditional Behaviour */
+	void SetSelectedConditionType(const ERCBehaviourConditionType InCondition);
+
+	/** The Actions List widget to be used for Conditional Behaviour */	
+	virtual TSharedPtr<SRCLogicPanelListBase> GetActionsListWidget(TSharedRef<SRCActionPanel> InActionPanel);
+
+	/** The condition type currently selected by the user */
+	ERCBehaviourConditionType Condition;
+
+protected:
+	/** Invoked by parent class after an action has been added from UI*/
+	virtual void OnActionAdded(URCAction* Action) override;
+
+private:
+	/** Builds a generic value widget of the Controller's type
+	* This is used to store the Comparand's value as entered by the user*/
+	void CreateComparandInputField();
+
+	/** The Conditional Behaviour (Data model) associated with us*/
+	TWeakObjectPtr<URCBehaviourConditional> ConditionalBehaviourWeakPtr;
+
+	/** The row generator used to build a generic Value Widget*/
+	TSharedPtr<IPropertyRowGenerator> PropertyRowGenerator;
+
+	/** Used to create a generic value widget for comparison based on the active Controller's type*/
+	TWeakPtr<IDetailTreeNode> DetailTreeNodeWeakPtr;	
+
+	/** Virtual property used to build the Comparand - i.e. the property with which the Controller will be compared for a given condition*/
+	TObjectPtr<URCVirtualPropertySelfContainer> Comparand;	
+
+	/** Behaviour Details Widget (Conditions Widget in the case of this Behaviour) */
+	TSharedPtr<SRCBehaviourConditional> BehaviourDetailsWidget;
+};
\ No newline at end of file
diff --git a/Engine/Plugins/VirtualProduction/RemoteControl/Source/RemoteControlUI/Private/UI/Behaviour/Builtin/Conditional/SRCBehaviourConditional.cpp b/Engine/Plugins/VirtualProduction/RemoteControl/Source/RemoteControlUI/Private/UI/Behaviour/Builtin/Conditional/SRCBehaviourConditional.cpp
new file mode 100644
--- /dev/null
+++ b/Engine/Plugins/VirtualProduction/RemoteControl/Source/RemoteControlUI/Private/UI/Behaviour/Builtin/Conditional/SRCBehaviourConditional.cpp
@@ -0,0 +1,189 @@
+// Copyright Epic Games, Inc. All Rights Reserved.
+
+#include "SRCBehaviourConditional.h"
+
+#include "Action/RCAction.h"
+#include "Behaviour/Builtin/Conditional/RCBehaviourConditional.h"
+#include "Controller/RCController.h"
+#include "SlateOptMacros.h"
+#include "SPositiveActionButton.h"
+#include "Styling/CoreStyle.h"
+#include "Styling/RemoteControlStyles.h"
+#include "UI/Behaviour/Builtin/Conditional/RCBehaviourConditionalModel.h"
+#include "UI/RemoteControlPanelStyle.h"
+#include "Widgets/Images/SImage.h"
+#include "Widgets/Input/SButton.h"
+#include "Widgets/Layout/SSeparator.h"
+#include "Widgets/Layout/SBox.h"
+#include "Widgets/Views/SListView.h"
+#include "Widgets/Views/STableRow.h"
+
+#define LOCTEXT_NAMESPACE "SRCBehaviourConditional"
+
+class SPositiveActionButton;
+
+namespace UE::SRCBehaviourConditional::Private
+{
+	static TSharedPtr<ERCBehaviourConditionType> ConditionTypeIsEqual = MakeShared<ERCBehaviourConditionType>(ERCBehaviourConditionType::IsEqual);
+	static TSharedPtr<ERCBehaviourConditionType> ConditionTypeIsLesserThan = MakeShared<ERCBehaviourConditionType>(ERCBehaviourConditionType::IsLesserThan);
+	static TSharedPtr<ERCBehaviourConditionType> ConditionTypeIsGreaterThan = MakeShared<ERCBehaviourConditionType>(ERCBehaviourConditionType::IsGreaterThan);
+	static TSharedPtr<ERCBehaviourConditionType> ConditionTypeIsLesserThanOrEqualTo = MakeShared<ERCBehaviourConditionType>(ERCBehaviourConditionType::IsLesserThanOrEqualTo);
+	static TSharedPtr<ERCBehaviourConditionType> ConditionTypeIsGreaterThanOrEqualTo = MakeShared<ERCBehaviourConditionType>(ERCBehaviourConditionType::IsGreaterThanOrEqualTo);
+	static TSharedPtr<ERCBehaviourConditionType> ConditionTypeElse = MakeShared<ERCBehaviourConditionType>(ERCBehaviourConditionType::Else);
+}
+
+using namespace UE::SRCBehaviourConditional::Private;
+
+BEGIN_SLATE_FUNCTION_BUILD_OPTIMIZATION
+
+void SRCBehaviourConditional::Construct(const FArguments& InArgs, TSharedRef<FRCBehaviourConditionalModel> InBehaviourItem)
+{
+	ConditionalBehaviourItemWeakPtr = InBehaviourItem;
+
+	RCPanelStyle = &FRemoteControlPanelStyle::Get()->GetWidgetStyle<FRCPanelStyle>("RemoteControlPanel.MinorPanel");
+
+	if (URCBehaviourConditional* Behaviour = Cast<URCBehaviourConditional>(InBehaviourItem->GetBehaviour()))
+	{
+		if (URCController* Controller = Behaviour->ControllerWeakPtr.Get())
+		{
+			Conditions.Add(ConditionTypeIsEqual);
+
+			if (Controller->IsNumericType())
+			{
+				Conditions.Add(ConditionTypeIsLesserThan);
+				Conditions.Add(ConditionTypeIsGreaterThan);
+				Conditions.Add(ConditionTypeIsLesserThanOrEqualTo);
+				Conditions.Add(ConditionTypeIsGreaterThanOrEqualTo);
+			}
+
+			Conditions.Add(ConditionTypeElse);
+		}
+	}
+
+	ChildSlot
+	[
+		SNew(SVerticalBox)
+		// Conditions Panel Title
+		+ SVerticalBox::Slot()
+		.HAlign(HAlign_Left)
+		.Padding(FMargin(3.f))
+		.AutoHeight()
+		[
+			SNew(STextBlock)
+			.Text(LOCTEXT("ConditionsPanelTitle", "Condition"))
+			.Font(FRemoteControlPanelStyle::Get()->GetFontStyle("RemoteControlPanel.Actions.ValuePanelHeader"))
+		]
+
+		// Conditions List
+		+ SVerticalBox::Slot()
+		.HAlign(HAlign_Left)
+		.Padding(FMargin(3.f))
+		.AutoHeight()
+		[
+			SAssignNew(ListViewConditions, SListView<TSharedPtr<ERCBehaviourConditionType>>)
+			.ListItemsSource(&Conditions)
+			.Orientation(Orient_Horizontal) // Horizontal List
+			.OnSelectionChanged(this, &SRCBehaviourConditional::OnConditionsListSelectionChanged)
+			.OnGenerateRow(this, &SRCBehaviourConditional::OnGenerateWidgetForConditionsList)
+			// @todo: Add dark background for list as per Figma
+		]
+
+		// Comparand Label
+		+ SVerticalBox::Slot()
+		.HAlign(HAlign_Left)
+		.Padding(FMargin(3.f, 6.f, 3.f, 3.f))
+		.AutoHeight()
+		[
+			SNew(STextBlock)
+			.Text(LOCTEXT("ConditionsPanelComparandLabel", "Input Value"))
+			.Font(FRemoteControlPanelStyle::Get()->GetFontStyle("RemoteControlPanel.Actions.ValuePanelHeader"))
+		]
+
+		// Comparand Value Widget
+		+ SVerticalBox::Slot()
+		.HAlign(HAlign_Left)
+		.Padding(FMargin(3.f))
+		.AutoHeight()
+		[
+			SAssignNew(ComparandFieldBoxWidget, SBox)
+			[
+				InBehaviourItem->GetComparandFieldWidget()
+			]			
+		]
+	];
+
+	// Select equality condition by default
+	ListViewConditions->SetSelection(Conditions[0], ESelectInfo::Direct);
+}
+
+END_SLATE_FUNCTION_BUILD_OPTIMIZATION
+
+void SRCBehaviourConditional::Shutdown()
+{
+	// Reset all static shared pointers
+	ConditionTypeIsEqual.Reset();
+	ConditionTypeIsLesserThan.Reset();
+	ConditionTypeIsGreaterThan.Reset();
+	ConditionTypeIsLesserThanOrEqualTo.Reset();
+	ConditionTypeIsGreaterThan.Reset();
+	ConditionTypeElse.Reset();
+}
+
+void SRCBehaviourConditional::RefreshPropertyWidget()
+{
+	if (ComparandFieldBoxWidget)
+	{
+		if (TSharedPtr<const FRCBehaviourConditionalModel> BehaviourItem = ConditionalBehaviourItemWeakPtr.Pin())
+		{
+			ComparandFieldBoxWidget->SetContent(BehaviourItem->GetComparandFieldWidget());
+		}		
+	}
+}
+
+TSharedRef<ITableRow> SRCBehaviourConditional::OnGenerateWidgetForConditionsList(TSharedPtr<ERCBehaviourConditionType> InItem, const TSharedRef<STableViewBase>& OwnerTable)
+{
+	FText ConditionDisplayText;
+
+	if (TSharedPtr<const FRCBehaviourConditionalModel> BehaviourItem = ConditionalBehaviourItemWeakPtr.Pin())
+	{
+		if (URCBehaviourConditional* ConditionalBehaviour = Cast<URCBehaviourConditional>(BehaviourItem->GetBehaviour()))
+		{
+			if (InItem.IsValid())
+			{
+				ConditionDisplayText = ConditionalBehaviour->GetConditionTypeAsText(*InItem.Get());
+			}
+		}
+	}
+
+	return SNew(STableRow<TSharedPtr<FText>>, OwnerTable)
+		.Style(&RCPanelStyle->TableRowStyle)
+		[
+			SNew(SBox)
+			.Padding(FMargin(6.f))
+			[
+				SNew(STextBlock).Text(ConditionDisplayText)
+			]
+			
+		];
+}
+
+void SRCBehaviourConditional::OnConditionsListSelectionChanged(TSharedPtr<ERCBehaviourConditionType> InItem, ESelectInfo::Type)
+{
+	if (!ensure(InItem.IsValid()))
+	{
+		return;
+	}
+
+	const ERCBehaviourConditionType ConditionType = *InItem.Get();
+
+	if (TSharedPtr<FRCBehaviourConditionalModel> BehaviourItem = ConditionalBehaviourItemWeakPtr.Pin())
+	{
+		BehaviourItem->Condition = ConditionType;
+	}
+
+	// Else Condition does not need to show the Comparand field
+	const EVisibility bComparandFieldVisibility = (ConditionType == ERCBehaviourConditionType::Else) ? EVisibility::Hidden : EVisibility::Visible;
+	ComparandFieldBoxWidget->SetVisibility(bComparandFieldVisibility);
+}
+
+#undef LOCTEXT_NAMESPACE
\ No newline at end of file
diff --git a/Engine/Plugins/VirtualProduction/RemoteControl/Source/RemoteControlUI/Private/UI/Behaviour/Builtin/Conditional/SRCBehaviourConditional.h b/Engine/Plugins/VirtualProduction/RemoteControl/Source/RemoteControlUI/Private/UI/Behaviour/Builtin/Conditional/SRCBehaviourConditional.h
new file mode 100644
--- /dev/null
+++ b/Engine/Plugins/VirtualProduction/RemoteControl/Source/RemoteControlUI/Private/UI/Behaviour/Builtin/Conditional/SRCBehaviourConditional.h
@@ -0,0 +1,66 @@
+// Copyright Epic Games, Inc. All Rights Reserved.
+
+#pragma once
+
+#include "CoreMinimal.h"
+#include "RemoteControlField.h"
+#include "Widgets/SCompoundWidget.h"
+
+enum class ERCBehaviourConditionType : uint8;
+class FRCBehaviourConditionalModel;
+struct FRCPanelStyle;
+class ITableRow;
+class STableViewBase;
+class SBox;
+template <typename ItemType> class SListView;
+
+/*
+* ~ SRCBehaviourConditional ~
+*
+* Behaviour specific details panel for Conditional Behaviour
+* 
+* Provides users with the ability to definite Conditions on a per action basis
+* Contains a Conditions list (for choosing a Condition) and a Comparand input field (for specifying the value with which to compare, per condition)
+*/
+class REMOTECONTROLUI_API SRCBehaviourConditional : public SCompoundWidget
+{
+public:
+	SLATE_BEGIN_ARGS(SRCBehaviourConditional)
+		{
+		}
+
+	SLATE_END_ARGS()
+
+	/** Constructs this widget with InArgs */
+	void Construct(const FArguments& InArgs, TSharedRef<FRCBehaviourConditionalModel> InBehaviourItem);
+
+	/** Shutdown panel  (used to clear static shared pointers for the Conditions List) */
+	static void Shutdown();
+
+	/** Refreshes the Comparand input field.
+	* This needs to be performed each time the user creates a new Action, because the previous Comparand field has been consumed
+	*and so we create a new virtual property Comparand with whose memory the UI widget needs to be rebound.*/
+	void RefreshPropertyWidget();
+
+private:
+	/** OnGenerateRow delegate for the Conditions List*/
+	TSharedRef<ITableRow> OnGenerateWidgetForConditionsList(TSharedPtr<ERCBehaviourConditionType> InItem, const TSharedRef<STableViewBase>& OwnerTable);
+
+	/** OnSelectionChanged delegate for the Conditions List */
+	void OnConditionsListSelectionChanged(TSharedPtr<ERCBehaviourConditionType> InItem, ESelectInfo::Type);
+
+	/** List of possible Condition Types*/
+	TArray<TSharedPtr<ERCBehaviourConditionType>> Conditions;
+
+	/** The Behaviour (UI model) associated with us*/
+	TWeakPtr<FRCBehaviourConditionalModel> ConditionalBehaviourItemWeakPtr;
+
+	/** Panel Style reference. */
+	const FRCPanelStyle* RCPanelStyle;
+
+	/** Conditions List View widget - Enables the user to pick a Condition Type per Condition*/
+	TSharedPtr <SListView<TSharedPtr<ERCBehaviourConditionType>>> ListViewConditions;
+
+	/** Parent Box for the Comparand input field (which enables us to refresh it whenever required)*/
+	TSharedPtr<SBox> ComparandFieldBoxWidget;
+};
diff --git a/Engine/Plugins/VirtualProduction/RemoteControl/Source/RemoteControlUI/Private/UI/Behaviour/Builtin/RCBehaviourIsEqualModel.cpp b/Engine/Plugins/VirtualProduction/RemoteControl/Source/RemoteControlUI/Private/UI/Behaviour/Builtin/RCBehaviourIsEqualModel.cpp
--- a/Engine/Plugins/VirtualProduction/RemoteControl/Source/RemoteControlUI/Private/UI/Behaviour/Builtin/RCBehaviourIsEqualModel.cpp
+++ b/Engine/Plugins/VirtualProduction/RemoteControl/Source/RemoteControlUI/Private/UI/Behaviour/Builtin/RCBehaviourIsEqualModel.cpp
@@ -36,7 +36,7 @@ FRCIsEqualBehaviourModel::FRCIsEqualBehaviourModel(URCIsEqualBehaviour* IsEqualB
 	}
 }
 
-TSharedRef<SWidget> FRCIsEqualBehaviourModel::GetBehaviourDetailsWidget() const
+TSharedRef<SWidget> FRCIsEqualBehaviourModel::GetBehaviourDetailsWidget()
 {
 	return SNew(SRCBehaviourIsEqual, SharedThis(this));
 }
diff --git a/Engine/Plugins/VirtualProduction/RemoteControl/Source/RemoteControlUI/Private/UI/Behaviour/Builtin/RCBehaviourIsEqualModel.h b/Engine/Plugins/VirtualProduction/RemoteControl/Source/RemoteControlUI/Private/UI/Behaviour/Builtin/RCBehaviourIsEqualModel.h
--- a/Engine/Plugins/VirtualProduction/RemoteControl/Source/RemoteControlUI/Private/UI/Behaviour/Builtin/RCBehaviourIsEqualModel.h
+++ b/Engine/Plugins/VirtualProduction/RemoteControl/Source/RemoteControlUI/Private/UI/Behaviour/Builtin/RCBehaviourIsEqualModel.h
@@ -21,7 +21,7 @@ public:
 	FRCIsEqualBehaviourModel(URCIsEqualBehaviour* IsEqualBehaviour);
 
 	/** Builds a Behaviour specific widget as required for Is Equal behaviour*/
-	virtual TSharedRef<SWidget> GetBehaviourDetailsWidget() const override;
+	virtual TSharedRef<SWidget> GetBehaviourDetailsWidget() override;
 
 	/** Builds a generic Value Widget of the Controller's type
 	* Use to store user input for performing the Is Equal Behaviour comparison*/
diff --git a/Engine/Plugins/VirtualProduction/RemoteControl/Source/RemoteControlUI/Private/UI/Behaviour/RCBehaviourModel.cpp b/Engine/Plugins/VirtualProduction/RemoteControl/Source/RemoteControlUI/Private/UI/Behaviour/RCBehaviourModel.cpp
--- a/Engine/Plugins/VirtualProduction/RemoteControl/Source/RemoteControlUI/Private/UI/Behaviour/RCBehaviourModel.cpp
+++ b/Engine/Plugins/VirtualProduction/RemoteControl/Source/RemoteControlUI/Private/UI/Behaviour/RCBehaviourModel.cpp
@@ -4,6 +4,10 @@
 
 #include "Behaviour/RCBehaviourNode.h"
 #include "Engine/Blueprint.h"
+#include "Styling/RemoteControlStyles.h"
+#include "UI/Action/RCActionModel.h"
+#include "UI/Action/SRCActionPanel.h"
+#include "UI/Action/SRCActionPanelList.h"
 #include "UI/BaseLogicUI/RCLogicHelpers.h"
 #include "UI/RemoteControlPanelStyle.h"
 #include "UI/SRemoteControlPanel.h"
@@ -15,11 +19,27 @@
 FRCBehaviourModel::FRCBehaviourModel(URCBehaviour* InBehaviour)
 	: BehaviourWeakPtr(InBehaviour)
 {
+	RCPanelStyle = &FRemoteControlPanelStyle::Get()->GetWidgetStyle<FRCPanelStyle>("RemoteControlPanel.BehaviourPanel");
+
 	const FText BehaviorDisplayName = BehaviourWeakPtr->GetDisplayName().ToUpper();
 
 	SAssignNew(BehaviourTitleText, STextBlock)
 		.Text(BehaviorDisplayName)
-		.Font(FRemoteControlPanelStyle::Get()->GetFontStyle("RemoteControlPanel.Behaviours.Title"));
+		.TextStyle(&RCPanelStyle->HeaderTextStyle);
+}
+
+URCAction* FRCBehaviourModel::AddAction(const TSharedRef<const FRemoteControlField> InRemoteControlField)
+{
+	URCAction* NewAction = nullptr;
+
+	if (URCBehaviour* Behaviour = BehaviourWeakPtr.Get())
+	{
+		NewAction = Behaviour->AddAction(InRemoteControlField);
+
+		OnActionAdded(NewAction);
+	}
+
+	return NewAction;
 }
 
 TSharedRef<SWidget> FRCBehaviourModel::GetWidget() const
@@ -41,7 +61,7 @@ TSharedRef<SWidget> FRCBehaviourModel::GetWidget() const
 		];
 }
 
-TSharedRef<SWidget> FRCBehaviourModel::GetBehaviourDetailsWidget() const
+TSharedRef<SWidget> FRCBehaviourModel::GetBehaviourDetailsWidget()
 {
 	return SNullWidget::NullWidget;
 }
@@ -81,6 +101,13 @@ void FRCBehaviourModel::SetIsBehaviourEnabled(const bool bIsEnabled)
 	}
 }
 
+TSharedPtr<SRCLogicPanelListBase> FRCBehaviourModel::GetActionsListWidget(TSharedRef<SRCActionPanel> InActionPanel)
+{
+	// Returns the default Actions List; child classes can override as required
+
+	return SNew(SRCActionPanelList<FRCActionModel>, InActionPanel, SharedThis(this));
+}
+
 URCBehaviour* FRCBehaviourModel::GetBehaviour() const
 {
 	return BehaviourWeakPtr.Get();
diff --git a/Engine/Plugins/VirtualProduction/RemoteControl/Source/RemoteControlUI/Private/UI/Behaviour/RCBehaviourModel.h b/Engine/Plugins/VirtualProduction/RemoteControl/Source/RemoteControlUI/Private/UI/Behaviour/RCBehaviourModel.h
--- a/Engine/Plugins/VirtualProduction/RemoteControl/Source/RemoteControlUI/Private/UI/Behaviour/RCBehaviourModel.h
+++ b/Engine/Plugins/VirtualProduction/RemoteControl/Source/RemoteControlUI/Private/UI/Behaviour/RCBehaviourModel.h
@@ -4,8 +4,11 @@
 #include "Behaviour/RCIsEqualBehaviour.h"
 #include "UI/BaseLogicUI/RCLogicModeBase.h"
 
+struct FRCPanelStyle;
 class IPropertyRowGenerator;
 class STextBlock;
+class SRCActionPanel;
+class SRCLogicPanelListBase;
 class SWidget;
 class URCBehaviour;
 
@@ -20,13 +23,16 @@ class FRCBehaviourModel : public FRCLogicModeBase
 public:
 	FRCBehaviourModel(URCBehaviour* InBehaviour);
 
+	/** Add a Logic Action using a remote control field as input */
+	virtual URCAction* AddAction(const TSharedRef<const FRemoteControlField> InRemoteControlField);
+
 	/** The widget to be rendered for this Property in the Behaviour panel
 	* Currently displays Behaviour Name metadata in the Behaviours Panel List
 	*/
 	virtual TSharedRef<SWidget> GetWidget() const override;
 
 	/** Builds a Behaviour specific widget that child Behaviour classes can implement as required*/
-	virtual TSharedRef<SWidget> GetBehaviourDetailsWidget() const;
+	virtual TSharedRef<SWidget> GetBehaviourDetailsWidget();
 
 	/** Returns the Behaviour (Data model) associated with us*/
 	URCBehaviour* GetBehaviour() const;
@@ -40,10 +46,22 @@ public:
 	/** Set the Enabled state of our underlying Behaviour */
 	void SetIsBehaviourEnabled(const bool bIsEnabled);
 
+	/** The Actions List widget to be used for this Behaviour.
+	* The default actions table is used if a behaviour doesn't specify it.
+	*/
+	virtual TSharedPtr<SRCLogicPanelListBase> GetActionsListWidget(TSharedRef<SRCActionPanel> InActionPanel);
+
+protected:
+	/** Invoked after an action has been added for this Behaviour in the actions panel */
+	virtual void OnActionAdded(URCAction* Action) {}
+
 private:
 	/** The Behaviour (Data model) associated with us*/
 	TWeakObjectPtr<URCBehaviour> BehaviourWeakPtr;
 
 	/** Text block widget for representing the Behaviour's title */
 	TSharedPtr<STextBlock> BehaviourTitleText;
+
+	/** Panel Style reference. */
+	const FRCPanelStyle* RCPanelStyle;
 };
\ No newline at end of file
diff --git a/Engine/Plugins/VirtualProduction/RemoteControl/Source/RemoteControlUI/Private/UI/Behaviour/SRCBehaviourPanelList.cpp b/Engine/Plugins/VirtualProduction/RemoteControl/Source/RemoteControlUI/Private/UI/Behaviour/SRCBehaviourPanelList.cpp
--- a/Engine/Plugins/VirtualProduction/RemoteControl/Source/RemoteControlUI/Private/UI/Behaviour/SRCBehaviourPanelList.cpp
+++ b/Engine/Plugins/VirtualProduction/RemoteControl/Source/RemoteControlUI/Private/UI/Behaviour/SRCBehaviourPanelList.cpp
@@ -2,6 +2,7 @@
 
 #include "SRCBehaviourPanelList.h"
 
+#include "Behaviour/Builtin/Conditional/RCBehaviourConditional.h"
 #include "Controller/RCController.h"
 #include "RCBehaviourModel.h"
 #include "RemoteControlPreset.h"
@@ -9,6 +10,7 @@
 #include "SRCBehaviourPanel.h"
 #include "Styling/RemoteControlStyles.h"
 #include "UI/Behaviour/Builtin/RCBehaviourIsEqualModel.h"
+#include "UI/Behaviour/Builtin/Conditional/RCBehaviourConditionalModel.h"
 #include "UI/RemoteControlPanelStyle.h"
 #include "UI/SRemoteControlPanel.h"
 #include "UI/Controller/RCControllerModel.h"
@@ -87,6 +89,10 @@ void SRCBehaviourPanelList::Reset()
 				{
 					BehaviourItems.Add(MakeShared<FRCIsEqualBehaviourModel>(IsEqualBehaviour));
 				}
+				else if (URCBehaviourConditional* ConditionalBehaviour = Cast<URCBehaviourConditional>(Behaviour))
+				{
+					BehaviourItems.Add(MakeShared<FRCBehaviourConditionalModel>(ConditionalBehaviour));
+				}
 				else
 				{
 					BehaviourItems.Add(MakeShared<FRCBehaviourModel>(Behaviour));
diff --git a/Engine/Plugins/VirtualProduction/RemoteControl/Source/RemoteControlUI/Private/UI/Behaviour/SRCBehaviourPanelList.h b/Engine/Plugins/VirtualProduction/RemoteControl/Source/RemoteControlUI/Private/UI/Behaviour/SRCBehaviourPanelList.h
--- a/Engine/Plugins/VirtualProduction/RemoteControl/Source/RemoteControlUI/Private/UI/Behaviour/SRCBehaviourPanelList.h
+++ b/Engine/Plugins/VirtualProduction/RemoteControl/Source/RemoteControlUI/Private/UI/Behaviour/SRCBehaviourPanelList.h
@@ -40,10 +40,10 @@ public:
 	virtual int32 Num() const override;
 
 	/** Whether the Behaviours List View currently has focus.*/
-	bool IsListFocused() const;
+	virtual bool IsListFocused() const override;
 
 	/** Deletes currently selected items from the list view*/
-	void DeleteSelectedPanelItem();
+	virtual void DeleteSelectedPanelItem() override;
 
 private:
 
diff --git a/Engine/Plugins/VirtualProduction/RemoteControl/Source/RemoteControlUI/Private/UI/Controller/SRCControllerPanelList.h b/Engine/Plugins/VirtualProduction/RemoteControl/Source/RemoteControlUI/Private/UI/Controller/SRCControllerPanelList.h
--- a/Engine/Plugins/VirtualProduction/RemoteControl/Source/RemoteControlUI/Private/UI/Controller/SRCControllerPanelList.h
+++ b/Engine/Plugins/VirtualProduction/RemoteControl/Source/RemoteControlUI/Private/UI/Controller/SRCControllerPanelList.h
@@ -40,10 +40,10 @@ public:
 	virtual int32 Num() const override;
 
 	/** Whether the Controllers List View currently has focus.*/
-	bool IsListFocused() const;
+	virtual bool IsListFocused() const override;
 
 	/** Deletes currently selected items from the list view*/
-	void DeleteSelectedPanelItem();
+	virtual void DeleteSelectedPanelItem() override;
 
 	void EnterRenameMode();
 
diff --git a/Engine/Plugins/VirtualProduction/RemoteControl/Source/RemoteControlUI/Private/UI/RemoteControlPanelStyle.cpp b/Engine/Plugins/VirtualProduction/RemoteControl/Source/RemoteControlUI/Private/UI/RemoteControlPanelStyle.cpp
--- a/Engine/Plugins/VirtualProduction/RemoteControl/Source/RemoteControlUI/Private/UI/RemoteControlPanelStyle.cpp
+++ b/Engine/Plugins/VirtualProduction/RemoteControl/Source/RemoteControlUI/Private/UI/RemoteControlPanelStyle.cpp
@@ -99,7 +99,6 @@ void FRemoteControlPanelStyle::Initialize()
 	StyleSet->Set("Switch.ToggleOn", new IMAGE_PLUGIN_BRUSH("Icons/Switch_ON", Icon28x14));
 
 	// Remote Control Logic UI
-	StyleSet->Set("RemoteControlPanel.Behaviours.Title", FCoreStyle::GetDefaultFontStyle("Regular", 12));
 	StyleSet->Set("RemoteControlPanel.Behaviours.BehaviourDescription", FCoreStyle::GetDefaultFontStyle("Regular", 10));
 	StyleSet->Set("RemoteControlPanel.Behaviours.CustomBlueprint", new IMAGE_BRUSH_SVG("Starship/MainToolbar/blueprints", Icon20x20));
 	StyleSet->Set("RemoteControlPanel.Actions.ValuePanelHeader", FCoreStyle::GetDefaultFontStyle("Bold", 12));
@@ -206,6 +205,15 @@ void FRemoteControlPanelStyle::SetupPanelStyles(TSharedRef<FSlateStyleSet> InSty
 		.SetSectionHeaderTextStyle(HeaderText);
 
 	StyleSet->Set("RemoteControlPanel.MinorPanel", RCMinorPanelStyle);
+
+	// Behaviour Panel style
+	FRCPanelStyle RCBehaviourPanelStyle = RCMinorPanelStyle;
+
+	const FTextBlockStyle BehaviourPanelHeaderText = FTextBlockStyle(NormalText)
+		.SetFont(FCoreStyle::GetDefaultFontStyle("Regular", 12));
+	RCBehaviourPanelStyle.SetHeaderTextStyle(BehaviourPanelHeaderText);
+
+	StyleSet->Set("RemoteControlPanel.BehaviourPanel", RCBehaviourPanelStyle);
 }
 
 #undef IMAGE_PLUGIN_BRUSH
