SetAssetByPath Behaviour

Add Behaviour with custom Widgets, which takes a String Input and loads/sets the Asset onto a targeted Property.

diff --git a/Engine/Plugins/VirtualProduction/RemoteControl/Source/RemoteControlLogic/Private/Behaviour/Builtin/RCBehaviourSetAssetByPathNode.cpp b/Engine/Plugins/VirtualProduction/RemoteControl/Source/RemoteControlLogic/Private/Behaviour/Builtin/RCBehaviourSetAssetByPathNode.cpp
new file mode 100644
--- /dev/null
+++ b/Engine/Plugins/VirtualProduction/RemoteControl/Source/RemoteControlLogic/Private/Behaviour/Builtin/RCBehaviourSetAssetByPathNode.cpp
@@ -0,0 +1,83 @@
+// Copyright Epic Games, Inc. All Rights Reserved.
+
+#include "Behaviour/Builtin/RCBehaviourSetAssetByPathNode.h"
+
+#include "Behaviour/RCBehaviour.h"
+#include "Behaviour/RCSetAssetByPathBehaviour.h"
+#include "Controller/RCController.h"
+#include "RCVirtualProperty.h"
+#include "RCVirtualPropertyContainer.h"
+
+URCBehaviourSetAssetByPathNode::URCBehaviourSetAssetByPathNode()
+{
+	DisplayName = NSLOCTEXT("Remote Control Behaviour", "Behaviour Name - Set Asset By Path", "Set Asset By Path");
+	BehaviorDescription = NSLOCTEXT("Remote Control Behaviour", "Behavior Desc - Set Asset By Path", "Triggers an event which sets an object based on the TargetProperty Name.");
+}
+
+bool URCBehaviourSetAssetByPathNode::Execute_Implementation(URCBehaviour* InBehaviour) const
+{
+	if (!ensure(InBehaviour))
+	{
+		return false;
+	}
+
+	URCController* RCController = InBehaviour->ControllerWeakPtr.Get();
+	if (!RCController)
+	{
+		return false;
+	}
+
+	URCSetAssetByPathBehaviour* SetAssetByPathBehaviour = Cast<URCSetAssetByPathBehaviour>(InBehaviour);
+	if (!SetAssetByPathBehaviour)
+	{
+		return false;
+	}
+
+	FString ControllerString;
+	FString TargetPropertyString;
+	FString DefaultString;
+	
+	RCController->GetValueString(ControllerString);
+	SetAssetByPathBehaviour->PropertyInContainer->GetVirtualProperty(SetAssetByPathBehaviourHelpers::TargetProperty)->GetValueString(TargetPropertyString);
+	SetAssetByPathBehaviour->PropertyInContainer->GetVirtualProperty(SetAssetByPathBehaviourHelpers::DefaultProperty)->GetValueString(DefaultString);
+
+	// Add Path String Concat
+	FString ConcatPath = SetAssetByPathBehaviourHelpers::ContentFolder;
+	for (FString PathPart : SetAssetByPathBehaviour->PathStruct.PathArray)
+	{
+		if (PathPart.IsEmpty())
+		{
+			continue;
+		}
+		// Add '/' if needed
+		int32 IndexFound;
+
+		if (PathPart.FindLastChar('_', IndexFound) && IndexFound == PathPart.Len() - 1)
+		{
+			// In the case there's underscore char in the at the end of one of the Path Strings, do nothing to facilitate more complex pathing behaviours.
+		}
+		else if (!PathPart.FindLastChar('/', IndexFound) || IndexFound < PathPart.Len() - 1)
+		{
+			PathPart = PathPart.AppendChar('/');
+		}
+		
+		ConcatPath += PathPart;
+	}
+	
+	return SetAssetByPathBehaviour->SetAssetByPath(ConcatPath, TargetPropertyString, DefaultString);
+}
+
+bool URCBehaviourSetAssetByPathNode::IsSupported_Implementation(URCBehaviour* InBehaviour) const
+{
+	static TArray<TPair<EPropertyBagPropertyType, UObject*>> SupportedPropertyBagTypes =
+	{
+		{ EPropertyBagPropertyType::String, nullptr }
+	};
+	
+	return SupportedPropertyBagTypes.ContainsByPredicate(GetIsSupportedCallback(InBehaviour));
+}
+
+UClass* URCBehaviourSetAssetByPathNode::GetBehaviourClass() const
+{
+	return URCSetAssetByPathBehaviour::StaticClass();
+}
diff --git a/Engine/Plugins/VirtualProduction/RemoteControl/Source/RemoteControlLogic/Private/Behaviour/RCBehaviour.cpp b/Engine/Plugins/VirtualProduction/RemoteControl/Source/RemoteControlLogic/Private/Behaviour/RCBehaviour.cpp
--- a/Engine/Plugins/VirtualProduction/RemoteControl/Source/RemoteControlLogic/Private/Behaviour/RCBehaviour.cpp
+++ b/Engine/Plugins/VirtualProduction/RemoteControl/Source/RemoteControlLogic/Private/Behaviour/RCBehaviour.cpp
@@ -16,10 +16,6 @@ URCBehaviour::URCBehaviour()
 
 void URCBehaviour::Execute()
 {
-	if (!ActionContainer->Actions.Num())
-	{
-		return;
-	}
 	const URCBehaviourNode* BehaviourNode = GetBehaviourNode();
 
 	// Execute before the logic
diff --git a/Engine/Plugins/VirtualProduction/RemoteControl/Source/RemoteControlLogic/Private/Behaviour/RCSetAssetByPathBehaviour.cpp b/Engine/Plugins/VirtualProduction/RemoteControl/Source/RemoteControlLogic/Private/Behaviour/RCSetAssetByPathBehaviour.cpp
new file mode 100644
--- /dev/null
+++ b/Engine/Plugins/VirtualProduction/RemoteControl/Source/RemoteControlLogic/Private/Behaviour/RCSetAssetByPathBehaviour.cpp
@@ -0,0 +1,188 @@
+// Copyright Epic Games, Inc. All Rights Reserved.
+
+#include "Behaviour/RCSetAssetByPathBehaviour.h"
+
+#include "Backends/CborStructDeserializerBackend.h"
+#include "Backends/CborStructSerializerBackend.h"
+#include "Components/MeshComponent.h"
+#include "Components/StaticMeshComponent.h"
+#include "Controller/RCController.h"
+#include "Engine/StaticMesh.h"
+#include "Engine/World.h"
+#include "IRemoteControlModule.h"
+#include "IRemoteControlPropertyHandle.h"
+#include "Materials/MaterialInstance.h"
+#include "RCVirtualProperty.h"
+#include "RCVirtualPropertyContainer.h"
+#include "RemoteControlPreset.h"
+#include "StructSerializer.h"
+
+URCSetAssetByPathBehaviour::URCSetAssetByPathBehaviour()
+{
+	PropertyInContainer = CreateDefaultSubobject<URCVirtualPropertyContainerBase>(FName("VirtualPropertyInContainer"));
+}
+
+void URCSetAssetByPathBehaviour::Initialize()
+{
+	PropertyInContainer->AddProperty(SetAssetByPathBehaviourHelpers::TargetProperty, URCController::StaticClass(), EPropertyBagPropertyType::String);
+	PropertyInContainer->AddProperty(SetAssetByPathBehaviourHelpers::DefaultProperty, URCController::StaticClass(), EPropertyBagPropertyType::String);
+	
+	Super::Initialize();
+}
+
+bool URCSetAssetByPathBehaviour::SetAssetByPath(const FString& AssetPath, const FString& TargetPropertyString, const FString& DefaultString)
+{
+	URCController* Controller = ControllerWeakPtr.Get();
+	
+	if (AssetPath.IsEmpty() || !Controller)
+	{
+		return false;
+	}
+	
+	FString ControllerString;
+	if (!Controller->GetValueString(ControllerString))
+	{
+		return false;
+	}
+	
+	FSoftObjectPath MainObjectRef(AssetPath + ControllerString);
+	
+	if (UObject* Object = MainObjectRef.TryLoad())
+	{
+		return SetAsset(Object, TargetPropertyString);
+	}
+	
+	if (!DefaultString.IsEmpty())
+	{
+		FSoftObjectPath DefaultObjectRef(AssetPath + DefaultString);
+		
+		return SetAsset(DefaultObjectRef.TryLoad(), TargetPropertyString);
+	}
+	
+	return false;
+}
+
+bool URCSetAssetByPathBehaviour::SetAsset(UObject* SetterObject, const FString& PropertyString)
+{
+	URCController* Controller = ControllerWeakPtr.Get();
+	if (!SetterObject || PropertyString.IsEmpty() || !Controller)
+	{
+		return false;
+	}
+	
+	const FName TargetPropertyLabel = FName(*PropertyString);
+	TWeakObjectPtr<URemoteControlPreset> PresetPtr = Controller->PresetWeakPtr;
+	TArray<TWeakPtr<FRemoteControlProperty>> ExposedProperties = PresetPtr->GetExposedEntities<FRemoteControlProperty>();
+	TArray<TWeakPtr<FRemoteControlEntity>> ExposedEntities = PresetPtr->GetExposedEntities<FRemoteControlEntity>();
+	
+	TWeakPtr<FRemoteControlEntity>* ExposedEntityPtr = ExposedEntities.FindByPredicate([TargetPropertyLabel](const TWeakPtr<FRemoteControlEntity>& Entity) { return Entity.Pin()->GetLabel() == TargetPropertyLabel;});
+	TWeakPtr<FRemoteControlProperty>* ExposedPropertyPtr = ExposedProperties.FindByPredicate([TargetPropertyLabel](const TWeakPtr<FRemoteControlProperty>& Property){ return Property.Pin()->GetLabel() == TargetPropertyLabel; });
+	
+	if (ExposedPropertyPtr)
+	{
+		TWeakPtr<FRemoteControlProperty> ExposedProperty = *ExposedPropertyPtr;
+		if (AssetClass == UStaticMesh::StaticClass() && SetterObject->IsA(UStaticMesh::StaticClass()))
+		{
+			UStaticMeshComponent* StaticMeshComponent = Cast<UStaticMeshComponent>(ExposedProperty.Pin()->GetBoundObject());
+			StaticMeshComponent->SetStaticMesh(Cast<UStaticMesh>(SetterObject));
+			
+			return true;
+		}
+		
+		if (AssetClass == UMaterial::StaticClass() && SetterObject->IsA(UMaterial::StaticClass()))
+		{
+			if (UMeshComponent* MaterialOwnerComponent = Cast<UMeshComponent>(ExposedProperty.Pin()->GetBoundObject()))
+			{
+				MaterialOwnerComponent->SetMaterial(0, Cast<UMaterial>(SetterObject));
+			}
+			
+			return true;
+		}
+		
+		if (GetSupportedClasses().Contains(AssetClass) && SetterObject->IsA(AssetClass))
+		{
+			FRCObjectReference ObjectRef;
+			TSharedPtr<FRemoteControlProperty> RemoteControlProperty = ExposedProperty.Pin();
+			ObjectRef.Property = RemoteControlProperty->GetProperty();
+			ObjectRef.Access = ExposedProperty.Pin()->GetPropertyHandle()->ShouldGenerateTransaction() ? ERCAccess::WRITE_TRANSACTION_ACCESS : ERCAccess::WRITE_ACCESS;
+			ObjectRef.PropertyPathInfo = RemoteControlProperty->FieldPathInfo.ToString();
+		
+			for (UObject* Object : RemoteControlProperty->GetBoundObjects())
+			{
+				if (IRemoteControlModule::Get().ResolveObjectProperty(ObjectRef.Access, Object, ObjectRef.PropertyPathInfo, ObjectRef))
+				{
+					/**
+					 * Setting with Pointer and using Serialization of itself afterwards as a workaround with the Texture Asset not updating in the world.
+					 */
+					FObjectPropertyBase* ObjectProperty = CastField<FObjectPropertyBase>(ExposedProperty.Pin()->GetProperty());
+					if (!ObjectProperty)
+					{
+						return false;
+					}
+
+					FProperty* Property = RemoteControlProperty->GetProperty();
+					uint8* ValueAddress = Property->ContainerPtrToValuePtr<uint8>(ObjectRef.ContainerAdress);
+					ObjectProperty->SetObjectPropertyValue(ValueAddress, SetterObject);
+					
+					TArray<uint8> Buffer;
+					FMemoryReader Reader(Buffer);
+					FCborStructDeserializerBackend DeserializerBackend(Reader);
+
+					IRemoteControlModule::Get().SetObjectProperties(ObjectRef, DeserializerBackend, ERCPayloadType::Cbor, Buffer);
+					
+					return true;
+				}
+			}
+		}
+		
+		return false;
+	}
+	
+	if (ExposedEntityPtr)
+	{
+		TWeakPtr<FRemoteControlEntity> ExposedEntity = *ExposedEntityPtr;
+		if (AssetClass == UBlueprint::StaticClass() && SetterObject->IsA(UBlueprint::StaticClass()))
+		{
+
+			AActor* OldActor = Cast<AActor>(ExposedEntity.Pin()->GetBoundObject());
+			UWorld* World = OldActor->GetWorld();
+			
+			OldActor->UnregisterAllComponents();
+			
+			FVector OldLocation = OldActor->GetActorLocation();
+			FRotator OldRotation = OldActor->GetActorRotation();
+			const FName OldActorName = OldActor->GetFName();
+			
+			FActorSpawnParameters SpawnParams;
+			SpawnParams.Name = OldActorName;
+			
+			const FName OldActorReplacedNamed = MakeUniqueObjectName(OldActor->GetOuter(), OldActor->GetClass(), *FString::Printf(TEXT("%s_REPLACED"), *OldActor->GetFName().ToString()));
+			OldActor->Rename(*OldActorReplacedNamed.ToString(), OldActor->GetOuter(), REN_DoNotDirty | REN_DontCreateRedirectors | REN_ForceNoResetLoaders);
+
+			if (AActor* NewActor = World->SpawnActor(Cast<UBlueprint>(SetterObject)->GeneratedClass, &OldLocation, &OldRotation, SpawnParams))
+			{
+				World->DestroyActor(OldActor);
+				
+				return true;
+			}
+
+			OldActor->UObject::Rename(*OldActorName.ToString(), OldActor->GetOuter(), REN_DoNotDirty | REN_DontCreateRedirectors | REN_ForceNoResetLoaders);
+			OldActor->RegisterAllComponents();
+		}
+	}
+	
+	return false;
+}
+
+TArray<UClass*> URCSetAssetByPathBehaviour::GetSupportedClasses() const
+{
+	static TArray<UClass*> SupportedClasses =
+	{
+		UStaticMesh::StaticClass(),
+		UMaterial::StaticClass(),
+		UTexture::StaticClass(),
+		UBlueprint::StaticClass(),
+	};
+
+	return SupportedClasses;
+}
diff --git a/Engine/Plugins/VirtualProduction/RemoteControl/Source/RemoteControlLogic/Public/Behaviour/Builtin/RCBehaviourSetAssetByPathNode.h b/Engine/Plugins/VirtualProduction/RemoteControl/Source/RemoteControlLogic/Public/Behaviour/Builtin/RCBehaviourSetAssetByPathNode.h
new file mode 100644
--- /dev/null
+++ b/Engine/Plugins/VirtualProduction/RemoteControl/Source/RemoteControlLogic/Public/Behaviour/Builtin/RCBehaviourSetAssetByPathNode.h
@@ -0,0 +1,28 @@
+// Copyright Epic Games, Inc. All Rights Reserved.
+
+#pragma once
+
+#include "Behaviour/RCBehaviourNode.h"
+#include "RCBehaviourSetAssetByPathNode.generated.h"
+
+/**
+ * Takes the string as a path and goes on to search for the Asset it is connected to, setting it to a Target Exposed Property.
+ */
+UCLASS()
+class REMOTECONTROLLOGIC_API URCBehaviourSetAssetByPathNode : public URCBehaviourNode
+{
+	GENERATED_BODY()
+	
+public:
+	URCBehaviourSetAssetByPathNode();
+	
+	//~ Begin URCBehaviourNode interface
+	UFUNCTION(BlueprintNativeEvent, BlueprintCallable, Category = "Remote Control Behaviour")
+	bool Execute(URCBehaviour* InBehaviour) const override;
+
+	UFUNCTION(BlueprintNativeEvent, BlueprintCallable, Category = "Remote Control Behaviour")
+	bool IsSupported(URCBehaviour* InBehaviour) const override;
+
+	UClass* GetBehaviourClass() const override;
+	//~ End URCBehaviourNode interface
+};
diff --git a/Engine/Plugins/VirtualProduction/RemoteControl/Source/RemoteControlLogic/Public/Behaviour/RCSetAssetByPathBehaviour.h b/Engine/Plugins/VirtualProduction/RemoteControl/Source/RemoteControlLogic/Public/Behaviour/RCSetAssetByPathBehaviour.h
new file mode 100644
--- /dev/null
+++ b/Engine/Plugins/VirtualProduction/RemoteControl/Source/RemoteControlLogic/Public/Behaviour/RCSetAssetByPathBehaviour.h
@@ -0,0 +1,63 @@
+// Copyright Epic Games, Inc. All Rights Reserved.
+
+#pragma once
+
+#include "RCBehaviour.h"
+#include "RCVirtualProperty.h"
+
+#include "RCSetAssetByPathBehaviour.generated.h"
+
+class URCVirtualPropertyContainerBase;
+class URCUserDefinedStruct;
+
+namespace SetAssetByPathBehaviourHelpers
+{
+	const FString ContentFolder = FString(TEXT("/Game/"));
+	const FName TargetProperty = FName(TEXT("Target Property"));
+	const FName DefaultProperty = FName(TEXT("Default Property"));
+}
+
+/** Struct to help generate Widgts for the DetailsPanel of the Bahviour */
+USTRUCT()
+struct FRCSetAssetPath
+{
+	GENERATED_BODY()
+
+	FRCSetAssetPath() = default;
+	
+	/** An Array of Strings holding the Path of an Asset, seperated in several String. Will concated back together later. */
+	UPROPERTY()
+	TArray<FString> PathArray;
+};
+
+/**
+ * Custom behaviour for Set Asset By Path
+ */
+UCLASS()
+class REMOTECONTROLLOGIC_API URCSetAssetByPathBehaviour : public URCBehaviour
+{
+	GENERATED_BODY()
+public:
+	URCSetAssetByPathBehaviour();
+
+	//~ Begin URCBehaviour interface
+	virtual void Initialize() override;
+	//~ End URCBehaviour interface
+
+	bool SetAssetByPath(const FString& AssetPath, const FString& TargetPropertyString, const FString& DefaultString);
+
+	TArray<UClass*> GetSupportedClasses() const;
+public:
+	/** Pointer to property container */
+	UPROPERTY()
+	TObjectPtr<URCVirtualPropertyContainerBase> PropertyInContainer;
+
+	UPROPERTY()
+	UClass* AssetClass;
+
+	UPROPERTY()
+	FRCSetAssetPath PathStruct;
+
+private:
+	bool SetAsset(UObject* SetterObject, const FString& PropertyString);
+};
diff --git a/Engine/Plugins/VirtualProduction/RemoteControl/Source/RemoteControlUI/Private/UI/Behaviour/Builtin/RCBehaviourSetAssetByPathModel.cpp b/Engine/Plugins/VirtualProduction/RemoteControl/Source/RemoteControlUI/Private/UI/Behaviour/Builtin/RCBehaviourSetAssetByPathModel.cpp
new file mode 100644
--- /dev/null
+++ b/Engine/Plugins/VirtualProduction/RemoteControl/Source/RemoteControlUI/Private/UI/Behaviour/Builtin/RCBehaviourSetAssetByPathModel.cpp
@@ -0,0 +1,203 @@
+// Copyright Epic Games, Inc. All Rights Reserved.
+
+#include "RCBehaviourSetAssetByPathModel.h"
+
+#include "IDetailTreeNode.h"
+#include "IPropertyRowGenerator.h"
+#include "Modules/ModuleManager.h"
+#include "RCVirtualProperty.h"
+#include "RCVirtualPropertyContainer.h"
+#include "RemoteControlPreset.h"
+#include "SRCBehaviourSetAssetByPath.h"
+#include "Widgets/Layout/SBox.h"
+
+#define LOCTEXT_NAMESPACE "FRCSetAssetByPathBehaviourModel"
+
+FRCSetAssetByPathBehaviourModel::FRCSetAssetByPathBehaviourModel(URCSetAssetByPathBehaviour* SetAssetByPathBehaviour)
+	: FRCBehaviourModel(SetAssetByPathBehaviour)
+	, SetAssetByPathBehaviourWeakPtr(SetAssetByPathBehaviour)
+{
+	FPropertyRowGeneratorArgs Args;
+	Args.bShouldShowHiddenProperties = true;
+	PropertyRowGenerator = FModuleManager::GetModuleChecked<FPropertyEditorModule>("PropertyEditor").CreatePropertyRowGenerator(Args);
+	PropertyRowGeneratorArray = FModuleManager::GetModuleChecked<FPropertyEditorModule>("PropertyEditor").CreatePropertyRowGenerator(Args);
+
+	if (SetAssetByPathBehaviour)
+	{
+		PropertyRowGenerator->SetStructure(SetAssetByPathBehaviour->PropertyInContainer->CreateStructOnScope());
+		DetailTreeNodeWeakPtr.Empty();
+		for (const TSharedRef<IDetailTreeNode>& CategoryNode : PropertyRowGenerator->GetRootTreeNodes())
+		{
+			TArray<TSharedRef<IDetailTreeNode>> Children;
+			CategoryNode->GetChildren(Children);
+			for (const TSharedRef<IDetailTreeNode>& Child : Children)
+			{
+				DetailTreeNodeWeakPtr.Add(Child);
+			}
+		}
+
+		// Secondary TArray Struct
+		PropertyRowGeneratorArray->SetStructure(MakeShareable(new FStructOnScope(FRCSetAssetPath::StaticStruct(), (uint8*) &SetAssetByPathBehaviour->PathStruct)));
+		PropertyRowGeneratorArray->OnRowsRefreshed().AddLambda([this]()
+		{
+			RegeneratePathArrayWidget();
+		});
+		DetailTreeNodeWeakPtrArray.Empty();
+		for (const TSharedRef<IDetailTreeNode>& CategoryNode : PropertyRowGeneratorArray->GetRootTreeNodes())
+		{
+			TArray<TSharedRef<IDetailTreeNode>> Children;
+			CategoryNode->GetChildren(Children);
+			for (const TSharedRef<IDetailTreeNode>& Child : Children)
+			{
+				DetailTreeNodeWeakPtrArray.Add(Child);
+			}
+		}
+	}
+	PathArrayWidget = SNew(SBox);
+	RegeneratePathArrayWidget();
+}
+
+TSharedRef<SWidget> FRCSetAssetByPathBehaviourModel::GetBehaviourDetailsWidget()
+{
+	return SNew(SRCBehaviourSetAssetByPath, SharedThis(this));
+}
+
+TSharedRef<SWidget> FRCSetAssetByPathBehaviourModel::GetPropertyWidget() const
+{
+	TSharedRef<SVerticalBox> FieldWidget = SNew(SVerticalBox);
+	
+	const TSharedPtr<IDetailTreeNode> PinnedNodeTargetProperty = DetailTreeNodeWeakPtr[0];
+	const TSharedPtr<IDetailTreeNode> PinnedNodeDefault = DetailTreeNodeWeakPtr[1];
+
+	const FNodeWidgets NodeWidgetsTarget = PinnedNodeTargetProperty->CreateNodeWidgets();
+	const FNodeWidgets NodeWidgetsDefault = PinnedNodeDefault->CreateNodeWidgets();
+	if (NodeWidgetsTarget.ValueWidget && NodeWidgetsDefault.ValueWidget)
+	{
+		FieldWidget->AddSlot()
+			.Padding(FMargin(3.0f, 2.0f))
+			.AutoHeight()
+			[
+				SNew(SHorizontalBox)
+				+ SHorizontalBox::Slot()
+				[
+					NodeWidgetsTarget.NameWidget.ToSharedRef()
+				]
+				+ SHorizontalBox::Slot()
+				[
+					NodeWidgetsDefault.NameWidget.ToSharedRef()
+				]
+			];
+		
+		FieldWidget->AddSlot()
+			.Padding(FMargin(3.0f, 2.0f))
+			.AutoHeight()
+			[
+				SNew(SHorizontalBox)
+				+ SHorizontalBox::Slot()
+				[
+					NodeWidgetsTarget.ValueWidget.ToSharedRef()
+				]
+				+ SHorizontalBox::Slot()
+				[
+					NodeWidgetsDefault.ValueWidget.ToSharedRef()
+				]
+			];
+	}
+	else if (NodeWidgetsTarget.WholeRowWidget && NodeWidgetsDefault.WholeRowWidget)
+	{
+		FieldWidget->AddSlot()
+			.Padding(FMargin(3.0f, 2.0f))
+			.AutoHeight()
+			[
+				NodeWidgetsTarget.WholeRowWidget.ToSharedRef()
+			];
+		
+		FieldWidget->AddSlot()
+			.Padding(FMargin(3.0f, 2.0f))
+			.AutoHeight()
+			[
+				NodeWidgetsDefault.WholeRowWidget.ToSharedRef()
+			];
+	}
+
+	FieldWidget->AddSlot()
+		.Padding(FMargin(3.f, 0.f))
+		.AutoHeight()
+		[
+			PathArrayWidget->AsShared()
+		];
+	
+	return 	SNew(SHorizontalBox)
+		.Clipping(EWidgetClipping::OnDemand)
+		+ SHorizontalBox::Slot()
+		.VAlign(VAlign_Center)
+		.AutoWidth()
+		.FillWidth(1.0f)
+		[
+			FieldWidget
+		];
+}
+
+void FRCSetAssetByPathBehaviourModel::RegenerateWeakPtrInternal()
+{
+	DetailTreeNodeWeakPtrArray.Empty();
+	for (const TSharedRef<IDetailTreeNode>& CategoryNode : PropertyRowGeneratorArray->GetRootTreeNodes())
+	{
+		TArray<TSharedRef<IDetailTreeNode>> Children;
+		CategoryNode->GetChildren(Children);
+		for (const TSharedRef<IDetailTreeNode>& Child : Children)
+		{
+			DetailTreeNodeWeakPtrArray.Add(Child);
+		}
+	}
+}
+
+void FRCSetAssetByPathBehaviourModel::RegeneratePathArrayWidget()
+{
+	TSharedPtr<SVerticalBox> FieldArrayWidget = SNew(SVerticalBox);
+
+	RegenerateWeakPtrInternal();
+	
+	for (const TSharedPtr<IDetailTreeNode>& DetailTreeNodeArray : DetailTreeNodeWeakPtrArray)
+	{
+		const TSharedPtr<IDetailTreeNode> PinnedNode = DetailTreeNodeArray;
+
+		TArray<TSharedRef<IDetailTreeNode>> Children;
+		PinnedNode->GetChildren(Children);
+		Children.Insert(PinnedNode.ToSharedRef(), 0);
+		
+		for (uint8 Counter = 0; Counter < Children.Num(); Counter++)
+		{
+			const FNodeWidgets NodeWidgets = Children[Counter]->CreateNodeWidgets();
+			if (NodeWidgets.ValueWidget)
+			{
+				FieldArrayWidget->AddSlot()
+					.Padding(FMargin(3.0f, 2.0f))
+					.AutoHeight()
+					[
+						SNew(SHorizontalBox)
+						+ SHorizontalBox::Slot()
+						[
+							NodeWidgets.NameWidget.ToSharedRef()
+						]
+						+ SHorizontalBox::Slot()
+						[
+							NodeWidgets.ValueWidget.ToSharedRef()
+						]
+					];
+			}
+			else if (NodeWidgets.WholeRowWidget)
+			{
+				FieldArrayWidget->AddSlot()
+					.Padding(FMargin(3.0f, 2.0f))
+					.AutoHeight()
+					[
+						NodeWidgets.WholeRowWidget.ToSharedRef()
+					];
+			}
+		}
+	}
+	PathArrayWidget->SetContent(FieldArrayWidget.ToSharedRef());
+}
+
+#undef LOCTEXT_NAMESPACE
diff --git a/Engine/Plugins/VirtualProduction/RemoteControl/Source/RemoteControlUI/Private/UI/Behaviour/Builtin/RCBehaviourSetAssetByPathModel.h b/Engine/Plugins/VirtualProduction/RemoteControl/Source/RemoteControlUI/Private/UI/Behaviour/Builtin/RCBehaviourSetAssetByPathModel.h
new file mode 100644
--- /dev/null
+++ b/Engine/Plugins/VirtualProduction/RemoteControl/Source/RemoteControlUI/Private/UI/Behaviour/Builtin/RCBehaviourSetAssetByPathModel.h
@@ -0,0 +1,58 @@
+// Copyright Epic Games, Inc. All Rights Reserved.
+
+#pragma once
+#include "Behaviour/RCSetAssetByPathBehaviour.h"
+#include "UI/Behaviour/RCBehaviourModel.h"
+
+class IDetailTreeNode;
+class SBox;
+
+
+/*
+ * ~ FRCSetAssetByPathModel ~
+ *
+ * Child Behaviour class representing the "Set Asset By Path" Behaviour's UI model.
+ *
+ * Generates several Widgets where users can enter the RootPath, Target Property, and Default Property as FStrings.
+ * The values are then put together, getting a path towards a possible object.
+ */
+class FRCSetAssetByPathBehaviourModel : public FRCBehaviourModel
+{
+public:
+	FRCSetAssetByPathBehaviourModel(URCSetAssetByPathBehaviour* SetAssetByPathBehaviour);
+
+	/** Builds a Behaviour specific widget as required for the Set Asset By Path Behaviour */
+	virtual TSharedRef<SWidget> GetBehaviourDetailsWidget() override;
+
+	/*
+	 * Builds the Property Details Widget, including a generic expendable Array Widget and
+	 * further two Text Widgets representing elements needed for the Path Behaviour.
+	 * All of them store the user input and use them to perform the SetAssetByPath Behaviour.
+	 */
+	TSharedRef<SWidget> GetPropertyWidget() const;
+	
+private:
+	/** The SetAssetByPath Behaviour associated with this Model */
+	TWeakObjectPtr<URCSetAssetByPathBehaviour> SetAssetByPathBehaviourWeakPtr;
+
+	/** Pointer to the Widget holding the PathArray created for this behaviour */
+	TSharedPtr<SBox> PathArrayWidget;
+
+	/** The row generator used to build the generic value widgets */
+	TSharedPtr<IPropertyRowGenerator> PropertyRowGenerator;
+	
+	/** The row generator used to build the generic array value widget for the Path Array */
+	TSharedPtr<IPropertyRowGenerator> PropertyRowGeneratorArray;
+
+	/** Used to create a generic Value Widget based on the active Controller's type*/
+	TArray<TSharedPtr<IDetailTreeNode>> DetailTreeNodeWeakPtr;
+	
+	/** Used to create a generic Value Widget based on the Paths Available*/
+	TArray<TSharedPtr<IDetailTreeNode>> DetailTreeNodeWeakPtrArray;
+
+private:
+	void RegenerateWeakPtrInternal();
+	
+	/** Regenerates and creates a new PathArray Widget if changed */
+	void RegeneratePathArrayWidget();
+};
diff --git a/Engine/Plugins/VirtualProduction/RemoteControl/Source/RemoteControlUI/Private/UI/Behaviour/Builtin/SRCBehaviourSetAssetByPath.cpp b/Engine/Plugins/VirtualProduction/RemoteControl/Source/RemoteControlUI/Private/UI/Behaviour/Builtin/SRCBehaviourSetAssetByPath.cpp
new file mode 100644
--- /dev/null
+++ b/Engine/Plugins/VirtualProduction/RemoteControl/Source/RemoteControlUI/Private/UI/Behaviour/Builtin/SRCBehaviourSetAssetByPath.cpp
@@ -0,0 +1,74 @@
+// Copyright Epic Games, Inc. All Rights Reserved.
+
+#include "SRCBehaviourSetAssetByPath.h"
+
+#include "IDetailTreeNode.h"
+#include "SPositiveActionButton.h"
+#include "SlateOptMacros.h"
+#include "UI/Behaviour/Builtin/RCBehaviourSetAssetByPathModel.h"
+#include "UI/RemoteControlPanelStyle.h"
+#include "Widgets/Input/SButton.h"
+
+#define LOCTEXT_NAMESPACE "SRCBehaviourSetAssetByPath"
+
+class SPositiveActionButton;
+
+BEGIN_SLATE_FUNCTION_BUILD_OPTIMIZATION
+
+void SRCBehaviourSetAssetByPath::Construct(const FArguments& InArgs, TSharedRef<const FRCSetAssetByPathBehaviourModel> InBehaviourItem)
+{
+	
+
+	TSharedRef<SHorizontalBox> SelectedClassWidget = SNew(SHorizontalBox);
+	URCSetAssetByPathBehaviour* PathBehaviour = Cast<URCSetAssetByPathBehaviour>(InBehaviourItem->GetBehaviour());
+
+	if (!PathBehaviour)
+	{
+		ChildSlot
+		[
+			SNullWidget::NullWidget
+		];
+		return;
+	}
+		
+	for (UClass* SupportedClass : PathBehaviour->GetSupportedClasses())
+	{
+		SelectedClassWidget->AddSlot()
+		.HAlign(HAlign_Right)
+		.VAlign(VAlign_Center)
+		.Padding(FMargin(0.0f,0.0f))
+		.FillWidth(1.0f)
+		.AutoWidth()
+		[
+			SNew(SButton)
+			.Text(FText::FromString(SupportedClass->GetName()))
+			.ButtonColorAndOpacity_Lambda([PathBehaviour, SupportedClass]()
+			{
+				FSlateColor Color = PathBehaviour->AssetClass == SupportedClass ? FAppStyle::Get().GetSlateColor("Colors.AccentBlue") : FAppStyle::Get().GetSlateColor("Colors.AccentWhite");
+				return Color;
+			})
+			.OnPressed_Lambda([PathBehaviour, SupportedClass]()
+			{
+				PathBehaviour->AssetClass = SupportedClass;
+			})
+		];
+	}
+	
+	ChildSlot
+	[
+		SNew(SVerticalBox)
+		+ SVerticalBox::Slot()
+		[
+			SelectedClassWidget
+		]
+		+ SVerticalBox::Slot()
+		.AutoHeight()
+		[
+			InBehaviourItem->GetPropertyWidget()
+		]
+	];
+}
+
+END_SLATE_FUNCTION_BUILD_OPTIMIZATION
+
+#undef LOCTEXT_NAMESPACE
diff --git a/Engine/Plugins/VirtualProduction/RemoteControl/Source/RemoteControlUI/Private/UI/Behaviour/Builtin/SRCBehaviourSetAssetByPath.h b/Engine/Plugins/VirtualProduction/RemoteControl/Source/RemoteControlUI/Private/UI/Behaviour/Builtin/SRCBehaviourSetAssetByPath.h
new file mode 100644
--- /dev/null
+++ b/Engine/Plugins/VirtualProduction/RemoteControl/Source/RemoteControlUI/Private/UI/Behaviour/Builtin/SRCBehaviourSetAssetByPath.h
@@ -0,0 +1,25 @@
+// Copyright Epic Games, Inc. All Rights Reserved.
+
+#pragma once
+
+#include "CoreMinimal.h"
+#include "Widgets/SCompoundWidget.h"
+
+class FRCSetAssetByPathBehaviourModel;
+
+class REMOTECONTROLUI_API SRCBehaviourSetAssetByPath : public  SCompoundWidget
+{
+public:
+	SLATE_BEGIN_ARGS(SRCBehaviourSetAssetByPath)
+		{
+		}
+
+	SLATE_END_ARGS()
+
+	/** Constructs this widget with InArgs */
+	void Construct(const FArguments& InArgs, TSharedRef<const FRCSetAssetByPathBehaviourModel> InBehaviourItem);
+
+private:
+	/** The Behaviour (UI model) associated with us*/
+	TWeakPtr<const FRCSetAssetByPathBehaviourModel> SetAssetByPathWeakPtr;
+};
diff --git a/Engine/Plugins/VirtualProduction/RemoteControl/Source/RemoteControlUI/Private/UI/Behaviour/SRCBehaviourPanelList.cpp b/Engine/Plugins/VirtualProduction/RemoteControl/Source/RemoteControlUI/Private/UI/Behaviour/SRCBehaviourPanelList.cpp
--- a/Engine/Plugins/VirtualProduction/RemoteControl/Source/RemoteControlUI/Private/UI/Behaviour/SRCBehaviourPanelList.cpp
+++ b/Engine/Plugins/VirtualProduction/RemoteControl/Source/RemoteControlUI/Private/UI/Behaviour/SRCBehaviourPanelList.cpp
@@ -3,17 +3,19 @@
 #include "SRCBehaviourPanelList.h"
 
 #include "Behaviour/Builtin/Conditional/RCBehaviourConditional.h"
+#include "Behaviour/RCSetAssetByPathBehaviour.h"
 #include "Controller/RCController.h"
 #include "RCBehaviourModel.h"
 #include "RemoteControlPreset.h"
-#include "SlateOptMacros.h"
 #include "SRCBehaviourPanel.h"
+#include "SlateOptMacros.h"
 #include "Styling/RemoteControlStyles.h"
-#include "UI/Behaviour/Builtin/RCBehaviourIsEqualModel.h"
 #include "UI/Behaviour/Builtin/Conditional/RCBehaviourConditionalModel.h"
+#include "UI/Behaviour/Builtin/RCBehaviourIsEqualModel.h"
+#include "UI/Behaviour/Builtin/RCBehaviourSetAssetByPathModel.h"
+#include "UI/Controller/RCControllerModel.h"
 #include "UI/RemoteControlPanelStyle.h"
 #include "UI/SRemoteControlPanel.h"
-#include "UI/Controller/RCControllerModel.h"
 #include "Widgets/Views/SHeaderRow.h"
 
 #define LOCTEXT_NAMESPACE "RemoteControlPanelBehavioursList"
@@ -93,6 +95,10 @@ void SRCBehaviourPanelList::Reset()
 				{
 					BehaviourItems.Add(MakeShared<FRCBehaviourConditionalModel>(ConditionalBehaviour));
 				}
+				else if (URCSetAssetByPathBehaviour* SetAssetByPathBehaviour = Cast<URCSetAssetByPathBehaviour>(Behaviour))
+				{
+					BehaviourItems.Add(MakeShared<FRCSetAssetByPathBehaviourModel>(SetAssetByPathBehaviour));
+				}
 				else
 				{
 					BehaviourItems.Add(MakeShared<FRCBehaviourModel>(Behaviour));
