Conditional Behaviour Improvements

1. Advanced handling of Else condition for supporting multiple actions associated with each unique condition value and Else branch
2. Add All Actions now correctly stores condition value for each action allowing quicker creationg of complex scenarios.
3. Drag-and-drop Group feature implemented for quicker user workflow
4. Crash fix 1 - GC crash related fixed
5. Crash fix 2 - Handled edge case for opening of RC preset in a different level (than where the actions were created) or if a level exits with the RC preset saved, but not the bound actors.

diff --git a/Engine/Plugins/VirtualProduction/RemoteControl/Source/RemoteControlLogic/Private/Action/RCActionContainer.cpp b/Engine/Plugins/VirtualProduction/RemoteControl/Source/RemoteControlLogic/Private/Action/RCActionContainer.cpp
--- a/Engine/Plugins/VirtualProduction/RemoteControl/Source/RemoteControlLogic/Private/Action/RCActionContainer.cpp
+++ b/Engine/Plugins/VirtualProduction/RemoteControl/Source/RemoteControlLogic/Private/Action/RCActionContainer.cpp
@@ -69,6 +69,13 @@ URCPropertyAction* URCActionContainer::AddPropertyAction(const TSharedRef<const
 	NewPropertyAction->PresetWeakPtr = PresetWeakPtr;
 	NewPropertyAction->ExposedFieldId = InRemoteControlProperty->GetId();
 	NewPropertyAction->Id = FGuid::NewGuid();
+
+	if (!InRemoteControlProperty->GetBoundObjects().Num())
+	{
+		// This is possible if an exposed property was either deleted directly by the user, or if a project exits without saving the linked actor, etc
+
+		return nullptr;
+	}
 	
 	if (FRCObjectReference ObjectRef; IRemoteControlModule::Get().ResolveObjectProperty(ERCAccess::READ_ACCESS, InRemoteControlProperty->GetBoundObjects()[0], InRemoteControlProperty->FieldPathInfo.ToString(), ObjectRef))
 	{
diff --git a/Engine/Plugins/VirtualProduction/RemoteControl/Source/RemoteControlLogic/Private/Behaviour/Builtin/Conditional/RCBehaviourConditional.cpp b/Engine/Plugins/VirtualProduction/RemoteControl/Source/RemoteControlLogic/Private/Behaviour/Builtin/Conditional/RCBehaviourConditional.cpp
--- a/Engine/Plugins/VirtualProduction/RemoteControl/Source/RemoteControlLogic/Private/Behaviour/Builtin/Conditional/RCBehaviourConditional.cpp
+++ b/Engine/Plugins/VirtualProduction/RemoteControl/Source/RemoteControlLogic/Private/Behaviour/Builtin/Conditional/RCBehaviourConditional.cpp
@@ -34,17 +34,27 @@ void URCBehaviourConditional::Execute()
 
 	bool bConditionPass = false;
 
-	/* Consider the following sequence of conditions for a hypothetical Controller "Tricode" with value "Tri2"
-	*   =Tri1    <FALSE> (Action 1)       ...skip...
-	*   =Tri2   <TRUE>    (Action 2)   ...execute...
-	*   =Tri2   <TRUE>    (Action 3)   ...execute...
-	*   =Tri3   <FALSE>  (Action 4)       ...skip...
-	*   Else                                          ...skip...
+	/* Consider the following sequence of conditions for a hypothetical Controller "Tricode"
+	*  For Tricode with current value "Tri2"
+	* 
+	*   =Tri1    <FALSE>  (Action 1)           ...skip...
+	* 
+	*   =Tri2   <TRUE>    (Action 2a)      ...execute...
+	*   =Tri2   <TRUE>    (Action 2b)      ...execute...
+	* 
+	*   =Tri3   <FALSE>  (Action 3)          ...skip...
+	* 
+	*    Else                    (Action 4a)        ...skip...
+	*    Else                    (Action 4b)        ...skip...
+	*    Else                    (Action 4c)        ...skip...
 	* 
 	* For such multiple equality rows we want to know if at least one of them succeeded. 
 	* The flag bHasEqualitySuccess is used for determining whether Else should be executed 
 	*/
 	bool bHasEqualitySuccess = false;
+	bool bPreviousConditionPass = false;
+
+	ERCBehaviourConditionType PreviousConditionType = ERCBehaviourConditionType::None;
 
 	for (TObjectPtr<URCAction> Action : ActionContainer->Actions)
 	{
@@ -55,7 +65,15 @@ void URCBehaviourConditional::Execute()
 			continue;
 		}
 
-		const ERCBehaviourConditionType ConditionType = Condition->ConditionType;		
+		const ERCBehaviourConditionType ConditionType = Condition->ConditionType;
+
+		if (ConditionType != PreviousConditionType) // New block
+		{
+			if (ConditionType != ERCBehaviourConditionType::Else)
+			{
+				bHasEqualitySuccess = false; // Reset flag for starting a new block of equality checks
+			}
+		}
 
 		URCController* RCController = ControllerWeakPtr.Get();
 		if (RCController)
@@ -85,16 +103,18 @@ void URCBehaviourConditional::Execute()
 
 			case ERCBehaviourConditionType::Else:
 				// If the previous condition failed and no prior equality condition succeeded (among multiple =rows above) then execute Else!
-				bConditionPass = !bConditionPass && !bHasEqualitySuccess;
+				bConditionPass = !bPreviousConditionPass && !bHasEqualitySuccess;
 				break;
 
 			default:				
 				ensureAlwaysMsgf(false, TEXT("Unimplemented comparator!"));
 			}
 
-			if (ConditionType != ERCBehaviourConditionType::IsEqual)
+			PreviousConditionType = ConditionType;
+
+			if (ConditionType != ERCBehaviourConditionType::Else)
 			{
-				bHasEqualitySuccess = false; // Reset flag; either we reached an else clause (resolved above) or we moved from equality rows to a different comparison type
+				bPreviousConditionPass = bConditionPass; // Else should not contribute to the previous state flag to support multiple Else Actions (created via Add All Action, etc)
 			}
 		}
 
@@ -110,7 +130,7 @@ void URCBehaviourConditional::OnActionAdded(URCAction* Action, const ERCBehaviou
 {
 	FRCBehaviourCondition Condition(InConditionType, InComparand);
 
-	Conditions.Add(Action, MoveTemp(Condition));	
+	Conditions.Add(Action, MoveTemp(Condition));
 }
 
 URCAction* URCBehaviourConditional::AddAction(const TSharedRef<const FRemoteControlField> InRemoteControlField, const ERCBehaviourConditionType InConditionType, const TObjectPtr<URCVirtualPropertySelfContainer> InComparand)
diff --git a/Engine/Plugins/VirtualProduction/RemoteControl/Source/RemoteControlLogic/Public/Behaviour/Builtin/Conditional/RCBehaviourConditional.h b/Engine/Plugins/VirtualProduction/RemoteControl/Source/RemoteControlLogic/Public/Behaviour/Builtin/Conditional/RCBehaviourConditional.h
--- a/Engine/Plugins/VirtualProduction/RemoteControl/Source/RemoteControlLogic/Public/Behaviour/Builtin/Conditional/RCBehaviourConditional.h
+++ b/Engine/Plugins/VirtualProduction/RemoteControl/Source/RemoteControlLogic/Public/Behaviour/Builtin/Conditional/RCBehaviourConditional.h
@@ -21,7 +21,8 @@ enum class ERCBehaviourConditionType : uint8
 	IsLesserThan,
 	IsGreaterThanOrEqualTo,
 	IsLesserThanOrEqualTo,
-	Else
+	Else,
+	None
 };
 
 /**
@@ -40,7 +41,7 @@ struct FRCBehaviourCondition
 
 	/** The type of condition to be used for comparison*/
 	UPROPERTY()
-	ERCBehaviourConditionType ConditionType;
+	ERCBehaviourConditionType ConditionType = ERCBehaviourConditionType::None;
 
 	/** The value with which to compare */
 	UPROPERTY()
@@ -90,4 +91,8 @@ public:
 	*  Each Action is associated with a unique condition (for the Conditional Behaviour) */
 	UPROPERTY()
 	TMap<URCAction*, FRCBehaviourCondition> Conditions;
+
+	/** Virtual property used to build the Comparand - i.e. the property with which the Controller will be compared for a given condition*/
+	UPROPERTY()
+	TObjectPtr<URCVirtualPropertySelfContainer> Comparand;
 };
diff --git a/Engine/Plugins/VirtualProduction/RemoteControl/Source/RemoteControlUI/Private/UI/Action/SRCActionPanelList.cpp b/Engine/Plugins/VirtualProduction/RemoteControl/Source/RemoteControlUI/Private/UI/Action/SRCActionPanelList.cpp
--- a/Engine/Plugins/VirtualProduction/RemoteControl/Source/RemoteControlUI/Private/UI/Action/SRCActionPanelList.cpp
+++ b/Engine/Plugins/VirtualProduction/RemoteControl/Source/RemoteControlUI/Private/UI/Action/SRCActionPanelList.cpp
@@ -17,6 +17,7 @@
 #include "UI/RCUIHelpers.h"
 #include "UI/RemoteControlPanelStyle.h"
 #include "UI/SRCPanelExposedField.h"
+#include "UI/SRCPanelFieldGroup.h"
 #include "UI/SRemoteControlPanel.h"
 #include "Widgets/Views/SListView.h"
 #include "Widgets/Views/STableRow.h"
@@ -103,21 +104,58 @@ void SRCActionPanelList<ActionType>::Reset()
 	}
 }
 
+template <class ActionType>
+URCAction* SRCActionPanelList<ActionType>::AddAction(const FGuid& InRemoteControlFieldId)
+{
+	if (const URemoteControlPreset* Preset = GetPreset())
+	{
+		if (TSharedPtr<const FRemoteControlField> RemoteControlField = Preset->GetExposedEntity<FRemoteControlField>(InRemoteControlFieldId).Pin())
+		{
+			if (const TSharedPtr<SRCActionPanel> ActionPanel = GetActionPanel())
+			{
+				return ActionPanel->AddAction(RemoteControlField.ToSharedRef());
+			}
+		}
+	}
+
+	return nullptr;
+}
+
 template <class ActionType>
 FReply SRCActionPanelList<ActionType>::OnExposedFieldDrop(TSharedPtr<FDragDropOperation> DragDropOperation)
 {
-	if (DragDropOperation && DragDropOperation->IsOfType<FExposedEntityDragDrop>())
+	if (DragDropOperation)
 	{
-		if (TSharedPtr<FExposedEntityDragDrop> DragDropOp = StaticCastSharedPtr<FExposedEntityDragDrop>(DragDropOperation))
+		if (DragDropOperation->IsOfType<FExposedEntityDragDrop>())
 		{
-			if (const URemoteControlPreset* Preset = GetPreset())
+			if (TSharedPtr<FExposedEntityDragDrop> DragDropOp = StaticCastSharedPtr<FExposedEntityDragDrop>(DragDropOperation))
 			{
-				FGuid ExposedEntityId = DragDropOp->GetId();
-				if (TSharedPtr<const FRemoteControlField> RemoteControlField = Preset->GetExposedEntity<FRemoteControlField>(ExposedEntityId).Pin())
+				// Fetch the Exposed Entity
+				const FGuid ExposedEntityId = DragDropOp->GetId();
+
+				// Add Action
+				AddAction(ExposedEntityId);
+			}
+		}
+		else if (DragDropOperation->IsOfType<FFieldGroupDragDropOp>())
+		{
+			if (TSharedPtr<FFieldGroupDragDropOp> DragDropOp = StaticCastSharedPtr<FFieldGroupDragDropOp>(DragDropOperation))
+			{
+				if (URemoteControlPreset* Preset = GetPreset())
 				{
-					if (const TSharedPtr<SRCActionPanel> ActionPanel = GetActionPanel())
+					// Fetch the Group
+					const FGuid GroupId = DragDropOp->GetGroupId();
+					const FRemoteControlPresetGroup* Group = Preset->Layout.GetGroup(GroupId);
+
+					if (ensure(Group))
 					{
-						ActionPanel->AddAction(RemoteControlField.ToSharedRef());
+						const TArray<FGuid> GroupFields = Group->GetFields();
+
+						// Add Action for all fields in the Group
+						for (const FGuid RemoteControlFieldId : GroupFields)
+						{
+							AddAction(RemoteControlFieldId);
+						}
 					}
 				}
 			}
diff --git a/Engine/Plugins/VirtualProduction/RemoteControl/Source/RemoteControlUI/Private/UI/Action/SRCActionPanelList.h b/Engine/Plugins/VirtualProduction/RemoteControl/Source/RemoteControlUI/Private/UI/Action/SRCActionPanelList.h
--- a/Engine/Plugins/VirtualProduction/RemoteControl/Source/RemoteControlUI/Private/UI/Action/SRCActionPanelList.h
+++ b/Engine/Plugins/VirtualProduction/RemoteControl/Source/RemoteControlUI/Private/UI/Action/SRCActionPanelList.h
@@ -57,6 +57,9 @@ public:
 		return BehaviourItemWeakPtr.Pin();
 	}
 
+	/** Adds an Action by Remote Control Field Guid*/
+	URCAction* AddAction(const FGuid& InRemoteControlFieldId);
+
 private:
 
 	/** OnGenerateRow delegate for the Actions List View*/
@@ -102,5 +105,4 @@ private:
 
 	/** Panel Style reference. */
 	const FRCPanelStyle* RCPanelStyle;
-};
-
+};
\ No newline at end of file
diff --git a/Engine/Plugins/VirtualProduction/RemoteControl/Source/RemoteControlUI/Private/UI/Behaviour/Builtin/Conditional/RCBehaviourConditionalModel.cpp b/Engine/Plugins/VirtualProduction/RemoteControl/Source/RemoteControlUI/Private/UI/Behaviour/Builtin/Conditional/RCBehaviourConditionalModel.cpp
--- a/Engine/Plugins/VirtualProduction/RemoteControl/Source/RemoteControlUI/Private/UI/Behaviour/Builtin/Conditional/RCBehaviourConditionalModel.cpp
+++ b/Engine/Plugins/VirtualProduction/RemoteControl/Source/RemoteControlUI/Private/UI/Behaviour/Builtin/Conditional/RCBehaviourConditionalModel.cpp
@@ -28,13 +28,18 @@ FRCBehaviourConditionalModel::FRCBehaviourConditionalModel(URCBehaviourCondition
 	CreateComparandInputField();
 }
 
+URCBehaviourConditional* FRCBehaviourConditionalModel::GetConditionalBehaviour()
+{
+	return Cast<URCBehaviourConditional>(GetBehaviour());
+}
+
 URCAction* FRCBehaviourConditionalModel::AddAction(const TSharedRef<const FRemoteControlField> InRemoteControlField)
 {
 	URCAction* NewAction = nullptr;
 
-	if (URCBehaviourConditional* ConditionalBehaviour = Cast<URCBehaviourConditional>(GetBehaviour()))
+	if (URCBehaviourConditional* ConditionalBehaviour = GetConditionalBehaviour())
 	{
-		NewAction = ConditionalBehaviour->AddAction(InRemoteControlField, Condition, Comparand);
+		NewAction = ConditionalBehaviour->AddAction(InRemoteControlField, Condition, ConditionalBehaviour->Comparand);
 
 		OnActionAdded(NewAction);
 	}
@@ -44,13 +49,25 @@ URCAction* FRCBehaviourConditionalModel::AddAction(const TSharedRef<const FRemot
 
 void FRCBehaviourConditionalModel::CreateComparandInputField()
 {
-	if (URCBehaviourConditional* ConditionalBehaviour = Cast<URCBehaviourConditional>(GetBehaviour()))
+	if (URCBehaviourConditional* ConditionalBehaviour = GetConditionalBehaviour())
 	{
 		if (URCController* Controller = ConditionalBehaviour->ControllerWeakPtr.Get())
 		{
+			TObjectPtr<URCVirtualPropertySelfContainer>& Comparand = ConditionalBehaviour->Comparand;
+			URCVirtualPropertySelfContainer* PreviousComparand = Comparand;
+
 			// Virtual Property for the Comparand
 			Comparand = NewObject<URCVirtualPropertySelfContainer>(ConditionalBehaviour);
-			Comparand->DuplicateProperty(FName("Comparand"), Controller->GetProperty());
+
+			if (PreviousComparand)
+			{
+				// Users typically enter multiple actions for a single condition so make it easier for them by retaining Condition value
+				Comparand->DuplicatePropertyWithCopy(PreviousComparand);
+			}
+			else
+			{
+				Comparand->DuplicateProperty(FName("Comparand"), Controller->GetProperty());
+			}
 
 			// UI widget (via Property Generator)
 			PropertyRowGenerator->SetStructure(Comparand->CreateStructOnScope());
@@ -75,9 +92,9 @@ TSharedRef<SWidget> FRCBehaviourConditionalModel::GetBehaviourDetailsWidget()
 
 void FRCBehaviourConditionalModel::OnActionAdded(URCAction* Action)
 {
-	if (URCBehaviourConditional* ConditionalBehaviour = Cast<URCBehaviourConditional>(GetBehaviour()))
+	if (URCBehaviourConditional* ConditionalBehaviour = GetConditionalBehaviour())
 	{
-		ConditionalBehaviour->OnActionAdded(Action, Condition, Comparand);
+		ConditionalBehaviour->OnActionAdded(Action, Condition, ConditionalBehaviour->Comparand);
 
 		// Create a new Comparand for the user (the previous input field is already associated with the newly created action)
 		CreateComparandInputField();
diff --git a/Engine/Plugins/VirtualProduction/RemoteControl/Source/RemoteControlUI/Private/UI/Behaviour/Builtin/Conditional/RCBehaviourConditionalModel.h b/Engine/Plugins/VirtualProduction/RemoteControl/Source/RemoteControlUI/Private/UI/Behaviour/Builtin/Conditional/RCBehaviourConditionalModel.h
--- a/Engine/Plugins/VirtualProduction/RemoteControl/Source/RemoteControlUI/Private/UI/Behaviour/Builtin/Conditional/RCBehaviourConditionalModel.h
+++ b/Engine/Plugins/VirtualProduction/RemoteControl/Source/RemoteControlUI/Private/UI/Behaviour/Builtin/Conditional/RCBehaviourConditionalModel.h
@@ -24,6 +24,9 @@ class FRCBehaviourConditionalModel : public FRCBehaviourModel
 public:
 	FRCBehaviourConditionalModel(URCBehaviourConditional* IsEqualBehaviour);
 
+	/** Fetches the Conditional Behaviour UObject */
+	URCBehaviourConditional* GetConditionalBehaviour();
+
 	/** Add a Logic Action using a remote control field as input */
 	virtual URCAction* AddAction(const TSharedRef<const FRemoteControlField> InRemoteControlField);
 
@@ -60,9 +63,6 @@ private:
 	/** Used to create a generic value widget for comparison based on the active Controller's type*/
 	TWeakPtr<IDetailTreeNode> DetailTreeNodeWeakPtr;	
 
-	/** Virtual property used to build the Comparand - i.e. the property with which the Controller will be compared for a given condition*/
-	TObjectPtr<URCVirtualPropertySelfContainer> Comparand;	
-
 	/** Behaviour Details Widget (Conditions Widget in the case of this Behaviour) */
 	TSharedPtr<SRCBehaviourConditional> BehaviourDetailsWidget;
 };
\ No newline at end of file
