Adding additional includes in preparation for removing a global include

#rnx

diff --git a/Engine/Plugins/Experimental/ChaosEditor/Source/FractureEditor/Private/FractureEditorMode.cpp b/Engine/Plugins/Experimental/ChaosEditor/Source/FractureEditor/Private/FractureEditorMode.cpp
--- a/Engine/Plugins/Experimental/ChaosEditor/Source/FractureEditor/Private/FractureEditorMode.cpp
+++ b/Engine/Plugins/Experimental/ChaosEditor/Source/FractureEditor/Private/FractureEditorMode.cpp
@@ -23,6 +23,7 @@
 #include "GeometryCollection/GeometryCollectionClusteringUtility.h"
 #include "FractureSelectionTools.h"
 #include "EditorViewportClient.h"
+#include "Misc/ITransaction.h"
 #include "ScopedTransaction.h"
 #include "GeometryCollection/GeometryCollectionAlgo.h"
 #include "EdModeInteractiveToolsContext.h"
diff --git a/Engine/Plugins/Experimental/PCG/Source/PCGEditor/Private/PCGEditor.cpp b/Engine/Plugins/Experimental/PCG/Source/PCGEditor/Private/PCGEditor.cpp
--- a/Engine/Plugins/Experimental/PCG/Source/PCGEditor/Private/PCGEditor.cpp
+++ b/Engine/Plugins/Experimental/PCG/Source/PCGEditor/Private/PCGEditor.cpp
@@ -49,6 +49,7 @@
 #include "Preferences/UnrealEdOptions.h"
 #include "PropertyEditorModule.h"
 #include "SNodePanel.h"
+#include "Misc/ITransaction.h"
 #include "ScopedTransaction.h"
 #include "SourceCodeNavigation.h"
 #include "ToolMenus.h"
diff --git a/Engine/Plugins/VirtualProduction/RemoteControl/Source/RemoteControlProtocolWidgets/Private/ViewModels/ProtocolBindingViewModel.cpp b/Engine/Plugins/VirtualProduction/RemoteControl/Source/RemoteControlProtocolWidgets/Private/ViewModels/ProtocolBindingViewModel.cpp
--- a/Engine/Plugins/VirtualProduction/RemoteControl/Source/RemoteControlProtocolWidgets/Private/ViewModels/ProtocolBindingViewModel.cpp
+++ b/Engine/Plugins/VirtualProduction/RemoteControl/Source/RemoteControlProtocolWidgets/Private/ViewModels/ProtocolBindingViewModel.cpp
@@ -10,6 +10,7 @@
 #include "RCTypeTraits.h"
 #include "RemoteControlPreset.h"
 #include "RemoteControlProtocolWidgetsModule.h"
+#include "Misc/ITransaction.h"
 #include "ScopedTransaction.h"
 #include "Algo/AnyOf.h"
 
diff --git a/Engine/Plugins/VirtualProduction/RemoteControl/Source/RemoteControlProtocolWidgets/Private/ViewModels/ProtocolEntityViewModel.cpp b/Engine/Plugins/VirtualProduction/RemoteControl/Source/RemoteControlProtocolWidgets/Private/ViewModels/ProtocolEntityViewModel.cpp
--- a/Engine/Plugins/VirtualProduction/RemoteControl/Source/RemoteControlProtocolWidgets/Private/ViewModels/ProtocolEntityViewModel.cpp
+++ b/Engine/Plugins/VirtualProduction/RemoteControl/Source/RemoteControlProtocolWidgets/Private/ViewModels/ProtocolEntityViewModel.cpp
@@ -9,6 +9,7 @@
 #include "ProtocolCommandChange.h"
 #include "RemoteControlPreset.h"
 #include "RemoteControlProtocolWidgetsSettings.h"
+#include "Misc/ITransaction.h"
 #include "ScopedTransaction.h"
 #include "Algo/AllOf.h"
 #include "Algo/AnyOf.h"
diff --git a/Engine/Source/Editor/SubobjectDataInterface/Private/SubobjectDataSubsystem.cpp b/Engine/Source/Editor/SubobjectDataInterface/Private/SubobjectDataSubsystem.cpp
--- a/Engine/Source/Editor/SubobjectDataInterface/Private/SubobjectDataSubsystem.cpp
+++ b/Engine/Source/Editor/SubobjectDataInterface/Private/SubobjectDataSubsystem.cpp
@@ -23,6 +23,7 @@
 #include "Kismet2/BlueprintEditorUtils.h"
 #include "ClassViewerFilter.h"
 #include "AssetRegistry/AssetRegistryModule.h"
+#include "Misc/ITransaction.h"
 #include "ScopedTransaction.h"
 
 #include "GameProjectGenerationModule.h"	// Adding new component classes
