~~ Edit Condition Feature ~~

Provides fully generic support for displaying and editing any Virtual Property (Self Container) through a new widget type: SRCVirtualPropertyWidget.

Callers can implement the read-only version via an OnGenerateWidget call while the editable version (with Context menu and Double-click) are provided automatically by the widget.

diff --git a/Engine/Plugins/VirtualProduction/RemoteControl/Source/RemoteControlUI/Private/UI/Action/Conditional/RCActionConditionalModel.cpp b/Engine/Plugins/VirtualProduction/RemoteControl/Source/RemoteControlUI/Private/UI/Action/Conditional/RCActionConditionalModel.cpp
--- a/Engine/Plugins/VirtualProduction/RemoteControl/Source/RemoteControlUI/Private/UI/Action/Conditional/RCActionConditionalModel.cpp
+++ b/Engine/Plugins/VirtualProduction/RemoteControl/Source/RemoteControlUI/Private/UI/Action/Conditional/RCActionConditionalModel.cpp
@@ -11,6 +11,7 @@
 #include "RemoteControlField.h"
 #include "RemoteControlPreset.h"
 #include "Styling/RemoteControlStyles.h"
+#include "UI/Action/SRCVirtualPropertyWidget.h"
 #include "UI/Behaviour/Builtin/Conditional/RCBehaviourConditionalModel.h"
 #include "UI/RCUIHelpers.h"
 #include "UI/RemoteControlPanelStyle.h"
@@ -127,8 +128,33 @@ TSharedPtr<FRCActionConditionalModel> FRCActionConditionalModel::GetModelByActio
 	return nullptr;
 }
 
-TSharedRef<SWidget> FRCActionConditionalModel::GetConditionWidget() const
+TSharedRef<SWidget> FRCActionConditionalModel::GetConditionWidget()
 {
+	if (TSharedPtr<FRCBehaviourConditionalModel> BehaviourItem = StaticCastSharedPtr<FRCBehaviourConditionalModel>(BehaviourItemWeakPtr.Pin()))
+	{
+		if (URCBehaviourConditional* Behaviour = Cast<URCBehaviourConditional>(BehaviourItem->GetBehaviour()))
+		{
+			if (FRCBehaviourCondition* Condition = Behaviour->Conditions.Find(GetAction()))
+			{
+				if(ensure(Condition->Comparand))
+				{
+					return SAssignNew(EditableVirtualPropertyWidget, SRCVirtualPropertyWidget, Condition->Comparand)
+						.OnGenerateWidget(this, &FRCActionConditionalModel::OnGenerateConditionWidget);
+				}
+			}
+		}
+	}
+
+	return SNullWidget::NullWidget;
+}
+
+TSharedRef<SWidget> FRCActionConditionalModel::OnGenerateConditionWidget(URCVirtualPropertySelfContainer* InComparand)
+{
+	if (!ensure(InComparand))
+	{
+		return SNullWidget::NullWidget;
+	}
+
 	if (TSharedPtr<FRCBehaviourConditionalModel> BehaviourItem = StaticCastSharedPtr<FRCBehaviourConditionalModel>(BehaviourItemWeakPtr.Pin()))
 	{
 		if (URCBehaviourConditional* Behaviour = Cast<URCBehaviourConditional>(BehaviourItem->GetBehaviour()))
@@ -150,15 +176,10 @@ TSharedRef<SWidget> FRCActionConditionalModel::GetConditionWidget() const
 					ConditionComparandText = FText::Format(FText::FromName("{0} {1}"), ConditionText, ComparandValueText);
 				}
 
-				return SNew(SBox)
-					.Padding(FMargin(6.f))
-					[
-						SNew(STextBlock).Text(ConditionComparandText)
-					];
+				return SNew(STextBlock).Text(ConditionComparandText);
 			}
 		}
 	}
-	
 
 	return SNullWidget::NullWidget;
 }
diff --git a/Engine/Plugins/VirtualProduction/RemoteControl/Source/RemoteControlUI/Private/UI/Action/Conditional/RCActionConditionalModel.h b/Engine/Plugins/VirtualProduction/RemoteControl/Source/RemoteControlUI/Private/UI/Action/Conditional/RCActionConditionalModel.h
--- a/Engine/Plugins/VirtualProduction/RemoteControl/Source/RemoteControlUI/Private/UI/Action/Conditional/RCActionConditionalModel.h
+++ b/Engine/Plugins/VirtualProduction/RemoteControl/Source/RemoteControlUI/Private/UI/Action/Conditional/RCActionConditionalModel.h
@@ -6,13 +6,17 @@
 #include "Action/RCPropertyAction.h"
 #include "UI/Action/RCActionModel.h"
 
+class SActionItemListRow;
+class SBox;
+class SRCVirtualPropertyWidget;
+
 /* 
 * ~ FRCActionConditionalModel ~
 *
 * UI model for representing an Action of the Conditional Behaviour.
 * Contains a Condition widget for representing the comparand/comparator and target property/value widgets
 * 
-*                              <Condition>                   <Description>         <Value>
+*  <Actions Table>       <Condition>                <Description>         <Value>
 * 
 *  Color coding | Comparator+Comparand | Target Property | Value To Set
 */
@@ -24,7 +28,7 @@ public:
 	/** Widget representing the Condition (Comparator|Comparand) used
 	* to trigger the corresponding action
 	*/
-	virtual TSharedRef<SWidget> GetConditionWidget() const;
+	TSharedRef<SWidget> GetConditionWidget();
 
 	/** OnGenerateRow delegate for the Actions List View*/
 	TSharedRef<ITableRow> OnGenerateWidgetForList(TSharedPtr<FRCActionConditionalModel> InItem, const TSharedRef<STableViewBase>& OwnerTable);
@@ -34,6 +38,9 @@ public:
 
 	/** Chooses the appropriate Action model for the current class and field type*/
 	static TSharedPtr<FRCActionConditionalModel> GetModelByActionType(URCAction* InAction, const TSharedPtr<class FRCBehaviourModel> InBehaviourItem, const TSharedPtr<SRemoteControlPanel> InRemoteControlPanel);
+
+private:
+	TSharedRef<SWidget> OnGenerateConditionWidget(class URCVirtualPropertySelfContainer* InComparand);
 };
 
 /*
diff --git a/Engine/Plugins/VirtualProduction/RemoteControl/Source/RemoteControlUI/Private/UI/Action/RCActionModel.cpp b/Engine/Plugins/VirtualProduction/RemoteControl/Source/RemoteControlUI/Private/UI/Action/RCActionModel.cpp
--- a/Engine/Plugins/VirtualProduction/RemoteControl/Source/RemoteControlUI/Private/UI/Action/RCActionModel.cpp
+++ b/Engine/Plugins/VirtualProduction/RemoteControl/Source/RemoteControlUI/Private/UI/Action/RCActionModel.cpp
@@ -2,18 +2,23 @@
 
 #include "RCActionModel.h"
 
+#include "Commands/RemoteControlCommands.h"
 #include "Controller/RCController.h"
 #include "IDetailTreeNode.h"
+#include "Interfaces/IMainFrameModule.h"
 #include "IPropertyRowGenerator.h"
 #include "RCVirtualProperty.h"
 #include "RCVirtualPropertyContainer.h"
 #include "RemoteControlField.h"
 #include "RemoteControlPreset.h"
 #include "Styling/RemoteControlStyles.h"
+#include "UI/Action/SRCActionPanel.h"
+#include "UI/Action/SRCVirtualPropertyWidget.h"
 #include "UI/RCUIHelpers.h"
 #include "UI/RemoteControlPanelStyle.h"
 #include "UI/SRCPanelDragHandle.h"
 #include "UI/SRCPanelExposedEntity.h"
+#include "UI/SRemoteControlPanel.h"
 #include "Widgets/Layout/SBox.h"
 #include "Widgets/Views/SHeaderRow.h"
 #include "Widgets/Views/SListView.h"
@@ -209,58 +214,61 @@ const FName& FRCPropertyActionType::GetPropertyName() const
 }
 
 TSharedRef<SWidget> FRCPropertyActionType::GetPropertyWidget() const
+{
+	return UE::RCUIHelpers::GetGenericFieldWidget(DetailTreeNodeWeakPtr.Pin());
+}
+
+FLinearColor FRCPropertyActionType::GetPropertyTypeColor() const
 {
 	if (!ensure(PropertyActionWeakPtr.IsValid()))
 	{
-		return SNullWidget::NullWidget;
+		return FLinearColor::White;
 	}
 
-	const URCPropertyAction* PropertyAction = PropertyActionWeakPtr.Get();	
+	const URCPropertyAction* PropertyAction = PropertyActionWeakPtr.Get();
 
-	if (DetailTreeNodeWeakPtr.IsValid())
-	{
-		const FNodeWidgets NodeWidgets = DetailTreeNodeWeakPtr.Pin()->CreateNodeWidgets();
+	const FLinearColor TypeColor = UE::RCUIHelpers::GetFieldClassTypeColor(PropertyAction->PropertySelfContainer->GetProperty());
 
-		const TSharedRef<SHorizontalBox> FieldWidget = SNew(SHorizontalBox);
+	return TypeColor;
+}
+
+TSharedPtr<SWidget> FRCActionModel::GetContextMenuWidget()
+{
+	IMainFrameModule& MainFrame = FModuleManager::Get().LoadModuleChecked<IMainFrameModule>("MainFrame");
 
-		if (NodeWidgets.ValueWidget)
+	FMenuBuilder MenuBuilder(true, MainFrame.GetMainFrameCommandBindings());
+
+	MenuBuilder.BeginSection("Common");
+
+	if (TSharedPtr<SRemoteControlPanel> RemoteControlPanel = PanelWeakPtr.Pin())
+	{
+		// 1. Edit (if available)
+		if (EditableVirtualPropertyWidget)
 		{
-			FieldWidget->AddSlot()
-				.Padding(FMargin(3.0f, 2.0f))
-				.HAlign(HAlign_Right)
-				.AutoWidth()
-				[
-					NodeWidgets.ValueWidget.ToSharedRef()
-				];
+			EditableVirtualPropertyWidget->AddEditContextMenuOption(MenuBuilder);
 		}
-		else if (NodeWidgets.WholeRowWidget)
+
+		// 2. Delete
+		if (TSharedPtr<SRCActionPanel> ActionPanelList = RemoteControlPanel->GetLogicActionPanel())
 		{
-			FieldWidget->AddSlot()
-				.Padding(FMargin(3.0f, 2.0f))
-				.AutoWidth()
-				[
-					NodeWidgets.WholeRowWidget.ToSharedRef()
-				];
+			FUIAction Action(FExecuteAction::CreateSP(ActionPanelList.Get(), &SRCActionPanel::DeleteSelectedPanelItem));
+			MenuBuilder.AddMenuEntry(LOCTEXT("ContextMenuEdit", "Delete"),
+				LOCTEXT("ContextMenuEditTooltip", "Delete the selected action"),
+				FSlateIcon(), Action);
 		}
-
-		return FieldWidget;
 	}
-	
-	return SNullWidget::NullWidget;
+
+	MenuBuilder.EndSection();
+
+	return MenuBuilder.MakeWidget();
 }
 
-FLinearColor FRCPropertyActionType::GetPropertyTypeColor() const
+void FRCActionModel::OnSelectionExit()
 {
-	if (!ensure(PropertyActionWeakPtr.IsValid()))
+	if (EditableVirtualPropertyWidget)
 	{
-		return FLinearColor::White;
+		EditableVirtualPropertyWidget->ExitEditMode();
 	}
-
-	const URCPropertyAction* PropertyAction = PropertyActionWeakPtr.Get();
-
-	const FLinearColor TypeColor = UE::RCUIHelpers::GetFieldClassTypeColor(PropertyAction->PropertySelfContainer->GetProperty());
-
-	return TypeColor;
 }
 
 #undef LOCTEXT_NAMESPACE
\ No newline at end of file
diff --git a/Engine/Plugins/VirtualProduction/RemoteControl/Source/RemoteControlUI/Private/UI/Action/RCActionModel.h b/Engine/Plugins/VirtualProduction/RemoteControl/Source/RemoteControlUI/Private/UI/Action/RCActionModel.h
--- a/Engine/Plugins/VirtualProduction/RemoteControl/Source/RemoteControlUI/Private/UI/Action/RCActionModel.h
+++ b/Engine/Plugins/VirtualProduction/RemoteControl/Source/RemoteControlUI/Private/UI/Action/RCActionModel.h
@@ -15,6 +15,7 @@ class ITableRow;
 class IPropertyRowGenerator;
 class SHeaderRow;
 class SHorizontalBox;
+class SRCVirtualPropertyWidget;
 class STableViewBase;
 class SWidget;
 class URCAction;
@@ -57,6 +58,12 @@ public:
 	/** Chooses the appropriate Action model for the current class and field type*/
 	static TSharedPtr<FRCActionModel> GetModelByActionType(URCAction* InAction, const TSharedPtr<class FRCBehaviourModel> InBehaviourItem, const TSharedPtr<SRemoteControlPanel> InRemoteControlPanel);
 
+	/** Context menu for Actions panel list */
+	virtual TSharedPtr<SWidget> GetContextMenuWidget();
+
+	/** Invoked when this Action item no longer has focus*/
+	virtual void OnSelectionExit();
+
 protected:
 	/** The Action (data model) associated with us*/
 	TWeakObjectPtr<URCAction> ActionWeakPtr;
@@ -66,6 +73,10 @@ protected:
 
 	/** The Behaviour (UI model) associated with us*/
 	TWeakPtr<FRCBehaviourModel> BehaviourItemWeakPtr;
+
+	/** Optional: An editable virtual property widget associated with this row
+	* Child action models can assign this variable and receive the Edit context menu option automatically*/
+	TSharedPtr<SRCVirtualPropertyWidget> EditableVirtualPropertyWidget;
 };
 
 /*
@@ -96,7 +107,7 @@ protected:
 	TSharedPtr<IPropertyRowGenerator> PropertyRowGenerator;
 
 	/** Used to create a generic Value Widget for the property row widget*/
-	TWeakPtr<IDetailTreeNode> DetailTreeNodeWeakPtr;	
+	TWeakPtr<IDetailTreeNode> DetailTreeNodeWeakPtr;
 };
 
 /*
diff --git a/Engine/Plugins/VirtualProduction/RemoteControl/Source/RemoteControlUI/Private/UI/Action/SRCActionPanelList.cpp b/Engine/Plugins/VirtualProduction/RemoteControl/Source/RemoteControlUI/Private/UI/Action/SRCActionPanelList.cpp
--- a/Engine/Plugins/VirtualProduction/RemoteControl/Source/RemoteControlUI/Private/UI/Action/SRCActionPanelList.cpp
+++ b/Engine/Plugins/VirtualProduction/RemoteControl/Source/RemoteControlUI/Private/UI/Action/SRCActionPanelList.cpp
@@ -38,7 +38,9 @@ void SRCActionPanelList<ActionType>::Construct(const FArguments& InArgs, const T
 	ListView = SNew(SListView<TSharedPtr<ActionType>>)
 		.ListItemsSource(&ActionItems)
 		.OnGenerateRow(this, &SRCActionPanelList::OnGenerateWidgetForList)
+		.OnSelectionChanged(this, &SRCActionPanelList::OnSelectionChanged)
 		.ListViewStyle(&RCPanelStyle->TableViewStyle)
+		.OnContextMenuOpening(this, &SRCActionPanelList::GetContextMenuWidget)
 		.HeaderRow(ActionType::GetHeaderRow());
 
 	ChildSlot
@@ -240,7 +242,29 @@ void SRCActionPanelList<ActionType>::DeleteSelectedPanelItem()
 	DeleteItemFromLogicPanel<ActionType>(ActionItems, ListView->GetSelectedItems());
 }
 
+template <typename ActionType>
+TSharedPtr<SWidget> SRCActionPanelList<ActionType>::GetContextMenuWidget()
+{
+	if (SelectedActionItem)
+	{
+		return SelectedActionItem->GetContextMenuWidget();
+	}
+
+	return SNullWidget::NullWidget;
+}
+
+template <typename ActionType>
+void SRCActionPanelList<ActionType>::OnSelectionChanged(TSharedPtr<ActionType> InItem, ESelectInfo::Type)
+{
+	if (SelectedActionItem)
+	{
+		SelectedActionItem->OnSelectionExit();
+	}
+
+	SelectedActionItem = InItem;
+}
+
 #undef LOCTEXT_NAMESPACE
 
 template class SRCActionPanelList<FRCActionModel>;
-template class SRCActionPanelList<FRCActionConditionalModel>;
+template class SRCActionPanelList<FRCActionConditionalModel>;
\ No newline at end of file
diff --git a/Engine/Plugins/VirtualProduction/RemoteControl/Source/RemoteControlUI/Private/UI/Action/SRCActionPanelList.h b/Engine/Plugins/VirtualProduction/RemoteControl/Source/RemoteControlUI/Private/UI/Action/SRCActionPanelList.h
--- a/Engine/Plugins/VirtualProduction/RemoteControl/Source/RemoteControlUI/Private/UI/Action/SRCActionPanelList.h
+++ b/Engine/Plugins/VirtualProduction/RemoteControl/Source/RemoteControlUI/Private/UI/Action/SRCActionPanelList.h
@@ -89,8 +89,17 @@ private:
 	/** Whether drag and drop is possible from the current exposed property to the Actions table */
 	bool OnAllowDrop(TSharedPtr<FDragDropOperation> DragDropOperation);
 
+	/** Context menu for Actions panel list */
+	TSharedPtr<SWidget> GetContextMenuWidget();
+
+	/** OnSelectionChanged delegate for Actions List View */
+	void OnSelectionChanged(TSharedPtr<ActionType> InItem, ESelectInfo::Type);
+
 private:
 
+	/** The currently selected Action item*/
+	TSharedPtr<ActionType> SelectedActionItem;
+
 	/** The parent Action Panel widget*/
 	TWeakPtr<SRCActionPanel> ActionPanelWeakPtr;
 
diff --git a/Engine/Plugins/VirtualProduction/RemoteControl/Source/RemoteControlUI/Private/UI/Action/SRCVirtualPropertyWidget.cpp b/Engine/Plugins/VirtualProduction/RemoteControl/Source/RemoteControlUI/Private/UI/Action/SRCVirtualPropertyWidget.cpp
new file mode 100644
--- /dev/null
+++ b/Engine/Plugins/VirtualProduction/RemoteControl/Source/RemoteControlUI/Private/UI/Action/SRCVirtualPropertyWidget.cpp
@@ -0,0 +1,118 @@
+// Copyright Epic Games, Inc. All Rights Reserved.
+
+#include "SRCVirtualPropertyWidget.h"
+
+#include "Framework/MultiBox/MultiBoxBuilder.h"
+#include "Interfaces/IMainFrameModule.h"
+#include "PropertyHandle.h"
+#include "RCVirtualProperty.h"
+#include "SlateOptMacros.h"
+#include "UI/RCUIHelpers.h"
+#include "Widgets/Layout/SBox.h"
+#include "Widgets/Layout/SBorder.h"
+
+#define LOCTEXT_NAMESPACE "SRCVirtualPropertyWidget"
+
+void SRCVirtualPropertyWidget::Construct(const FArguments& InArgs, URCVirtualPropertySelfContainer* InVirtualProperty)
+{
+	if (!ensure(InVirtualProperty))
+	{
+		return;
+	}
+
+	this->OnGenerateWidget = InArgs._OnGenerateWidget;
+
+	VirtualPropertyWeakPtr = InVirtualProperty;
+
+	TSharedRef<SWidget> VirtualPropertyDisplayWidget = SNullWidget::NullWidget;
+
+	if (ensureMsgf(OnGenerateWidget.IsBound(), TEXT("OnGenerateWidget must be implemented for read-only display of the virtual property")))
+	{
+		VirtualPropertyDisplayWidget = OnGenerateWidget.Execute(InVirtualProperty);
+	}
+
+	ChildSlot
+		[
+			SNew(SBorder)
+			.BorderImage(FAppStyle::GetBrush("NoBorder"))
+			.OnMouseDoubleClick(this, &SRCVirtualPropertyWidget::OnMouseDoubleClick)
+			[
+				SAssignNew(VirtualPropertyWidgetBox, SBox)
+				.Padding(FMargin(6.f))
+				[
+					VirtualPropertyDisplayWidget
+				]
+			]
+		];
+}
+
+FReply SRCVirtualPropertyWidget::OnMouseDoubleClick(const FGeometry& InMyGeometry, const FPointerEvent& InMouseEvent)
+{
+	EnterEditMode();
+
+	return FReply::Handled();
+}
+
+void SRCVirtualPropertyWidget::EnterEditMode()
+{
+	if (bIsEditMode)
+	{
+		return;
+	}
+
+	if (URCVirtualPropertySelfContainer* VirtualProperty = VirtualPropertyWeakPtr.Get())
+	{
+		DetailTreeNodeWeakPtr = UE::RCUIHelpers::GetDetailTreeNodeForVirtualProperty(VirtualProperty, PropertyRowGenerator);
+
+		if (ensure(DetailTreeNodeWeakPtr.IsValid()))
+		{
+			TSharedPtr<IPropertyHandle> PropertyHandle;
+
+			TSharedRef<SWidget> Widget = UE::RCUIHelpers::GetGenericFieldWidget(DetailTreeNodeWeakPtr.Pin(), &PropertyHandle, true);
+
+			// Updates the parent box with the editable field widget for the virtual property which facilitates user input
+			VirtualPropertyWidgetBox->SetContent(Widget);
+
+			if (ensure(PropertyHandle.IsValid()))
+			{
+				PropertyHandle->SetOnPropertyValueChanged(FSimpleDelegate::CreateSP(this, &SRCVirtualPropertyWidget::OnPropertyValueChanged));
+			}
+
+			bIsEditMode = true;
+		}
+	}
+}
+
+void SRCVirtualPropertyWidget::ExitEditMode()
+{
+	if (!bIsEditMode)
+	{
+		return;
+	}
+
+	if (ensure(OnGenerateWidget.IsBound()))
+	{
+		TSharedRef<SWidget> VirtualPropertyDisplayWidget = OnGenerateWidget.Execute(VirtualPropertyWeakPtr.Get());
+
+		// Restore the readonly view of the virtual property widget
+		VirtualPropertyWidgetBox->SetContent(VirtualPropertyDisplayWidget);
+	}
+
+	bIsEditMode = false;
+}
+
+void SRCVirtualPropertyWidget::OnPropertyValueChanged()
+{
+	ExitEditMode();
+}
+
+void SRCVirtualPropertyWidget::AddEditContextMenuOption(FMenuBuilder& MenuBuilder)
+{
+	FUIAction Action(FExecuteAction::CreateSP(this, &SRCVirtualPropertyWidget::EnterEditMode));
+
+	MenuBuilder.AddMenuEntry(LOCTEXT("ContextMenuEdit", "Edit"),
+		LOCTEXT("ContextMenuEditTooltip", "Edit the selected value"),
+		FSlateIcon(), Action);
+}
+
+#undef LOCTEXT_NAMESPACE
\ No newline at end of file
diff --git a/Engine/Plugins/VirtualProduction/RemoteControl/Source/RemoteControlUI/Private/UI/Action/SRCVirtualPropertyWidget.h b/Engine/Plugins/VirtualProduction/RemoteControl/Source/RemoteControlUI/Private/UI/Action/SRCVirtualPropertyWidget.h
new file mode 100644
--- /dev/null
+++ b/Engine/Plugins/VirtualProduction/RemoteControl/Source/RemoteControlUI/Private/UI/Action/SRCVirtualPropertyWidget.h
@@ -0,0 +1,77 @@
+// Copyright Epic Games, Inc. All Rights Reserved.
+
+#pragma once
+
+#include "CoreMinimal.h"
+#include "Framework/SlateDelegates.h"
+#include "Widgets/SCompoundWidget.h"
+
+class FMenuBuilder;
+class IDetailTreeNode;
+class IPropertyRowGenerator;
+class SBox;
+class URCVirtualPropertySelfContainer;
+
+/*
+* ~ SRCVirtualPropertyWidget ~
+*
+* Multi purpose widget for Virtual Properties covering both read-only and Edit usage scenarios
+* 
+* 1) In Read-Only mode, the OnGenerateWidget delegate is used to build a view widget as required by the caller.
+* 2) In Edit mode, a generic field widget is generated from the Virtual property
+* 
+* Supported edit actions: 1) Double-click 2) Right-click -> Contet menu -> Edit
+* Context menu support is automatically provided when this widget is set to the EditableVirtualPropertyWidget member of FRCActionModel.
+*/
+class REMOTECONTROLUI_API SRCVirtualPropertyWidget : public SCompoundWidget
+{
+public:
+	using FOnGenerateWidget = typename TSlateDelegates<URCVirtualPropertySelfContainer*>::FOnGenerateWidget;
+
+	SLATE_BEGIN_ARGS(SRCVirtualPropertyWidget)
+		: _OnGenerateWidget()
+		{
+		}
+
+	SLATE_EVENT(FOnGenerateWidget, OnGenerateWidget)
+
+	SLATE_END_ARGS()
+
+	/** Constructs this widget with InArgs */
+	void Construct(const FArguments& InArgs, URCVirtualPropertySelfContainer* InVirtualProperty);
+
+	/** Activates Edit Mode, generates a generic field widget for the virtual property to facilitate user entry */
+	void EnterEditMode();
+
+	/** Exits Edit Mode, restores the original read-only widget view as defined by the caller (via OnGenerateWidget) */
+	void ExitEditMode();
+
+	/** Automatically provides Right-click Edit support (invoked from FRCActionModel's GetContextMenuWidget function*/
+	void AddEditContextMenuOption(FMenuBuilder& MenuBuilder);
+
+protected:
+	/** Delegate to be invoked when the list needs to generate a new widget from a data item. */
+	FOnGenerateWidget OnGenerateWidget;
+
+private:
+	/** Double-click event for the border wrapping the Virtual Property widget. Provides Edit mode entry*/
+	FReply OnMouseDoubleClick(const FGeometry& InMyGeometry, const FPointerEvent& InMouseEvent);
+
+	/** Property value change listener for the generic field widget. Used to exit edit mode. 
+	* This is used in lieu of regular commit events as we don't have access to the generic field widget's internals */
+	void OnPropertyValueChanged();
+
+	/** The Virtual Property which this widget represents */
+	TWeakObjectPtr<URCVirtualPropertySelfContainer> VirtualPropertyWeakPtr;
+
+	/** Parent box encompassing read-only and edit modes.*/
+	TSharedPtr<SBox> VirtualPropertyWidgetBox;
+
+	/** Whether the virtual property is currently in edit mode */
+	bool bIsEditMode = false;
+
+	/** The property generator used to build a generic value widget*/
+	TSharedPtr<IPropertyRowGenerator> PropertyRowGenerator;
+
+	/** Detail tree node used to create a generic value widget*/
+	TWeakPtr<IDetailTreeNode> DetailTreeNodeWeakPtr;};
diff --git a/Engine/Plugins/VirtualProduction/RemoteControl/Source/RemoteControlUI/Private/UI/RCUIHelpers.cpp b/Engine/Plugins/VirtualProduction/RemoteControl/Source/RemoteControlUI/Private/UI/RCUIHelpers.cpp
--- a/Engine/Plugins/VirtualProduction/RemoteControl/Source/RemoteControlUI/Private/UI/RCUIHelpers.cpp
+++ b/Engine/Plugins/VirtualProduction/RemoteControl/Source/RemoteControlUI/Private/UI/RCUIHelpers.cpp
@@ -2,13 +2,19 @@
 
 #include "RCUIHelpers.h"
 
-#include "RCVirtualProperty.h"
-
+#include "Controller/RCController.h"
+#include "Editor.h"
 #include "EdGraphSchema_K2.h"
 #include "EdGraph/EdGraphPin.h"
-
+#include "Framework/Application/SlateApplication.h"
 #include "GraphEditorSettings.h"
-#include "Controller/RCController.h"
+#include "IDetailTreeNode.h"
+#include "IPropertyRowGenerator.h"
+#include "Modules/ModuleManager.h"
+#include "PropertyEditorModule.h"
+#include "PropertyHandle.h"
+#include "RCVirtualProperty.h"
+#include "TimerManager.h"
 
 FLinearColor UE::RCUIHelpers::GetFieldClassTypeColor(const FProperty* InProperty)
 {
@@ -61,4 +67,111 @@ FName UE::RCUIHelpers::GetFieldClassDisplayName(const FProperty* InProperty)
 	}
 
 	return FieldClassDisplayName;
-}
\ No newline at end of file
+}
+
+TSharedPtr<IDetailTreeNode> UE::RCUIHelpers::GetDetailTreeNodeForVirtualProperty(TObjectPtr<URCVirtualPropertySelfContainer> InVirtualPropertySelfContainer, 
+	TSharedPtr<IPropertyRowGenerator>& OutPropertyRowGenerator)
+{
+	TSharedPtr<IDetailTreeNode> DetailTreeNode;
+
+	FPropertyRowGeneratorArgs Args;
+	Args.bShouldShowHiddenProperties = true;
+	OutPropertyRowGenerator = FModuleManager::GetModuleChecked<FPropertyEditorModule>("PropertyEditor").CreatePropertyRowGenerator(Args);
+
+	OutPropertyRowGenerator->SetStructure(InVirtualPropertySelfContainer->CreateStructOnScope());
+
+	for (const TSharedRef<IDetailTreeNode>& CategoryNode : OutPropertyRowGenerator->GetRootTreeNodes())
+	{
+		TArray<TSharedRef<IDetailTreeNode>> Children;
+		CategoryNode->GetChildren(Children);
+
+		// Use the first child as the detail tree node
+		if (Children.Num())
+		{
+			DetailTreeNode = Children[0];
+			break;
+		}
+	}
+
+	return DetailTreeNode;
+}
+
+static void SetFocusToWidgetNextTick(const TSharedRef< SWidget> InWidget)
+{
+	GEditor->GetTimerManager()->SetTimerForNextTick(
+		FTimerDelegate::CreateLambda([InWidget]()
+			{
+				UE::RCUIHelpers::FindFocusableWidgetAndSetKeyboardFocus(InWidget);
+			})
+	);
+}
+
+TSharedRef<SWidget> UE::RCUIHelpers::GetGenericFieldWidget(const TSharedPtr<IDetailTreeNode> DetailTreeNode, 
+	TSharedPtr<IPropertyHandle>* OutPropertyHandle /*=nullptr*/, bool bFocusInputWidget/*= false*/)
+{
+	if (!DetailTreeNode.IsValid())
+	{
+		return SNullWidget::NullWidget;
+	}
+
+	if (OutPropertyHandle)
+	{
+		*OutPropertyHandle = DetailTreeNode->CreatePropertyHandle();
+	}
+
+	const FNodeWidgets NodeWidgets = DetailTreeNode->CreateNodeWidgets();
+
+	const TSharedRef<SHorizontalBox> FieldWidget = SNew(SHorizontalBox);
+
+	if (NodeWidgets.ValueWidget)
+	{
+		FieldWidget->AddSlot()
+			.Padding(FMargin(3.0f, 2.0f))
+			.VAlign(VAlign_Center)
+			.AutoWidth()
+			[
+				NodeWidgets.ValueWidget.ToSharedRef()
+			];
+
+		if (bFocusInputWidget)
+			SetFocusToWidgetNextTick(NodeWidgets.ValueWidget.ToSharedRef());
+
+	}
+	else if (NodeWidgets.WholeRowWidget)
+	{
+		FieldWidget->AddSlot()
+			.Padding(FMargin(3.0f, 2.0f))
+			.VAlign(VAlign_Center)
+			.AutoWidth()
+			[
+				NodeWidgets.WholeRowWidget.ToSharedRef()
+			];
+
+		if (bFocusInputWidget)
+			SetFocusToWidgetNextTick(NodeWidgets.WholeRowWidget.ToSharedRef());
+	}
+
+	return FieldWidget;
+}
+
+bool UE::RCUIHelpers::FindFocusableWidgetAndSetKeyboardFocus(const TSharedRef< SWidget> InWidget)
+{
+	if (InWidget->SupportsKeyboardFocus())
+	{
+		FSlateApplication::Get().SetKeyboardFocus(InWidget, EFocusCause::Navigation);
+
+		return true; // Success!
+	}
+
+	// Check child widgets...
+	if (FChildren* Children = InWidget->GetChildren())
+	{
+		if(Children->Num() > 0)
+		{
+			// Note: We're only interested in the first child for all current usecases.
+			return FindFocusableWidgetAndSetKeyboardFocus(Children->GetChildAt(0)); // Recurse!
+		}
+	}
+
+	return false; // No focusable widget found in the entire widget hierarchy
+}
diff --git a/Engine/Plugins/VirtualProduction/RemoteControl/Source/RemoteControlUI/Private/UI/RCUIHelpers.h b/Engine/Plugins/VirtualProduction/RemoteControl/Source/RemoteControlUI/Private/UI/RCUIHelpers.h
--- a/Engine/Plugins/VirtualProduction/RemoteControl/Source/RemoteControlUI/Private/UI/RCUIHelpers.h
+++ b/Engine/Plugins/VirtualProduction/RemoteControl/Source/RemoteControlUI/Private/UI/RCUIHelpers.h
@@ -3,8 +3,14 @@
 #pragma once
 
 #include "CoreMinimal.h"
+#include "UObject/ObjectPtr.h"
 
+class IDetailTreeNode;
+class IPropertyHandle;
+class IPropertyRowGenerator;
+class SWidget;
 class URCController;
+class URCVirtualPropertySelfContainer;
 
 namespace UE::RCUIHelpers
 {
@@ -18,4 +24,30 @@ namespace UE::RCUIHelpers
 	 * Fetches the display name associated with a given Unreal Type (FProperty)
 	 */
 	FName GetFieldClassDisplayName(const FProperty* InProperty);
+
+	/**
+	* GetDetailTreeNodeForVirtualProperty
+	*
+	* Given a Virtual Property this function generates the corresponding Detail Tree Node.
+	* The Property Row Generator used to create this is also set as an output parameter.
+	*/
+	TSharedPtr<IDetailTreeNode> GetDetailTreeNodeForVirtualProperty(TObjectPtr<URCVirtualPropertySelfContainer> InVirtualPropertySelfContainer, TSharedPtr<IPropertyRowGenerator>& OutPropertyRowGenerator);
+
+	/** 
+	* GetGenericFieldWidget
+	* 
+	* Constructs a widget representing a generic property using an input detail tree node
+    * Optionally returns the corresponding property handle as an output parameter.
+	*/
+	TSharedRef<SWidget> GetGenericFieldWidget(const TSharedPtr<IDetailTreeNode> DetailTreeNode,
+		TSharedPtr<IPropertyHandle>* OutPropertyHandle = nullptr,
+		bool bFocusInputWidget = false);
+
+	/** 
+	* FindFocusableWidgetAndSetKeyboardFocus
+	* 
+	* Searches the entire widget hierarchy of a given widget for a focusable child item and sets the focus on it 
+	* If the input widget is itself focusable, then we simply set the focus on it and return
+	*/
+	bool FindFocusableWidgetAndSetKeyboardFocus(const TSharedRef< SWidget> InWidget);
 }
\ No newline at end of file
diff --git a/Engine/Plugins/VirtualProduction/RemoteControl/Source/RemoteControlUI/Private/UI/SRemoteControlPanel.h b/Engine/Plugins/VirtualProduction/RemoteControl/Source/RemoteControlUI/Private/UI/SRemoteControlPanel.h
--- a/Engine/Plugins/VirtualProduction/RemoteControl/Source/RemoteControlUI/Private/UI/SRemoteControlPanel.h
+++ b/Engine/Plugins/VirtualProduction/RemoteControl/Source/RemoteControlUI/Private/UI/SRemoteControlPanel.h
@@ -34,6 +34,7 @@ class SComboButton;
 struct SRCPanelTreeNode;
 class SRCPanelFunctionPicker;
 class SRemoteControlPanel;
+class SRCActionPanel;
 class SRCPanelDrawer;
 class SRCPanelExposedEntitiesList;
 class SRCPanelFilter;
@@ -141,6 +142,12 @@ public:
 	void AddToolbarWidget(TSharedRef<SWidget> Widget);
 	void RemoveAllToolbarWidgets();
 
+	/** Retrieves the Logic Action panel. */
+	TSharedPtr<SRCActionPanel> GetLogicActionPanel()
+	{
+		return ActionPanel;
+	}
+
 private:
 
 	//~ Remote Control Commands
