Bind Behaviour

diff --git a/Engine/Plugins/VirtualProduction/RemoteControl/Source/RemoteControlLogic/Private/Action/Bind/RCPropertyBindAction.cpp b/Engine/Plugins/VirtualProduction/RemoteControl/Source/RemoteControlLogic/Private/Action/Bind/RCPropertyBindAction.cpp
new file mode 100644
--- /dev/null
+++ b/Engine/Plugins/VirtualProduction/RemoteControl/Source/RemoteControlLogic/Private/Action/Bind/RCPropertyBindAction.cpp
@@ -0,0 +1,193 @@
+// Copyright Epic Games, Inc. All Rights Reserved.
+
+#include "Action/Bind/RCPropertyBindAction.h"
+
+#include "Controller/RCController.h"
+#include "IRemoteControlPropertyHandle.h"
+#include "RCVirtualProperty.h"
+#include "RemoteControlPreset.h"
+#include "Action/RCAction.h"
+
+static void SetStructPropertyFromController(const FProperty* RemoteControlProperty, TSharedRef<IRemoteControlPropertyHandle> RemoteControlHandle, const URCController* Controller)
+{
+	const FStructProperty* SourceStructProperty = CastField<FStructProperty>(Controller->GetProperty());
+	const FStructProperty* TargetStructProperty = CastField<FStructProperty>(RemoteControlProperty);
+
+	if (ensure(SourceStructProperty && TargetStructProperty))
+	{
+		// Pass 1 - Structs of exactly the same type
+		if (SourceStructProperty->Struct == TargetStructProperty->Struct)
+		{
+			// FVector
+			if (TargetStructProperty->Struct == TBaseStructure<FVector>::Get())
+			{
+				FVector VectorValue;
+				Controller->GetValueVector(VectorValue);
+
+				RemoteControlHandle->SetValue(VectorValue);
+			}
+			// FColor
+			else if (TargetStructProperty->Struct == TBaseStructure<FColor>::Get())
+			{
+				FColor ColorValue;
+				Controller->GetValueColor(ColorValue);
+
+				RemoteControlHandle->SetValue(ColorValue);
+			}
+			// FRotator
+			else if (TargetStructProperty->Struct == TBaseStructure<FRotator>::Get())
+			{
+				FRotator RotatorValue;
+				Controller->GetValueRotator(RotatorValue);
+
+				RemoteControlHandle->SetValue(RotatorValue);
+			}
+		}
+		// Pass 2 - Structs that support conversion to other (related) Structs
+		else
+		{
+			// Source: FColor
+			if (SourceStructProperty->Struct == TBaseStructure<FColor>::Get())
+			{
+				// Target: FLinearColor
+				if (TargetStructProperty->Struct == TBaseStructure<FLinearColor>::Get())
+				{
+					FColor ColorValue;
+					Controller->GetValueColor(ColorValue);
+
+					RemoteControlHandle->SetValue(FLinearColor(ColorValue));
+				}
+			}
+		}
+	}
+}
+
+void URCPropertyBindAction::Execute() const
+{
+	if (!ensure(Controller))
+	{
+		return;
+	}
+
+	if (TSharedPtr<FRemoteControlProperty> RemoteControlEntityAsProperty = GetRemoteControlProperty())
+	{
+		const FProperty* ControllerAsProperty = Controller->GetProperty();
+		FProperty* RemoteControlProperty = RemoteControlEntityAsProperty->GetProperty();
+
+		TSharedPtr<IRemoteControlPropertyHandle> Handle = RemoteControlProperty ? RemoteControlEntityAsProperty->GetPropertyHandle() : nullptr;
+		if (!ensure(Handle))
+		{
+			return;
+		}
+
+		// String Controller
+		if (ControllerAsProperty->IsA(FStrProperty::StaticClass()))
+		{
+			FString StringValue;
+			Controller->GetValueString(StringValue);
+
+			// String to String/Name/Text
+			if (RemoteControlProperty->IsA(FStrProperty::StaticClass()) ||
+				RemoteControlProperty->IsA(FTextProperty::StaticClass()) ||
+				RemoteControlProperty->IsA(FNameProperty::StaticClass()))
+			{
+				// Set value directly. Implicitly converts to Text, Name, String
+				// Also provides replication support via IRemoteControlPropertyHandle
+				Handle->SetValue(StringValue);
+			}
+			else
+			{
+				const float NumericValue = FCString::Atof(*StringValue);
+
+				// String to Byte/Enum
+				if (RemoteControlProperty->IsA(FByteProperty::StaticClass()))
+				{
+					Handle->SetValue((uint8)NumericValue);
+				}
+				// String to Numeric
+				else if (RemoteControlProperty->IsA(FNumericProperty::StaticClass()))
+				{
+					Handle->SetValue(NumericValue);
+				}
+				// String to Boolean
+				else if (RemoteControlProperty->IsA(FBoolProperty::StaticClass()))
+				{
+					Handle->SetValue((bool)NumericValue);
+				}
+			}
+			// String Controller (end)
+		}
+		// Numeric Controller (Float/Integer)
+		else if (ControllerAsProperty->IsA(FNumericProperty::StaticClass()))
+		{
+			// Extract Numeric Value
+			float NumericValue = 0.f; // placate compiler warning
+			if (ControllerAsProperty->IsA(FIntProperty::StaticClass()))
+			{
+				int32 IntValue;
+				Controller->GetValueInt32(IntValue);
+				NumericValue = IntValue;
+			}
+			else if (ControllerAsProperty->IsA(FFloatProperty::StaticClass()))
+			{
+				Controller->GetValueFloat(NumericValue);
+			}
+
+			// Numeric to Byte/Enum
+			if (RemoteControlProperty->IsA(FByteProperty::StaticClass()))
+			{
+				Handle->SetValue((uint8)NumericValue);
+			}
+			// Numeric To Numeric
+			if (RemoteControlProperty->IsA(FNumericProperty::StaticClass()))
+			{
+				Handle->SetValue(NumericValue);
+			}
+			// Numeric To Boolean
+			else if (RemoteControlProperty->IsA(FBoolProperty::StaticClass()))
+			{
+				Handle->SetValue((bool)NumericValue);
+			}
+			// Numeric to String (explicit conversion)
+			else if (RemoteControlProperty->IsA(FStrProperty::StaticClass()) ||
+				RemoteControlProperty->IsA(FTextProperty::StaticClass()) ||
+				RemoteControlProperty->IsA(FNameProperty::StaticClass()))
+			{
+				FString StringValue;
+
+				if (ControllerAsProperty->IsA(FIntProperty::StaticClass()))
+				{
+					StringValue = FString::Printf(TEXT("%d"), (int32)NumericValue);
+					Handle->SetValue(StringValue);
+				}
+				else
+				{
+					// %g is used over %f to discard extra zero value fractional digits from the float (for visual clarity)
+					StringValue = FString::Printf(TEXT("%g"), NumericValue);
+				}
+
+				Handle->SetValue(StringValue);
+			}
+		}
+		// Boolean Controller
+		else if (ControllerAsProperty->IsA(FBoolProperty::StaticClass()))
+		{
+			bool BoolValue;
+			Controller->GetValueBool(BoolValue);
+
+			if (RemoteControlProperty->IsA(FBoolProperty::StaticClass()))
+			{
+				Handle->SetValue(BoolValue);
+			}
+			else if (RemoteControlProperty->IsA(FNumericProperty::StaticClass()))
+			{
+				Handle->SetValue((float)BoolValue);
+			}
+			// Note: Explicit conversion from Boolean to String/Text/Name currently not deemed as a requirement, so hasn't been implemented here
+		}
+		else if (ControllerAsProperty->IsA(FStructProperty::StaticClass()))
+		{
+			SetStructPropertyFromController(RemoteControlProperty, Handle.ToSharedRef(), Controller);
+		}
+	}
+}
\ No newline at end of file
diff --git a/Engine/Plugins/VirtualProduction/RemoteControl/Source/RemoteControlLogic/Private/Action/RCActionContainer.cpp b/Engine/Plugins/VirtualProduction/RemoteControl/Source/RemoteControlLogic/Private/Action/RCActionContainer.cpp
--- a/Engine/Plugins/VirtualProduction/RemoteControl/Source/RemoteControlLogic/Private/Action/RCActionContainer.cpp
+++ b/Engine/Plugins/VirtualProduction/RemoteControl/Source/RemoteControlLogic/Private/Action/RCActionContainer.cpp
@@ -2,6 +2,7 @@
 
 #include "Action/RCActionContainer.h"
 
+#include "Controller/RCController.h"
 #include "IRemoteControlModule.h"
 #include "RCVirtualProperty.h"
 #include "RemoteControlField.h"
diff --git a/Engine/Plugins/VirtualProduction/RemoteControl/Source/RemoteControlLogic/Private/Action/RCPropertyAction.cpp b/Engine/Plugins/VirtualProduction/RemoteControl/Source/RemoteControlLogic/Private/Action/RCPropertyAction.cpp
--- a/Engine/Plugins/VirtualProduction/RemoteControl/Source/RemoteControlLogic/Private/Action/RCPropertyAction.cpp
+++ b/Engine/Plugins/VirtualProduction/RemoteControl/Source/RemoteControlLogic/Private/Action/RCPropertyAction.cpp
@@ -2,7 +2,9 @@
 
 #include "Action/RCPropertyAction.h"
 
+#include "Controller/RCController.h"
 #include "IRemoteControlModule.h"
+#include "IRemoteControlPropertyHandle.h"
 #include "IStructSerializerBackend.h"
 #include "RCVirtualProperty.h"
 #include "RemoteControlPreset.h"
@@ -69,4 +71,17 @@ FName URCAction::GetExposedFieldLabel() const
 	}
 
 	return NAME_None;
+}
+
+TSharedPtr<FRemoteControlProperty> URCPropertyAction::GetRemoteControlProperty() const
+{
+	if (URemoteControlPreset* Preset = PresetWeakPtr.Get())
+	{
+		if (TSharedPtr<FRemoteControlProperty> RemoteControlProperty = Preset->GetExposedEntity<FRemoteControlProperty>(ExposedFieldId).Pin())
+		{
+			return RemoteControlProperty;
+		}
+	}
+
+	return nullptr;
 }
\ No newline at end of file
diff --git a/Engine/Plugins/VirtualProduction/RemoteControl/Source/RemoteControlLogic/Private/Behaviour/Builtin/Bind/RCBehaviourBind.cpp b/Engine/Plugins/VirtualProduction/RemoteControl/Source/RemoteControlLogic/Private/Behaviour/Builtin/Bind/RCBehaviourBind.cpp
new file mode 100644
--- /dev/null
+++ b/Engine/Plugins/VirtualProduction/RemoteControl/Source/RemoteControlLogic/Private/Behaviour/Builtin/Bind/RCBehaviourBind.cpp
@@ -0,0 +1,161 @@
+// Copyright Epic Games, Inc. All Rights Reserved.
+
+#include "Behaviour/Builtin/Bind/RCBehaviourBind.h"
+
+#include "Action/RCAction.h"
+#include "Action/RCFunctionAction.h"
+#include "Action/Bind/RCPropertyBindAction.h"
+#include "Behaviour/Builtin/Bind/RCBehaviourBindNode.h"
+#include "Controller/RCController.h"
+#include "RemoteControlField.h"
+#include "RCVirtualProperty.h"
+
+URCBehaviourBind::URCBehaviourBind()
+{
+}
+
+void URCBehaviourBind::Initialize()
+{
+	Super::Initialize();
+}
+
+URCPropertyBindAction* URCBehaviourBind::AddPropertyBindAction(const TSharedRef<const FRemoteControlProperty> InRemoteControlProperty)
+{
+	if (!ensure(ControllerWeakPtr.IsValid() && ActionContainer))
+	{
+		return nullptr;
+	}
+
+	URCPropertyBindAction* BindAction = NewObject<URCPropertyBindAction>(ActionContainer);
+	BindAction->PresetWeakPtr = ActionContainer->PresetWeakPtr;
+	BindAction->ExposedFieldId = InRemoteControlProperty->GetId();
+	BindAction->Controller = ControllerWeakPtr.Get();
+	BindAction->Id = FGuid::NewGuid();
+
+	// Add action to array
+	ActionContainer->Actions.Add(BindAction);
+
+	return BindAction;
+}
+
+URCAction* URCBehaviourBind::AddAction(const TSharedRef<const FRemoteControlField> InRemoteControlField)
+{
+	const TSharedRef<const FRemoteControlProperty> RemoteControlProperty = StaticCastSharedRef<const FRemoteControlProperty>(InRemoteControlField);
+
+	URCAction* BindAction = AddPropertyBindAction(RemoteControlProperty);
+
+	if (ensure(BindAction))
+	{
+		// For Bind Behaviour we want to pick up the Controller's current value immediately.
+		// Subsequent value updates from Controller to Action will be propagated via the usual OnModifyPropertyValue code path
+
+		BindAction->Execute();
+	}
+
+	return BindAction;
+}
+
+bool URCBehaviourBind::CanHaveActionForField(const TSharedPtr<FRemoteControlField> InRemoteControlField) const
+{
+	// Basic check (uniqueness)
+	if (ActionContainer->FindActionByFieldId(InRemoteControlField->GetId()))
+	{
+		return false; // already exists!
+	}
+
+	// Advanced checks (by Controller type and Target type)
+	if(TSharedPtr<FRemoteControlProperty> RemoteControlEntityAsProperty = StaticCastSharedPtr<FRemoteControlProperty>(InRemoteControlField))
+	{
+		if (URCController* Controller = ControllerWeakPtr.Get())
+		{
+			const FProperty* ControllerAsProperty = Controller->GetProperty();
+			const FProperty* RemoteControlProperty = RemoteControlEntityAsProperty->GetProperty();
+			const FStructProperty* RemoteControlEntityAsStructProperty = CastField<FStructProperty>(RemoteControlProperty);
+
+			bool IsNumericStruct = RemoteControlEntityAsStructProperty
+				&& (RemoteControlEntityAsStructProperty->Struct == TBaseStructure<FVector>::Get() ||
+					RemoteControlEntityAsStructProperty->Struct == TBaseStructure<FRotator>::Get() ||
+					RemoteControlEntityAsStructProperty->Struct == TBaseStructure<FColor>::Get() ||
+					RemoteControlEntityAsStructProperty->Struct == TBaseStructure<FLinearColor>::Get());
+
+			// String Controller
+			if (ControllerAsProperty->IsA(FStrProperty::StaticClass()))
+			{
+				if(RemoteControlProperty->IsA(FStrProperty::StaticClass()) ||
+					RemoteControlProperty->IsA(FTextProperty::StaticClass()) ||
+					RemoteControlProperty->IsA(FNameProperty::StaticClass()))
+				{
+					return true; // Implicit / direct conversions are possible between these types
+				}
+				else if (RemoteControlProperty->IsA(FNumericProperty::StaticClass()) ||
+					RemoteControlProperty->IsA(FBoolProperty::StaticClass()) ||
+					RemoteControlProperty->IsA(FByteProperty::StaticClass()))
+				{
+					// Explicit conversion to numbers, governed by user set flag
+					return bAllowNumericInputAsStrings;
+				}
+			}
+			// Numeric Controller (Float/Int)
+			else if (ControllerAsProperty->IsA(FNumericProperty::StaticClass()))
+			{
+				if (RemoteControlProperty->IsA(FNumericProperty::StaticClass()) ||
+					RemoteControlProperty->IsA(FBoolProperty::StaticClass()) || 
+					RemoteControlProperty->IsA(FByteProperty::StaticClass()))
+				{
+					return true; // Numeric to Numeric - Implicit / direct conversion
+
+					// Note: Numeric to Numeric struct conversion is not currently supported as it would ideally require Masking support (on the struct) for meaningful usage
+				}
+				else if (RemoteControlProperty->IsA(FStrProperty::StaticClass()) ||
+					RemoteControlProperty->IsA(FTextProperty::StaticClass()) ||
+					RemoteControlProperty->IsA(FNameProperty::StaticClass()))
+				{
+					// Explicit conversion to String/Name/Text, governed by user set flag
+					return bAllowNumericInputAsStrings;
+				}
+			}
+			// Boolean Controller
+			else if (ControllerAsProperty->IsA(FBoolProperty::StaticClass()))
+			{
+				if (RemoteControlProperty->IsA(FFloatProperty::StaticClass()) || // Note: FNumericPropery includes Enum, so we check explicitly for Float and Int here
+					RemoteControlProperty->IsA(FIntProperty::StaticClass()) ||
+					RemoteControlProperty->IsA(FBoolProperty::StaticClass()))
+				{
+					return true;
+				}
+				// Notes for explicit Bool conversions: 
+				// 1. Boolean to String/Text/Name is currently not deemed as a requirement, so hasn't been implemented here
+				// 2. Boolean to Struct is also not implemented. 
+			}
+			// Struct
+			else if (ControllerAsProperty->IsA(FStructProperty::StaticClass()))
+			{
+				if (const FStructProperty* StructProperty = CastField<FStructProperty>(ControllerAsProperty))
+				{
+					// FVector / FColor / FRotator
+					if (const FStructProperty* RCFieldStructProperty = CastField<FStructProperty>(RemoteControlProperty))
+					{
+						// Generally - A Struct Controller can only drive properties of the _same_ Struct type...
+						if (RCFieldStructProperty->Struct == StructProperty->Struct)
+						{
+							return true;
+						}
+						// ... with an exception for the following Struct conversions:
+						else
+						{
+							// FColor conversion:
+							if (StructProperty->Struct == TBaseStructure<FColor>::Get())
+							{
+								// To FLinearColor, is supported.
+								return RCFieldStructProperty->Struct == TBaseStructure<FLinearColor>::Get();
+							}
+						}
+
+					}
+				}
+			}
+		}
+	}
+
+	return false;
+}
diff --git a/Engine/Plugins/VirtualProduction/RemoteControl/Source/RemoteControlLogic/Private/Behaviour/Builtin/Bind/RCBehaviourBindNode.cpp b/Engine/Plugins/VirtualProduction/RemoteControl/Source/RemoteControlLogic/Private/Behaviour/Builtin/Bind/RCBehaviourBindNode.cpp
new file mode 100644
--- /dev/null
+++ b/Engine/Plugins/VirtualProduction/RemoteControl/Source/RemoteControlLogic/Private/Behaviour/Builtin/Bind/RCBehaviourBindNode.cpp
@@ -0,0 +1,34 @@
+// Copyright Epic Games, Inc. All Rights Reserved.
+
+#include "Behaviour/Builtin/Bind/RCBehaviourBindNode.h"
+
+#include "Behaviour/Builtin/Bind/RCBehaviourBind.h"
+
+#define LOCTEXT_NAMESPACE "RemoteControlBehaviours"
+
+URCBehaviourBindNode::URCBehaviourBindNode()
+{
+	DisplayName = LOCTEXT("BehaviorNameBind", "Bind");
+	BehaviorDescription = LOCTEXT("BehaviorDescBind", "Bind exposed properties to match the value of the controller. Properties added to the action list will not receive input values");
+}
+
+bool URCBehaviourBindNode::Execute_Implementation(URCBehaviour* InBehaviour) const
+{
+	return true;
+}
+
+bool URCBehaviourBindNode::IsSupported_Implementation(URCBehaviour* InBehaviour) const
+{
+	return true; 
+}
+
+UClass* URCBehaviourBindNode::GetBehaviourClass() const
+{
+	return URCBehaviourBind::StaticClass();
+}
+
+void URCBehaviourBindNode::OnPassed_Implementation(URCBehaviour* InBehaviour) const
+{
+}
+
+#undef LOCTEXT_NAMESPACE
\ No newline at end of file
diff --git a/Engine/Plugins/VirtualProduction/RemoteControl/Source/RemoteControlLogic/Private/Behaviour/Builtin/Conditional/RCBehaviourConditional.cpp b/Engine/Plugins/VirtualProduction/RemoteControl/Source/RemoteControlLogic/Private/Behaviour/Builtin/Conditional/RCBehaviourConditional.cpp
--- a/Engine/Plugins/VirtualProduction/RemoteControl/Source/RemoteControlLogic/Private/Behaviour/Builtin/Conditional/RCBehaviourConditional.cpp
+++ b/Engine/Plugins/VirtualProduction/RemoteControl/Source/RemoteControlLogic/Private/Behaviour/Builtin/Conditional/RCBehaviourConditional.cpp
@@ -133,7 +133,7 @@ void URCBehaviourConditional::OnActionAdded(URCAction* Action, const ERCBehaviou
 	Conditions.Add(Action, MoveTemp(Condition));
 }
 
-URCAction* URCBehaviourConditional::AddAction(const TSharedRef<const FRemoteControlField> InRemoteControlField, const ERCBehaviourConditionType InConditionType, const TObjectPtr<URCVirtualPropertySelfContainer> InComparand)
+URCAction* URCBehaviourConditional::AddConditionalAction(const TSharedRef<const FRemoteControlField> InRemoteControlField, const ERCBehaviourConditionType InConditionType, const TObjectPtr<URCVirtualPropertySelfContainer> InComparand)
 {
 	TRCActionUniquenessTest UniquenessTest = [this, InRemoteControlField, InConditionType, InComparand](const TSet<TObjectPtr<URCAction>>& Actions)
 	{
diff --git a/Engine/Plugins/VirtualProduction/RemoteControl/Source/RemoteControlLogic/Private/Controller/RCControllerContainer.cpp b/Engine/Plugins/VirtualProduction/RemoteControl/Source/RemoteControlLogic/Private/Controller/RCControllerContainer.cpp
--- a/Engine/Plugins/VirtualProduction/RemoteControl/Source/RemoteControlLogic/Private/Controller/RCControllerContainer.cpp
+++ b/Engine/Plugins/VirtualProduction/RemoteControl/Source/RemoteControlLogic/Private/Controller/RCControllerContainer.cpp
@@ -14,7 +14,7 @@ void URCControllerContainer::OnModifyPropertyValue(const FPropertyChangedEvent&
 			{
 				if (URCController* Controller = Cast<URCController>(VirtualProperty))
 				{
-					Controller->ExecuteBehaviours();
+					Controller->OnModifyPropertyValue();
 				}
 			}
 		}
diff --git a/Engine/Plugins/VirtualProduction/RemoteControl/Source/RemoteControlLogic/Public/Action/Bind/RCPropertyBindAction.h b/Engine/Plugins/VirtualProduction/RemoteControl/Source/RemoteControlLogic/Public/Action/Bind/RCPropertyBindAction.h
new file mode 100644
--- /dev/null
+++ b/Engine/Plugins/VirtualProduction/RemoteControl/Source/RemoteControlLogic/Public/Action/Bind/RCPropertyBindAction.h
@@ -0,0 +1,28 @@
+// Copyright Epic Games, Inc. All Rights Reserved.
+
+#pragma once
+
+#include "Action/RCPropertyAction.h"
+#include "RCPropertyBindAction.generated.h"
+
+class URCController;
+
+/**
+ * Property Bind Action specifically for Bind Behaviour
+ */
+UCLASS()
+class REMOTECONTROLLOGIC_API URCPropertyBindAction : public URCPropertyAction
+{
+	GENERATED_BODY()
+
+public:
+	URCPropertyBindAction() = default;
+
+	//~ Begin URCAction interface
+	virtual void Execute() const override;
+	//~ End URCAction interface
+
+	/** The Controller that drives us */
+	UPROPERTY()
+	TObjectPtr<URCController> Controller;
+};
diff --git a/Engine/Plugins/VirtualProduction/RemoteControl/Source/RemoteControlLogic/Public/Action/RCActionContainer.h b/Engine/Plugins/VirtualProduction/RemoteControl/Source/RemoteControlLogic/Public/Action/RCActionContainer.h
--- a/Engine/Plugins/VirtualProduction/RemoteControl/Source/RemoteControlLogic/Public/Action/RCActionContainer.h
+++ b/Engine/Plugins/VirtualProduction/RemoteControl/Source/RemoteControlLogic/Public/Action/RCActionContainer.h
@@ -9,8 +9,10 @@
 
 class URemoteControlPreset;
 class URCAction;
+class URCController;
 class URCFunctionAction;
 class URCPropertyAction;
+class URCPropertyBindAction;
 
 struct FRemoteControlField;
 struct FRemoteControlFunction;
diff --git a/Engine/Plugins/VirtualProduction/RemoteControl/Source/RemoteControlLogic/Public/Action/RCPropertyAction.h b/Engine/Plugins/VirtualProduction/RemoteControl/Source/RemoteControlLogic/Public/Action/RCPropertyAction.h
--- a/Engine/Plugins/VirtualProduction/RemoteControl/Source/RemoteControlLogic/Public/Action/RCPropertyAction.h
+++ b/Engine/Plugins/VirtualProduction/RemoteControl/Source/RemoteControlLogic/Public/Action/RCPropertyAction.h
@@ -5,6 +5,8 @@
 #include "RCAction.h"
 #include "RCPropertyAction.generated.h"
 
+struct FRemoteControlProperty;
+class URCController;
 class URCVirtualPropertySelfContainer;
 
 /**
@@ -26,4 +28,6 @@ public:
 	/** Virtual Property Container */
 	UPROPERTY()
 	TObjectPtr<URCVirtualPropertySelfContainer> PropertySelfContainer = nullptr;
-};
+
+	TSharedPtr<FRemoteControlProperty> GetRemoteControlProperty() const;
+};
\ No newline at end of file
diff --git a/Engine/Plugins/VirtualProduction/RemoteControl/Source/RemoteControlLogic/Public/Behaviour/Builtin/Bind/RCBehaviourBind.h b/Engine/Plugins/VirtualProduction/RemoteControl/Source/RemoteControlLogic/Public/Behaviour/Builtin/Bind/RCBehaviourBind.h
new file mode 100644
--- /dev/null
+++ b/Engine/Plugins/VirtualProduction/RemoteControl/Source/RemoteControlLogic/Public/Behaviour/Builtin/Bind/RCBehaviourBind.h
@@ -0,0 +1,47 @@
+// Copyright Epic Games, Inc. All Rights Reserved.
+
+#pragma once
+
+#include "Action/RCActionContainer.h"
+#include "Behaviour/RCBehaviour.h"
+
+#include "RCBehaviourBind.generated.h"
+
+class URCVirtualPropertySelfContainer;
+class URCAction;
+
+/**
+ * [Bind Behaviour]
+ * 
+ * Binds a given Controller with multiple exposed properties.
+ * Any changes to the value of the controller are directly propagated to the linked properties by Bind Behaviour. 
+ */
+UCLASS()
+class REMOTECONTROLLOGIC_API URCBehaviourBind : public URCBehaviour
+{
+	GENERATED_BODY()
+
+public:
+	URCBehaviourBind();
+	
+	//~ Begin URCBehaviour interface
+	virtual void Initialize() override;
+
+	/** Add a Logic action using a remote control field as input */
+	virtual URCAction* AddAction(const TSharedRef<const FRemoteControlField> InRemoteControlField) override;
+
+	/** Whether we can create an action pertaining to a given remote control field for the current behaviour */
+	virtual bool CanHaveActionForField(const TSharedPtr<FRemoteControlField> InRemoteControlField) const override;
+
+	//~ End URCBehaviour interface
+
+	/** Add an action specifically for Bind Behaviour */
+	URCPropertyBindAction* AddPropertyBindAction(const TSharedRef<const FRemoteControlProperty> InRemoteControlProperty);
+
+	/** Indicates whether we support binding of strings to numeric fields and vice versa.
+	* This flag determines the list of bindable properties the user sees in the "Add Action" menu*/
+	bool AreNumericInputsAllowedAsStrings() const { return bAllowNumericInputAsStrings; }
+	void SetAllowNumericInputAsStrings(bool Value) { bAllowNumericInputAsStrings = Value; }
+private:
+	bool bAllowNumericInputAsStrings = false;
+};
diff --git a/Engine/Plugins/VirtualProduction/RemoteControl/Source/RemoteControlLogic/Public/Behaviour/Builtin/Bind/RCBehaviourBindNode.h b/Engine/Plugins/VirtualProduction/RemoteControl/Source/RemoteControlLogic/Public/Behaviour/Builtin/Bind/RCBehaviourBindNode.h
new file mode 100644
--- /dev/null
+++ b/Engine/Plugins/VirtualProduction/RemoteControl/Source/RemoteControlLogic/Public/Behaviour/Builtin/Bind/RCBehaviourBindNode.h
@@ -0,0 +1,32 @@
+// Copyright Epic Games, Inc. All Rights Reserved.
+
+#pragma once
+
+#include "Behaviour/RCBehaviourNode.h"
+#include "RCBehaviourBindNode.generated.h"
+
+/**
+ * Behaviour Node class for Bind Behaviour
+ */
+UCLASS()
+class REMOTECONTROLLOGIC_API URCBehaviourBindNode : public URCBehaviourNode
+{
+	GENERATED_BODY()
+
+public:
+	
+	URCBehaviourBindNode();
+
+	//~ Begin URCBehaviourNode interface
+	UFUNCTION(BlueprintNativeEvent, BlueprintCallable, Category = "Remote Control Behaviour")
+	bool Execute(URCBehaviour* InBehaviour) const override;
+
+	UFUNCTION(BlueprintNativeEvent, BlueprintCallable, Category = "Remote Control Behaviour")
+	bool IsSupported(URCBehaviour* InBehaviour) const override;
+
+	UFUNCTION(BlueprintNativeEvent, BlueprintCallable, Category = "Remote Control Behaviour")
+	void OnPassed(URCBehaviour* InBehaviour) const;
+
+	UClass* GetBehaviourClass() const override;
+	//~ End URCBehaviourNode interface
+};
diff --git a/Engine/Plugins/VirtualProduction/RemoteControl/Source/RemoteControlLogic/Public/Behaviour/Builtin/Conditional/RCBehaviourConditional.h b/Engine/Plugins/VirtualProduction/RemoteControl/Source/RemoteControlLogic/Public/Behaviour/Builtin/Conditional/RCBehaviourConditional.h
--- a/Engine/Plugins/VirtualProduction/RemoteControl/Source/RemoteControlLogic/Public/Behaviour/Builtin/Conditional/RCBehaviourConditional.h
+++ b/Engine/Plugins/VirtualProduction/RemoteControl/Source/RemoteControlLogic/Public/Behaviour/Builtin/Conditional/RCBehaviourConditional.h
@@ -76,7 +76,7 @@ public:
 	void OnActionAdded(URCAction* Action, const ERCBehaviourConditionType InConditionType, const TObjectPtr<URCVirtualPropertySelfContainer> InComparand);
 
 	/** Add a Logic action using a remote control field as input */
-	URCAction* AddAction(const TSharedRef<const FRemoteControlField> InRemoteControlField, const ERCBehaviourConditionType InConditionType, const TObjectPtr<URCVirtualPropertySelfContainer> InComparand);
+	URCAction* AddConditionalAction(const TSharedRef<const FRemoteControlField> InRemoteControlField, const ERCBehaviourConditionType InConditionType, const TObjectPtr<URCVirtualPropertySelfContainer> InComparand);
 
 	/** Whether we can create an action pertaining to a given remote control field for the current behaviour */
 	virtual bool CanHaveActionForField(const TSharedPtr<FRemoteControlField> InRemoteControlField) const override
diff --git a/Engine/Plugins/VirtualProduction/RemoteControl/Source/RemoteControlLogic/Public/Behaviour/RCBehaviour.h b/Engine/Plugins/VirtualProduction/RemoteControl/Source/RemoteControlLogic/Public/Behaviour/RCBehaviour.h
--- a/Engine/Plugins/VirtualProduction/RemoteControl/Source/RemoteControlLogic/Public/Behaviour/RCBehaviour.h
+++ b/Engine/Plugins/VirtualProduction/RemoteControl/Source/RemoteControlLogic/Public/Behaviour/RCBehaviour.h
@@ -42,7 +42,7 @@ public:
 	virtual void Execute();
 
 	/** Add a Logic action using a remote control field as input */
-	URCAction* AddAction(const TSharedRef<const FRemoteControlField> InRemoteControlField);
+	virtual URCAction* AddAction(const TSharedRef<const FRemoteControlField> InRemoteControlField);
 
 	/** Get number of action associated with behaviour */
 	int32 GetNumActions() const;
diff --git a/Engine/Plugins/VirtualProduction/RemoteControl/Source/RemoteControlUI/Private/UI/Action/Bind/RCActionBindModel.cpp b/Engine/Plugins/VirtualProduction/RemoteControl/Source/RemoteControlUI/Private/UI/Action/Bind/RCActionBindModel.cpp
new file mode 100644
--- /dev/null
+++ b/Engine/Plugins/VirtualProduction/RemoteControl/Source/RemoteControlUI/Private/UI/Action/Bind/RCActionBindModel.cpp
@@ -0,0 +1,135 @@
+// Copyright Epic Games, Inc. All Rights Reserved.
+
+#include "RCActionBindModel.h"
+
+#include "Action/RCPropertyAction.h"
+#include "Behaviour/Builtin/Bind/RCBehaviourBind.h"
+#include "Controller/RCController.h"
+#include "RemoteControlField.h"
+#include "Styling/RemoteControlStyles.h"
+#include "UI/Behaviour/Builtin/Bind/RCBehaviourBindModel.h"
+#include "UI/RCUIHelpers.h"
+#include "UI/RemoteControlPanelStyle.h"
+#include "UI/SRCPanelDragHandle.h"
+#include "Widgets/Views/SHeaderRow.h"
+#include "Widgets/Views/SListView.h"
+#include "Widgets/Views/STableRow.h"
+
+#define LOCTEXT_NAMESPACE "RCActionModel"
+
+namespace UE::RCActionPanelBindList
+{
+	namespace Columns
+	{
+		const FName VariableColor = TEXT("VariableColor");
+		const FName DragDropHandle = TEXT("DragDropHandle");
+		const FName Description = TEXT("Description");
+		const FName Controller = TEXT("Controller");
+	}
+
+	class SBindActionRow : public SMultiColumnTableRow<TSharedRef<FRCActionModel>>
+	{
+	public:
+		void Construct(const FTableRowArgs& InArgs, const TSharedRef<STableViewBase>& OwnerTableView, TSharedRef<FRCActionBindModel> InActionItem)
+		{
+			ActionItem = InActionItem;
+			FSuperRowType::Construct(InArgs, OwnerTableView);
+		}
+
+		TSharedRef<SWidget> GenerateWidgetForColumn(const FName& ColumnName) override
+		{
+			if (!ensure(ActionItem.IsValid()))
+			{
+				return SNullWidget::NullWidget;
+			}
+		
+			if (ColumnName == UE::RCActionPanelBindList::Columns::VariableColor)
+			{
+				return ActionItem->GetVariableColorWidget();
+			}
+			else if (ColumnName == UE::RCActionPanelBindList::Columns::Description)
+			{
+				return ActionItem->GetNameWidget();
+			}
+			else if (ColumnName == UE::RCActionPanelBindList::Columns::Controller)
+			{
+				if (URCPropertyBindAction* BindAction = Cast< URCPropertyBindAction>(ActionItem->GetAction()))
+				{
+					if (ensure(BindAction->Controller))
+					{
+						const FText ControllerName = FText::FromName(BindAction->Controller->DisplayName);
+
+						return SNew(SBox)
+									.Padding(FMargin(6.f))
+									[
+										SNew(STextBlock).Text(ControllerName)
+									];
+							
+					}
+				}
+			}
+
+			return SNullWidget::NullWidget;
+		}
+
+	private:
+		TSharedPtr<FRCActionBindModel> ActionItem;
+	};
+}
+
+TSharedPtr<SHeaderRow> FRCActionBindModel::GetHeaderRow()
+{
+	const FRCPanelStyle* RCPanelStyle = &FRemoteControlPanelStyle::Get()->GetWidgetStyle<FRCPanelStyle>("RemoteControlPanel.MinorPanel");
+
+	return SNew(SHeaderRow)
+		.Style(&RCPanelStyle->HeaderRowStyle)
+
+		+ SHeaderRow::Column(UE::RCActionPanelBindList::Columns::VariableColor)
+		.DefaultLabel(LOCTEXT("VariableColorColumnHeader", ""))
+		.FixedWidth(5.f)
+		.HeaderContentPadding(RCPanelStyle->HeaderRowPadding)
+
+		+ SHeaderRow::Column(UE::RCActionPanelBindList::Columns::Description)
+		.DefaultLabel(LOCTEXT("DescriptionColumnHeader", "Description"))
+		.FillWidth(0.4f)
+		.HeaderContentPadding(RCPanelStyle->HeaderRowPadding)
+
+		+ SHeaderRow::Column(UE::RCActionPanelBindList::Columns::Controller)
+		.DefaultLabel(LOCTEXT("ControllerColumnHeader", "Controller"))
+		.FillWidth(0.4f)
+		.HeaderContentPadding(RCPanelStyle->HeaderRowPadding);
+}
+
+TSharedPtr<FRCActionBindModel> FRCActionBindModel::GetModelByActionType(URCAction* InAction, const TSharedPtr<class FRCBehaviourModel> InBehaviourItem, const TSharedPtr<SRemoteControlPanel> InRemoteControlPanel)
+{
+	if (URCPropertyBindAction* PropertyBindAction = Cast<URCPropertyBindAction>(InAction))
+	{
+		return MakeShared<FRCPropertyActionBindModel>(PropertyBindAction, InBehaviourItem, InRemoteControlPanel);
+	}
+
+	return nullptr;
+}
+
+TSharedRef<ITableRow> FRCActionBindModel::OnGenerateWidgetForList(TSharedPtr<FRCActionBindModel> InItem, const TSharedRef<STableViewBase>& OwnerTable)
+{
+	typedef UE::RCActionPanelBindList::SBindActionRow ActionRowType;
+
+	return SNew(ActionRowType, OwnerTable, InItem.ToSharedRef())
+		.Style(&RCPanelStyle->TableRowStyle)
+		.Padding(FMargin(3.f));
+}
+
+FLinearColor FRCPropertyActionBindModel::GetActionTypeColor() const
+{
+	if (const URCPropertyBindAction* PropertyAction = Cast< URCPropertyBindAction>(ActionWeakPtr.Get()))
+	{
+		if (TSharedPtr<FRemoteControlProperty> RemoteControlProperty = PropertyAction->GetRemoteControlProperty())
+		{
+			return UE::RCUIHelpers::GetFieldClassTypeColor(RemoteControlProperty->GetProperty());
+		}
+	}
+
+	return FLinearColor::White;
+}
+
+#undef LOCTEXT_NAMESPACE
\ No newline at end of file
diff --git a/Engine/Plugins/VirtualProduction/RemoteControl/Source/RemoteControlUI/Private/UI/Action/Bind/RCActionBindModel.h b/Engine/Plugins/VirtualProduction/RemoteControl/Source/RemoteControlUI/Private/UI/Action/Bind/RCActionBindModel.h
new file mode 100644
--- /dev/null
+++ b/Engine/Plugins/VirtualProduction/RemoteControl/Source/RemoteControlUI/Private/UI/Action/Bind/RCActionBindModel.h
@@ -0,0 +1,51 @@
+// Copyright Epic Games, Inc. All Rights Reserved.
+
+#pragma once
+
+#include "Action/Bind/RCPropertyBindAction.h"
+#include "UI/Action/RCActionModel.h"
+
+class SActionItemListRow;
+class SBox;
+
+/* 
+* ~ FRCActionBindModel ~
+*
+* UI model for representing an Action of the Bind Behaviour.
+*/
+class FRCActionBindModel : public FRCActionModel
+{
+public:
+	FRCActionBindModel(URCAction* InAction, const TSharedPtr<class FRCBehaviourModel> InBehaviourItem, const TSharedPtr<SRemoteControlPanel> InRemoteControlPanel)
+		: FRCActionModel(InAction, InBehaviourItem, InRemoteControlPanel)
+	{
+	}
+
+	/** OnGenerateRow delegate for the Actions List View*/
+	TSharedRef<ITableRow> OnGenerateWidgetForList(TSharedPtr<FRCActionBindModel> InItem, const TSharedRef<STableViewBase>& OwnerTable);
+
+	/** Get the Header Row appropriate for this particular Action model */
+	static TSharedPtr<SHeaderRow> GetHeaderRow();
+
+	/** Chooses the appropriate Action model for the current class and field type*/
+	static TSharedPtr<FRCActionBindModel> GetModelByActionType(URCAction* InAction, const TSharedPtr<class FRCBehaviourModel> InBehaviourItem, const TSharedPtr<SRemoteControlPanel> InRemoteControlPanel);
+};
+
+/*
+* ~ FRCPropertyActionBindModel ~
+*
+* UI model for Property based Bind Actions
+*/
+class FRCPropertyActionBindModel : public FRCActionBindModel
+{
+public:
+	FRCPropertyActionBindModel(URCPropertyBindAction* InPropertyAction, const TSharedPtr<class FRCBehaviourModel> InBehaviourItem, const TSharedPtr<SRemoteControlPanel> InRemoteControlPanel)
+	: FRCActionBindModel(InPropertyAction, InBehaviourItem, InRemoteControlPanel)
+	{
+	}
+
+	/** Color code for this Action*/
+	virtual FLinearColor GetActionTypeColor() const override;
+};
+
+// Note: Bind only supports Property based Actions, so Function action model is not necessary here
\ No newline at end of file
diff --git a/Engine/Plugins/VirtualProduction/RemoteControl/Source/RemoteControlUI/Private/UI/Action/SRCActionPanel.cpp b/Engine/Plugins/VirtualProduction/RemoteControl/Source/RemoteControlUI/Private/UI/Action/SRCActionPanel.cpp
--- a/Engine/Plugins/VirtualProduction/RemoteControl/Source/RemoteControlUI/Private/UI/Action/SRCActionPanel.cpp
+++ b/Engine/Plugins/VirtualProduction/RemoteControl/Source/RemoteControlUI/Private/UI/Action/SRCActionPanel.cpp
@@ -65,6 +65,12 @@ void SRCActionPanel::Construct(const FArguments& InArgs, const TSharedRef<SRemot
 
 	// Register delegates
 	InPanel->OnBehaviourSelectionChanged.AddSP(this, &SRCActionPanel::OnBehaviourSelectionChanged);
+
+	if (URemoteControlPreset* Preset = GetPreset())
+	{
+		Preset->Layout.OnFieldAdded().AddSP(this, &SRCActionPanel::OnRemoteControlFieldAdded);
+		Preset->Layout.OnFieldDeleted().AddSP(this, &SRCActionPanel::OnRemoteControlFieldDeleted);
+	}
 }
 
 END_SLATE_FUNCTION_BUILD_OPTIMIZATION
@@ -94,6 +100,8 @@ void SRCActionPanel::UpdateWrappedWidget(TSharedPtr<FRCBehaviourModel> InBehavio
 			];
 
 		// Add New Action Button
+		bAddActionMenuNeedsRefresh = true;
+
 		const TSharedRef<SWidget> AddNewActionButton = SNew(SComboButton)
 			.AddMetaData<FTagMetaData>(FTagMetaData(TEXT("Add Action")))
 			.HAlign(HAlign_Center)
@@ -114,10 +122,7 @@ void SRCActionPanel::UpdateWrappedWidget(TSharedPtr<FRCBehaviourModel> InBehavio
 					.Image(FAppStyle::GetBrush("Icons.PlusCircle"))
 				]
 			]
-			.MenuContent()
-			[
-				GetActionMenuContentWidget()
-			];
+			.OnGetMenuContent(this, &SRCActionPanel::GetActionMenuContentWidget);
 		
 		// Add All Button
 		TSharedRef<SWidget> AddAllActionsButton = SNew(SButton)
@@ -238,6 +243,15 @@ FReply SRCActionPanel::OnClickOverrideBlueprintButton()
 
 TSharedRef<SWidget> SRCActionPanel::GetActionMenuContentWidget()
 {
+	if (AddNewActionMenuWidget && !bAddActionMenuNeedsRefresh)
+	{
+		return AddNewActionMenuWidget.ToSharedRef();
+	}
+
+	// Either we are creating the menu the first time, or the list needs to be refreshed
+	// The latter can occur either when a remote control property is added or removed, or when a behaiour's CanHaveActionForField logic has changed
+	// For example, if the user sets the "Allow numeric input as Strings" flag for Bind Behaviour then we need to recalculate the list of eligible actions.
+
 	constexpr bool bShouldCloseWindowAfterMenuSelection = true;
 	FMenuBuilder MenuBuilder(bShouldCloseWindowAfterMenuSelection, nullptr);
 
@@ -274,7 +288,11 @@ TSharedRef<SWidget> SRCActionPanel::GetActionMenuContentWidget()
 		}
 	}
 
-	return MenuBuilder.MakeWidget();
+	bAddActionMenuNeedsRefresh = false; // reset
+
+	AddNewActionMenuWidget = MenuBuilder.MakeWidget();
+
+	return AddNewActionMenuWidget.ToSharedRef();
 }
 
 URCAction* SRCActionPanel::AddAction(const TSharedRef<const FRemoteControlField> InRemoteControlField)
@@ -295,6 +313,25 @@ URCAction* SRCActionPanel::AddAction(const TSharedRef<const FRemoteControlField>
 	return nullptr;
 }
 
+bool SRCActionPanel::CanHaveActionForField(const FGuid& InRemoteControlFieldId)
+{
+	if (URemoteControlPreset* Preset = GetPreset())
+	{
+		if (TSharedPtr<FRemoteControlField> RemoteControlField = Preset->GetExposedEntity<FRemoteControlField>(InRemoteControlFieldId).Pin())
+		{
+			if (const TSharedPtr<FRCBehaviourModel> BehaviourItem = SelectedBehaviourItemWeakPtr.Pin())
+			{
+				if (const URCBehaviour* Behaviour = BehaviourItem->GetBehaviour())
+				{
+					return Behaviour->CanHaveActionForField(RemoteControlField);
+				}
+			}
+		}
+	}
+
+	return false;
+}
+
 void SRCActionPanel::OnAddActionClicked(TSharedPtr<FRemoteControlField> InRemoteControlField)
 {
 	if (!SelectedBehaviourItemWeakPtr.IsValid() || !InRemoteControlField)
@@ -341,7 +378,13 @@ FReply SRCActionPanel::OnAddAllFields()
 			{
 				if (const TSharedPtr<FRemoteControlField> RemoteControlField = RemoteControlFieldWeakPtr.Pin())
 				{
-					AddAction(RemoteControlField.ToSharedRef());
+					if (const URCBehaviour* Behaviour = BehaviourItem->GetBehaviour())
+					{
+						if (Behaviour->CanHaveActionForField(RemoteControlField))
+						{
+							AddAction(RemoteControlField.ToSharedRef());
+						}
+					}
 				}
 			}
 		}
@@ -350,6 +393,16 @@ FReply SRCActionPanel::OnAddAllFields()
 	return FReply::Handled();
 }
 
+void SRCActionPanel::OnRemoteControlFieldAdded(const FGuid& GroupId, const FGuid& FieldId, int32 FieldPosition)
+{
+	bAddActionMenuNeedsRefresh = true;
+}
+
+void SRCActionPanel::OnRemoteControlFieldDeleted(const FGuid& GroupId, const FGuid& FieldId, int32 FieldPosition)
+{
+	bAddActionMenuNeedsRefresh = true;
+}
+
 bool SRCActionPanel::IsListFocused() const
 {
 	return ActionPanelList.IsValid() && ActionPanelList->IsListFocused();
diff --git a/Engine/Plugins/VirtualProduction/RemoteControl/Source/RemoteControlUI/Private/UI/Action/SRCActionPanel.h b/Engine/Plugins/VirtualProduction/RemoteControl/Source/RemoteControlUI/Private/UI/Action/SRCActionPanel.h
--- a/Engine/Plugins/VirtualProduction/RemoteControl/Source/RemoteControlUI/Private/UI/Action/SRCActionPanel.h
+++ b/Engine/Plugins/VirtualProduction/RemoteControl/Source/RemoteControlUI/Private/UI/Action/SRCActionPanel.h
@@ -42,6 +42,14 @@ public:
 	/** Adds an Action for the currently active Behaviour and broadcasts to parent panels */
 	URCAction* AddAction(const TSharedRef<const FRemoteControlField> InRemoteControlField);
 
+	/** Whether this Actions panel can create an action for the given remote control field */
+	bool CanHaveActionForField(const FGuid& InRemoteControlFieldId);
+
+	void RequestRefreshForAddActionsMenu()
+	{
+		bAddActionMenuNeedsRefresh = true;
+	}
+
 protected:
 
 	/** Warns user before deleting all items in a panel. */
@@ -79,6 +87,12 @@ private:
 	/** Handles click event for Add All button; Adds all possible actions for the active Remote Control Preset*/
 	FReply OnAddAllFields();
 
+	/** Event invoked when a new remote control field has been added to the Remote Control Preset associated with this Action panel */
+	void OnRemoteControlFieldAdded(const FGuid& GroupId, const FGuid& FieldId, int32 FieldPosition);
+
+	/** Event invoked when a remote control field is removed from the Remote Control Preset associated with this Action panel */
+	void OnRemoteControlFieldDeleted(const FGuid& GroupId, const FGuid& FieldId, int32 FieldPosition);
+
 private:
 
 	/** The parent Behaviour that this Action panel is associated with */
@@ -95,4 +109,10 @@ private:
 
 	/** Panel Style reference. */
 	const FRCPanelStyle* RCPanelStyle;
-};
+
+	/** Cached menu widget for Add New Action */
+	TSharedPtr<SWidget> AddNewActionMenuWidget;
+
+	/** Whether the Add Actions menu list is outdated and needs to be refreshed */
+	bool bAddActionMenuNeedsRefresh = false;
+};
\ No newline at end of file
diff --git a/Engine/Plugins/VirtualProduction/RemoteControl/Source/RemoteControlUI/Private/UI/Action/SRCActionPanelList.cpp b/Engine/Plugins/VirtualProduction/RemoteControl/Source/RemoteControlUI/Private/UI/Action/SRCActionPanelList.cpp
--- a/Engine/Plugins/VirtualProduction/RemoteControl/Source/RemoteControlUI/Private/UI/Action/SRCActionPanelList.cpp
+++ b/Engine/Plugins/VirtualProduction/RemoteControl/Source/RemoteControlUI/Private/UI/Action/SRCActionPanelList.cpp
@@ -1,270 +1,3 @@
 ﻿// Copyright Epic Games, Inc. All Rights Reserved.
 
-#include "SRCActionPanelList.h"
-
-#include "Action/RCAction.h"
-#include "Action/RCActionContainer.h"
-#include "Action/RCFunctionAction.h"
-#include "Action/RCPropertyAction.h"
-#include "RCActionModel.h"
-#include "RemoteControlPreset.h"
-#include "SlateOptMacros.h"
-#include "SRCActionPanel.h"
-#include "SDropTarget.h"
-#include "Styling/RemoteControlStyles.h"
-#include "UI/Action/Conditional/RCActionConditionalModel.h"
-#include "UI/Behaviour/RCBehaviourModel.h"
-#include "UI/RCUIHelpers.h"
-#include "UI/RemoteControlPanelStyle.h"
-#include "UI/SRCPanelExposedField.h"
-#include "UI/SRCPanelFieldGroup.h"
-#include "UI/SRemoteControlPanel.h"
-#include "Widgets/Views/SListView.h"
-#include "Widgets/Views/STableRow.h"
-#include "Widgets/Views/STableViewBase.h"
-
-#define LOCTEXT_NAMESPACE "SRCActionPanelList"
-
-template <class ActionType>
-void SRCActionPanelList<ActionType>::Construct(const FArguments& InArgs, const TSharedRef<SRCActionPanel> InActionPanel, TSharedPtr<FRCBehaviourModel> InBehaviourItem)
-{
-	SRCLogicPanelListBase::Construct(SRCLogicPanelListBase::FArguments());
-
-	ActionPanelWeakPtr = InActionPanel;
-	BehaviourItemWeakPtr = InBehaviourItem;
-	
-	RCPanelStyle = &FRemoteControlPanelStyle::Get()->GetWidgetStyle<FRCPanelStyle>("RemoteControlPanel.MinorPanel");
-
-	ListView = SNew(SListView<TSharedPtr<ActionType>>)
-		.ListItemsSource(&ActionItems)
-		.OnGenerateRow(this, &SRCActionPanelList::OnGenerateWidgetForList)
-		.OnSelectionChanged(this, &SRCActionPanelList::OnSelectionChanged)
-		.ListViewStyle(&RCPanelStyle->TableViewStyle)
-		.OnContextMenuOpening(this, &SRCActionPanelList::GetContextMenuWidget)
-		.HeaderRow(ActionType::GetHeaderRow());
-
-	ChildSlot
-	[
-		SNew(SDropTarget)
-		.VerticalImage(FRemoteControlPanelStyle::Get()->GetBrush("RemoteControlPanel.VerticalDash"))
-		.HorizontalImage(FRemoteControlPanelStyle::Get()->GetBrush("RemoteControlPanel.HorizontalDash"))
-		.OnDropped_Lambda([this](const FGeometry& InGeometry, const FDragDropEvent& InDragDropEvent) { return SRCActionPanelList::OnExposedFieldDrop(InDragDropEvent.GetOperation()); })
-		.OnAllowDrop(this, &SRCActionPanelList::OnAllowDrop)
-		.OnIsRecognized(this, &SRCActionPanelList::OnAllowDrop)
-		[
-			ListView.ToSharedRef()
-		]
-	];
-
-	// Add delegates
-	const TSharedPtr<SRemoteControlPanel> RemoteControlPanel = InActionPanel->GetRemoteControlPanel();
-	if (ensure(RemoteControlPanel))
-	{
-		RemoteControlPanel->OnActionAdded.AddSP(this, &SRCActionPanelList::OnActionAdded);
-		RemoteControlPanel->OnEmptyActions.AddSP(this, &SRCActionPanelList::OnEmptyActions);
-	}
-
-	Reset();
-}
-
-template <class ActionType>
-bool SRCActionPanelList<ActionType>::IsEmpty() const
-{
-	return ActionItems.IsEmpty();
-}
-
-template <class ActionType>
-int32 SRCActionPanelList<ActionType>::Num() const
-{
-	return ActionItems.Num();
-}
-
-template <class ActionType>
-void SRCActionPanelList<ActionType>::Reset()
-{
-	if (TSharedPtr<SRCActionPanel> ActionPanel = GetActionPanel())
-	{
-		if (TSharedPtr<SRemoteControlPanel> RemoteControlPanel = ActionPanel->GetRemoteControlPanel())
-		{
-			ActionItems.Empty();
-
-			if (TSharedPtr<FRCBehaviourModel> BehaviourItem = BehaviourItemWeakPtr.Pin())
-			{
-				if (URCBehaviour* Behaviour = Cast<URCBehaviour>(BehaviourItem->GetBehaviour()))
-				{
-					for (URCAction* Action : Behaviour->ActionContainer->Actions)
-					{
-						TSharedPtr<ActionType> ActionItem = ActionType::GetModelByActionType(Action, BehaviourItem, RemoteControlPanel);
-
-						ActionItems.Add(ActionItem);
-					}
-				}
-			}
-
-			ListView->RebuildList();
-		}
-	}
-}
-
-template <class ActionType>
-URCAction* SRCActionPanelList<ActionType>::AddAction(const FGuid& InRemoteControlFieldId)
-{
-	if (const URemoteControlPreset* Preset = GetPreset())
-	{
-		if (TSharedPtr<const FRemoteControlField> RemoteControlField = Preset->GetExposedEntity<FRemoteControlField>(InRemoteControlFieldId).Pin())
-		{
-			if (const TSharedPtr<SRCActionPanel> ActionPanel = GetActionPanel())
-			{
-				return ActionPanel->AddAction(RemoteControlField.ToSharedRef());
-			}
-		}
-	}
-
-	return nullptr;
-}
-
-template <class ActionType>
-FReply SRCActionPanelList<ActionType>::OnExposedFieldDrop(TSharedPtr<FDragDropOperation> DragDropOperation)
-{
-	if (DragDropOperation)
-	{
-		if (DragDropOperation->IsOfType<FExposedEntityDragDrop>())
-		{
-			if (TSharedPtr<FExposedEntityDragDrop> DragDropOp = StaticCastSharedPtr<FExposedEntityDragDrop>(DragDropOperation))
-			{
-				// Fetch the Exposed Entity
-				const FGuid ExposedEntityId = DragDropOp->GetId();
-
-				// Add Action
-				AddAction(ExposedEntityId);
-			}
-		}
-		else if (DragDropOperation->IsOfType<FFieldGroupDragDropOp>())
-		{
-			if (TSharedPtr<FFieldGroupDragDropOp> DragDropOp = StaticCastSharedPtr<FFieldGroupDragDropOp>(DragDropOperation))
-			{
-				if (URemoteControlPreset* Preset = GetPreset())
-				{
-					// Fetch the Group
-					const FGuid GroupId = DragDropOp->GetGroupId();
-					const FRemoteControlPresetGroup* Group = Preset->Layout.GetGroup(GroupId);
-
-					if (ensure(Group))
-					{
-						const TArray<FGuid> GroupFields = Group->GetFields();
-
-						// Add Action for all fields in the Group
-						for (const FGuid RemoteControlFieldId : GroupFields)
-						{
-							AddAction(RemoteControlFieldId);
-						}
-					}
-				}
-			}
-		}
-	}
-
-	return FReply::Handled();
-}
-
-template <class ActionType>
-bool SRCActionPanelList<ActionType>::OnAllowDrop(TSharedPtr<FDragDropOperation> DragDropOperation)
-{
-	return true;
-}
-
-template <class ActionType>
-TSharedRef<ITableRow> SRCActionPanelList<ActionType>::OnGenerateWidgetForList(TSharedPtr<ActionType> InItem, const TSharedRef<STableViewBase>& OwnerTable)
-{
-	if (ensure(InItem))
-	{
-		return InItem->OnGenerateWidgetForList(InItem, OwnerTable);
-	}
-
-	return SNew(STableRow<TSharedPtr<FString>>, OwnerTable)
-		[
-			SNullWidget::NullWidget
-		];
-}
-
-template <class ActionType>
-void SRCActionPanelList<ActionType>::OnActionAdded(URCAction* InAction)
-{
-	Reset();
-}
-
-template <class ActionType>
-void SRCActionPanelList<ActionType>::OnEmptyActions()
-{
-	Reset();
-}
-
-template <class ActionType>
-URemoteControlPreset* SRCActionPanelList<ActionType>::GetPreset()
-{
-	if (ActionPanelWeakPtr.IsValid())
-	{
-		return ActionPanelWeakPtr.Pin()->GetPreset();
-	}
-
-	return nullptr;
-}
-
-template <class ActionType>
-int32 SRCActionPanelList<ActionType>::RemoveModel(const TSharedPtr<FRCLogicModeBase> InModel)
-{
-	if (const TSharedPtr<FRCBehaviourModel> BehaviourModel = BehaviourItemWeakPtr.Pin())
-	{
-		if (const URCBehaviour* Behaviour = BehaviourModel->GetBehaviour())
-		{
-			if(const TSharedPtr<ActionType> SelectedAction = StaticCastSharedPtr<ActionType>(InModel))
-			{
-				// Remove Model from Data Container
-				const int32 RemoveCount = Behaviour->ActionContainer->RemoveAction(SelectedAction->GetAction());
-
-				return RemoveCount;
-			}
-		}
-	}
-
-	return 0;
-}
-
-template <class ActionType>
-bool SRCActionPanelList<ActionType>::IsListFocused() const
-{
-	return ListView->HasAnyUserFocus().IsSet();
-}
-
-template <class ActionType>
-void SRCActionPanelList<ActionType>::DeleteSelectedPanelItem()
-{
-	DeleteItemFromLogicPanel<ActionType>(ActionItems, ListView->GetSelectedItems());
-}
-
-template <typename ActionType>
-TSharedPtr<SWidget> SRCActionPanelList<ActionType>::GetContextMenuWidget()
-{
-	if (SelectedActionItem)
-	{
-		return SelectedActionItem->GetContextMenuWidget();
-	}
-
-	return SNullWidget::NullWidget;
-}
-
-template <typename ActionType>
-void SRCActionPanelList<ActionType>::OnSelectionChanged(TSharedPtr<ActionType> InItem, ESelectInfo::Type)
-{
-	if (SelectedActionItem)
-	{
-		SelectedActionItem->OnSelectionExit();
-	}
-
-	SelectedActionItem = InItem;
-}
-
-#undef LOCTEXT_NAMESPACE
-
-template class SRCActionPanelList<FRCActionModel>;
-template class SRCActionPanelList<FRCActionConditionalModel>;
\ No newline at end of file
+#include "SRCActionPanelList.h"
\ No newline at end of file
diff --git a/Engine/Plugins/VirtualProduction/RemoteControl/Source/RemoteControlUI/Private/UI/Action/SRCActionPanelList.h b/Engine/Plugins/VirtualProduction/RemoteControl/Source/RemoteControlUI/Private/UI/Action/SRCActionPanelList.h
--- a/Engine/Plugins/VirtualProduction/RemoteControl/Source/RemoteControlUI/Private/UI/Action/SRCActionPanelList.h
+++ b/Engine/Plugins/VirtualProduction/RemoteControl/Source/RemoteControlUI/Private/UI/Action/SRCActionPanelList.h
@@ -2,17 +2,28 @@
 
 #pragma once
 
+#include "Action/RCAction.h"
+#include "Action/RCActionContainer.h"
+#include "Action/RCFunctionAction.h"
+#include "Action/RCPropertyAction.h"
+#include "RCActionModel.h"
+#include "RemoteControlPreset.h"
+#include "SlateOptMacros.h"
+#include "SRCActionPanel.h"
+#include "SDropTarget.h"
+#include "Styling/RemoteControlStyles.h"
+#include "UI/Action/Bind/RCActionBindModel.h"
+#include "UI/Action/Conditional/RCActionConditionalModel.h"
 #include "UI/BaseLogicUI/SRCLogicPanelListBase.h"
-
-struct FRCPanelStyle;
-class FRCBehaviourModel;
-class ActionType;
-class ITableRow;
-class STableViewBase;
-class SRCActionPanel;
-class URCAction;
-class URemoteControlPreset;
-template <typename ItemType> class SListView;
+#include "UI/Behaviour/RCBehaviourModel.h"
+#include "UI/RCUIHelpers.h"
+#include "UI/RemoteControlPanelStyle.h"
+#include "UI/SRCPanelExposedField.h"
+#include "UI/SRCPanelFieldGroup.h"
+#include "UI/SRemoteControlPanel.h"
+#include "Widgets/Views/SListView.h"
+#include "Widgets/Views/STableRow.h"
+#include "Widgets/Views/STableViewBase.h"
 
 /*
 * ~ SRCActionPanelList ~
@@ -31,19 +42,70 @@ public:
 	SLATE_END_ARGS()
 
 	/** Constructs this widget with InArgs */
-	void Construct(const FArguments& InArgs, const TSharedRef<SRCActionPanel> InActionPanel, TSharedPtr<FRCBehaviourModel> InBehaviourItem);
+	void Construct(const FArguments& InArgs, const TSharedRef<SRCActionPanel> InActionPanel, TSharedPtr<FRCBehaviourModel> InBehaviourItem)
+	{
+		SRCLogicPanelListBase::Construct(SRCLogicPanelListBase::FArguments());
+
+		ActionPanelWeakPtr = InActionPanel;
+		BehaviourItemWeakPtr = InBehaviourItem;
+
+		RCPanelStyle = &FRemoteControlPanelStyle::Get()->GetWidgetStyle<FRCPanelStyle>("RemoteControlPanel.MinorPanel");
+
+		ListView = SNew(SListView<TSharedPtr<ActionType>>)
+			.ListItemsSource(&ActionItems)
+			.OnGenerateRow(this, &SRCActionPanelList::OnGenerateWidgetForList)
+			.OnSelectionChanged(this, &SRCActionPanelList::OnSelectionChanged)
+			.ListViewStyle(&RCPanelStyle->TableViewStyle)
+			.OnContextMenuOpening(this, &SRCActionPanelList::GetContextMenuWidget)
+			.HeaderRow(ActionType::GetHeaderRow());
+
+		ChildSlot
+			[
+				SNew(SDropTarget)
+				.VerticalImage(FRemoteControlPanelStyle::Get()->GetBrush("RemoteControlPanel.VerticalDash"))
+				.HorizontalImage(FRemoteControlPanelStyle::Get()->GetBrush("RemoteControlPanel.HorizontalDash"))
+				.OnDropped_Lambda([this](const FGeometry& InGeometry, const FDragDropEvent& InDragDropEvent) { return SRCActionPanelList::OnExposedFieldDrop(InDragDropEvent.GetOperation()); })
+				.OnAllowDrop(this, &SRCActionPanelList::OnAllowDrop)
+				.OnIsRecognized(this, &SRCActionPanelList::OnAllowDrop)
+				[
+					ListView.ToSharedRef()
+				]
+			];
+
+		// Add delegates
+		const TSharedPtr<SRemoteControlPanel> RemoteControlPanel = InActionPanel->GetRemoteControlPanel();
+		if (ensure(RemoteControlPanel))
+		{
+			RemoteControlPanel->OnActionAdded.AddSP(this, &SRCActionPanelList::OnActionAdded);
+			RemoteControlPanel->OnEmptyActions.AddSP(this, &SRCActionPanelList::OnEmptyActions);
+		}
+
+		Reset();
+	}
 
 	/** Returns true if the underlying list is valid and empty. */
-	virtual bool IsEmpty() const override;
+	virtual bool IsEmpty() const override
+	{
+		return ActionItems.IsEmpty();
+	}
 
 	/** Returns number of items in the list. */
-	virtual int32 Num() const override;
+	virtual int32 Num() const override
+	{
+		return ActionItems.Num();
+	}
 
 	/** Whether the Actions List View currently has focus.*/
-	virtual bool IsListFocused() const override;
+	virtual bool IsListFocused() const override
+	{
+		return ListView->HasAnyUserFocus().IsSet();
+	}
 
 	/** Deletes currently selected items from the list view*/
-	virtual void DeleteSelectedPanelItem() override;
+	virtual void DeleteSelectedPanelItem() override
+	{
+		DeleteItemFromLogicPanel<ActionType>(ActionItems, ListView->GetSelectedItems());
+	}
 
 	/** Fetches the parent Action panel*/
 	TSharedPtr<SRCActionPanel> GetActionPanel()
@@ -58,42 +120,230 @@ public:
 	}
 
 	/** Adds an Action by Remote Control Field Guid*/
-	URCAction* AddAction(const FGuid& InRemoteControlFieldId);
+	URCAction* AddAction(const FGuid& InRemoteControlFieldId)
+	{
+		if (const URemoteControlPreset* Preset = GetPreset())
+		{
+			if (TSharedPtr<const FRemoteControlField> RemoteControlField = Preset->GetExposedEntity<FRemoteControlField>(InRemoteControlFieldId).Pin())
+			{
+				if (const TSharedPtr<SRCActionPanel> ActionPanel = GetActionPanel())
+				{
+					return ActionPanel->AddAction(RemoteControlField.ToSharedRef());
+				}
+			}
+		}
+
+		return nullptr;
+	}
 
 private:
 
 	/** OnGenerateRow delegate for the Actions List View*/
-	TSharedRef<ITableRow> OnGenerateWidgetForList( TSharedPtr<ActionType> InItem, const TSharedRef<STableViewBase>& OwnerTable );
+	TSharedRef<ITableRow> OnGenerateWidgetForList( TSharedPtr<ActionType> InItem, const TSharedRef<STableViewBase>& OwnerTable )
+	{
+		if (ensure(InItem))
+		{
+			return InItem->OnGenerateWidgetForList(InItem, OwnerTable);
+		}
+
+		return SNew(STableRow<TSharedPtr<FString>>, OwnerTable)
+			[
+				SNullWidget::NullWidget
+			];
+	}
 
 	/** Responds to the selection of a newly created action. Resets UI state*/
-	void OnActionAdded(URCAction* InAction);
+	void OnActionAdded(URCAction* InAction)
+	{
+		Reset();
+	}
 
 	/** Responds to the removal of all actions. Rests UI state*/
-	void OnEmptyActions();
+	void OnEmptyActions()
+	{
+		Reset();
+	}
 
 	/** Refreshes the list from the latest state of the data model*/
-	virtual void Reset() override;
+	virtual void Reset() override
+	{
+		if (TSharedPtr<SRCActionPanel> ActionPanel = GetActionPanel())
+		{
+			if (TSharedPtr<SRemoteControlPanel> RemoteControlPanel = ActionPanel->GetRemoteControlPanel())
+			{
+				ActionItems.Empty();
+
+				if (TSharedPtr<FRCBehaviourModel> BehaviourItem = BehaviourItemWeakPtr.Pin())
+				{
+					if (URCBehaviour* Behaviour = Cast<URCBehaviour>(BehaviourItem->GetBehaviour()))
+					{
+						for (URCAction* Action : Behaviour->ActionContainer->Actions)
+						{
+							TSharedPtr<ActionType> ActionItem = ActionType::GetModelByActionType(Action, BehaviourItem, RemoteControlPanel);
+
+							ActionItems.Add(ActionItem);
+						}
+					}
+				}
+
+				ListView->RebuildList();
+			}
+		}
+	}
 
 	/** Handles broadcasting of a successful remove item operation.*/
 	virtual void BroadcastOnItemRemoved() override {}
 
 	/** Fetches the Remote Control preset associated with the parent panel */
-	virtual URemoteControlPreset* GetPreset() override;
+	virtual URemoteControlPreset* GetPreset() override
+	{
+		if (ActionPanelWeakPtr.IsValid())
+		{
+			return ActionPanelWeakPtr.Pin()->GetPreset();
+		}
+
+		return nullptr;
+	}
 
 	/** Removes the given Action UI model item from the list of UI models*/
-	virtual int32 RemoveModel(const TSharedPtr<FRCLogicModeBase> InModel) override;
+	virtual int32 RemoveModel(const TSharedPtr<FRCLogicModeBase> InModel) override
+	{
+		if (const TSharedPtr<FRCBehaviourModel> BehaviourModel = BehaviourItemWeakPtr.Pin())
+		{
+			if (const URCBehaviour* Behaviour = BehaviourModel->GetBehaviour())
+			{
+				if (const TSharedPtr<ActionType> SelectedAction = StaticCastSharedPtr<ActionType>(InModel))
+				{
+					// Remove Model from Data Container
+					const int32 RemoveCount = Behaviour->ActionContainer->RemoveAction(SelectedAction->GetAction());
+
+					return RemoveCount;
+				}
+			}
+		}
+
+		return 0;
+	}
 
 	/*Drag and drop event for creating an Action from an exposed field */
-	FReply OnExposedFieldDrop(TSharedPtr<FDragDropOperation> DragDropOperation);
+	FReply OnExposedFieldDrop(TSharedPtr<FDragDropOperation> DragDropOperation)
+	{
+		if (DragDropOperation)
+		{
+			if (DragDropOperation->IsOfType<FExposedEntityDragDrop>())
+			{
+				if (TSharedPtr<FExposedEntityDragDrop> DragDropOp = StaticCastSharedPtr<FExposedEntityDragDrop>(DragDropOperation))
+				{
+					// Fetch the Exposed Entity
+					const FGuid ExposedEntityId = DragDropOp->GetId();
+
+					if (TSharedPtr<SRCActionPanel> ActionPanel = GetActionPanel())
+
+					// Add Action
+					AddAction(ExposedEntityId);
+				}
+			}
+			else if (DragDropOperation->IsOfType<FFieldGroupDragDropOp>())
+			{
+				if (TSharedPtr<FFieldGroupDragDropOp> DragDropOp = StaticCastSharedPtr<FFieldGroupDragDropOp>(DragDropOperation))
+				{
+					if (URemoteControlPreset* Preset = GetPreset())
+					{
+						// Fetch the Group
+						const FGuid GroupId = DragDropOp->GetGroupId();
+						const FRemoteControlPresetGroup* Group = Preset->Layout.GetGroup(GroupId);
+
+						if (ensure(Group))
+						{
+							const TArray<FGuid> GroupFields = Group->GetFields();
+
+							// Add Action for all fields in the Group
+							for (const FGuid RemoteControlFieldId : GroupFields)
+							{
+								AddAction(RemoteControlFieldId);
+							}
+						}
+					}
+				}
+			}
+		}
+
+		return FReply::Handled();
+	}
 
 	/** Whether drag and drop is possible from the current exposed property to the Actions table */
-	bool OnAllowDrop(TSharedPtr<FDragDropOperation> DragDropOperation);
+	bool OnAllowDrop(TSharedPtr<FDragDropOperation> DragDropOperation)
+	{
+		if (DragDropOperation)
+		{
+			if (DragDropOperation->IsOfType<FExposedEntityDragDrop>())
+			{
+				if (TSharedPtr<FExposedEntityDragDrop> DragDropOp = StaticCastSharedPtr<FExposedEntityDragDrop>(DragDropOperation))
+				{
+					// Fetch the Exposed Entity
+					const FGuid ExposedEntityId = DragDropOp->GetId();
+
+					if (TSharedPtr<SRCActionPanel> ActionPanel = GetActionPanel())
+					{
+						return ActionPanel->CanHaveActionForField(ExposedEntityId);
+					}
+				}
+			}
+			else if (DragDropOperation->IsOfType<FFieldGroupDragDropOp>())
+			{
+				if (TSharedPtr<FFieldGroupDragDropOp> DragDropOp = StaticCastSharedPtr<FFieldGroupDragDropOp>(DragDropOperation))
+				{
+					if (URemoteControlPreset* Preset = GetPreset())
+					{
+						// Fetch the Group
+						const FGuid GroupId = DragDropOp->GetGroupId();
+						const FRemoteControlPresetGroup* Group = Preset->Layout.GetGroup(GroupId);
+
+						if (ensure(Group))
+						{
+							const TArray<FGuid> GroupFields = Group->GetFields();
+
+							// For Groups we accept the drag-drop operation if even just one field in the group can have an action added successfully
+							for (const FGuid RemoteControlFieldId : GroupFields)
+							{
+								if (TSharedPtr<SRCActionPanel> ActionPanel = GetActionPanel())
+								{
+									if (ActionPanel->CanHaveActionForField(RemoteControlFieldId))
+									{
+										return true; // we have at least one compatible field
+									}
+								}
+							}
+						}
+					}
+				}
+			}
+		}
+
+		return false;
+	}
 
 	/** Context menu for Actions panel list */
-	TSharedPtr<SWidget> GetContextMenuWidget();
+	TSharedPtr<SWidget> GetContextMenuWidget()
+	{
+		if (SelectedActionItem)
+		{
+			return SelectedActionItem->GetContextMenuWidget();
+		}
+
+		return SNullWidget::NullWidget;
+	}
 
 	/** OnSelectionChanged delegate for Actions List View */
-	void OnSelectionChanged(TSharedPtr<ActionType> InItem, ESelectInfo::Type);
+	void OnSelectionChanged(TSharedPtr<ActionType> InItem, ESelectInfo::Type)
+	{
+		if (SelectedActionItem)
+		{
+			SelectedActionItem->OnSelectionExit();
+		}
+
+		SelectedActionItem = InItem;
+	}
 
 private:
 
diff --git a/Engine/Plugins/VirtualProduction/RemoteControl/Source/RemoteControlUI/Private/UI/BaseLogicUI/RCLogicModeBase.cpp b/Engine/Plugins/VirtualProduction/RemoteControl/Source/RemoteControlUI/Private/UI/BaseLogicUI/RCLogicModeBase.cpp
--- a/Engine/Plugins/VirtualProduction/RemoteControl/Source/RemoteControlUI/Private/UI/BaseLogicUI/RCLogicModeBase.cpp
+++ b/Engine/Plugins/VirtualProduction/RemoteControl/Source/RemoteControlUI/Private/UI/BaseLogicUI/RCLogicModeBase.cpp
@@ -24,4 +24,9 @@ URemoteControlPreset* FRCLogicModeBase::GetPreset() const
 	return PanelWeakPtr.Pin()->GetPreset();
 }
 
+TSharedPtr< SRemoteControlPanel> FRCLogicModeBase::GetRemoteControlPanel() const
+{
+	return PanelWeakPtr.Pin();
+}
+
 #undef LOCTEXT_NAMESPACE
\ No newline at end of file
diff --git a/Engine/Plugins/VirtualProduction/RemoteControl/Source/RemoteControlUI/Private/UI/BaseLogicUI/RCLogicModeBase.h b/Engine/Plugins/VirtualProduction/RemoteControl/Source/RemoteControlUI/Private/UI/BaseLogicUI/RCLogicModeBase.h
--- a/Engine/Plugins/VirtualProduction/RemoteControl/Source/RemoteControlUI/Private/UI/BaseLogicUI/RCLogicModeBase.h
+++ b/Engine/Plugins/VirtualProduction/RemoteControl/Source/RemoteControlUI/Private/UI/BaseLogicUI/RCLogicModeBase.h
@@ -31,6 +31,9 @@ public:
 	/** Returns the Remote Control Preset object associated with us*/
 	URemoteControlPreset* GetPreset() const;
 
+	/** Returns the Remote Control Preset object associated with us*/
+	TSharedPtr< SRemoteControlPanel> GetRemoteControlPanel() const;
+
 protected:
 	/** The parent Remote Control Panel widget*/
 	TWeakPtr<SRemoteControlPanel> PanelWeakPtr;
diff --git a/Engine/Plugins/VirtualProduction/RemoteControl/Source/RemoteControlUI/Private/UI/BaseLogicUI/SRCLogicPanelListBase.cpp b/Engine/Plugins/VirtualProduction/RemoteControl/Source/RemoteControlUI/Private/UI/BaseLogicUI/SRCLogicPanelListBase.cpp
--- a/Engine/Plugins/VirtualProduction/RemoteControl/Source/RemoteControlUI/Private/UI/BaseLogicUI/SRCLogicPanelListBase.cpp
+++ b/Engine/Plugins/VirtualProduction/RemoteControl/Source/RemoteControlUI/Private/UI/BaseLogicUI/SRCLogicPanelListBase.cpp
@@ -2,13 +2,6 @@
 
 #include "SRCLogicPanelListBase.h"
 
-#include "RCLogicModeBase.h"
-#include "SlateOptMacros.h"
-#include "UI/Action/Conditional/RCActionConditionalModel.h"
-#include "UI/Action/RCActionModel.h"
-#include "UI/Behaviour/RCBehaviourModel.h"
-#include "UI/Controller/RCControllerModel.h"
-
 BEGIN_SLATE_FUNCTION_BUILD_OPTIMIZATION
 
 void SRCLogicPanelListBase::Construct(const FArguments& InArgs)
@@ -16,43 +9,3 @@ void SRCLogicPanelListBase::Construct(const FArguments& InArgs)
 }
 
 END_SLATE_FUNCTION_BUILD_OPTIMIZATION
-
-template<class T>
-void SRCLogicPanelListBase::DeleteItemFromLogicPanel(TArray<TSharedPtr<T>>& ItemsSource, const TArray<TSharedPtr<T>>& SelectedItems)
-{
-	bool bIsDeleted = false;
-	for (const TSharedPtr<T> SelectedItem : SelectedItems)
-	{
-		if (SelectedItem.IsValid())
-		{
-			// Remove Model from Data Container
-			const int32 RemoveCount = RemoveModel(SelectedItem);
-			if(ensure(RemoveCount > 0))
-			{
-				// Remove View Model from UI List
-				const int32 RemoveModelItemIndex = ItemsSource.IndexOfByPredicate([SelectedItem](TSharedPtr<T> InModel)
-					{
-						return SelectedItem == InModel;
-					});
-
-				if (RemoveModelItemIndex > INDEX_NONE)
-				{
-					ItemsSource.RemoveAt(RemoveModelItemIndex);
-
-					bIsDeleted = true;
-				}
-			}
-		}
-	}
-
-	if (bIsDeleted)
-	{
-		BroadcastOnItemRemoved();
-		Reset();
-	}
-}
-
-template void SRCLogicPanelListBase::DeleteItemFromLogicPanel(TArray<TSharedPtr<FRCControllerModel>>& ItemsSource, const TArray<TSharedPtr<FRCControllerModel>>& SelectedItems);
-template void SRCLogicPanelListBase::DeleteItemFromLogicPanel(TArray<TSharedPtr<FRCBehaviourModel>>& ItemsSource, const TArray<TSharedPtr<FRCBehaviourModel>>& SelectedItems);
-template void SRCLogicPanelListBase::DeleteItemFromLogicPanel(TArray<TSharedPtr<FRCActionModel>>& ItemsSource, const TArray<TSharedPtr<FRCActionModel>>& SelectedItems);
-template void SRCLogicPanelListBase::DeleteItemFromLogicPanel(TArray<TSharedPtr<FRCActionConditionalModel>>& ItemsSource, const TArray<TSharedPtr<FRCActionConditionalModel>>& SelectedItems);
\ No newline at end of file
diff --git a/Engine/Plugins/VirtualProduction/RemoteControl/Source/RemoteControlUI/Private/UI/BaseLogicUI/SRCLogicPanelListBase.h b/Engine/Plugins/VirtualProduction/RemoteControl/Source/RemoteControlUI/Private/UI/BaseLogicUI/SRCLogicPanelListBase.h
--- a/Engine/Plugins/VirtualProduction/RemoteControl/Source/RemoteControlUI/Private/UI/BaseLogicUI/SRCLogicPanelListBase.h
+++ b/Engine/Plugins/VirtualProduction/RemoteControl/Source/RemoteControlUI/Private/UI/BaseLogicUI/SRCLogicPanelListBase.h
@@ -3,11 +3,15 @@
 #pragma once
 
 #include "CoreMinimal.h"
+#include "RCLogicModeBase.h"
+#include "SlateOptMacros.h"
+#include "UI/Action/Bind/RCActionBindModel.h"
+#include "UI/Action/Conditional/RCActionConditionalModel.h"
+#include "UI/Action/RCActionModel.h"
+#include "UI/Behaviour/RCBehaviourModel.h"
+#include "UI/Controller/RCControllerModel.h"
 #include "Widgets/SCompoundWidget.h"
 
-class FRCLogicModeBase;
-class URemoteControlPreset;
-
 /*
 * ~ SRCLogicPanelListBase ~
 *
@@ -58,5 +62,37 @@ protected:
 	/** Helper function for handling common Delete Item functionality across all child panels (Actions/Behaviours/Controllers)
 	* Currently invoked from each Panel List child class with appropriate template class*/
 	template<class T>
-	void DeleteItemFromLogicPanel(TArray<TSharedPtr<T>>& ItemsSource, const TArray<TSharedPtr<T>>& SelectedItems);
+	void DeleteItemFromLogicPanel(TArray<TSharedPtr<T>>& ItemsSource, const TArray<TSharedPtr<T>>& SelectedItems)
+	{
+		bool bIsDeleted = false;
+		for (const TSharedPtr<T> SelectedItem : SelectedItems)
+		{
+			if (SelectedItem.IsValid())
+			{
+				// Remove Model from Data Container
+				const int32 RemoveCount = RemoveModel(SelectedItem);
+				if (ensure(RemoveCount > 0))
+				{
+					// Remove View Model from UI List
+					const int32 RemoveModelItemIndex = ItemsSource.IndexOfByPredicate([SelectedItem](TSharedPtr<T> InModel)
+						{
+							return SelectedItem == InModel;
+						});
+
+					if (RemoveModelItemIndex > INDEX_NONE)
+					{
+						ItemsSource.RemoveAt(RemoveModelItemIndex);
+
+						bIsDeleted = true;
+					}
+				}
+			}
+		}
+
+		if (bIsDeleted)
+		{
+			BroadcastOnItemRemoved();
+			Reset();
+		}
+	}
 };
\ No newline at end of file
diff --git a/Engine/Plugins/VirtualProduction/RemoteControl/Source/RemoteControlUI/Private/UI/Behaviour/Builtin/Bind/RCBehaviourBindModel.cpp b/Engine/Plugins/VirtualProduction/RemoteControl/Source/RemoteControlUI/Private/UI/Behaviour/Builtin/Bind/RCBehaviourBindModel.cpp
new file mode 100644
--- /dev/null
+++ b/Engine/Plugins/VirtualProduction/RemoteControl/Source/RemoteControlUI/Private/UI/Behaviour/Builtin/Bind/RCBehaviourBindModel.cpp
@@ -0,0 +1,34 @@
+// Copyright Epic Games, Inc. All Rights Reserved.
+
+#include "RCBehaviourBindModel.h"
+
+#include "Behaviour/Builtin/Bind/RCBehaviourBind.h"
+#include "SRCBehaviourBind.h"
+#include "UI/Action/Bind/RCActionBindModel.h"
+#include "UI/Action/SRCActionPanel.h"
+#include "UI/Action/SRCActionPanelList.h"
+
+#define LOCTEXT_NAMESPACE "FRCBehaviourBindModel"
+
+FRCBehaviourBindModel::FRCBehaviourBindModel(URCBehaviourBind* BindBehaviour, const TSharedPtr<SRemoteControlPanel> InRemoteControlPanel)
+	: FRCBehaviourModel(BindBehaviour, InRemoteControlPanel)
+	, BindBehaviourWeakPtr(BindBehaviour)
+{
+}
+
+TSharedRef<SWidget> FRCBehaviourBindModel::GetBehaviourDetailsWidget()
+{
+	return SAssignNew(BehaviourDetailsWidget, SRCBehaviourBind, SharedThis(this));
+}
+
+TSharedPtr<SRCLogicPanelListBase> FRCBehaviourBindModel::GetActionsListWidget(TSharedRef<SRCActionPanel> InActionPanel)
+{
+	return SNew(SRCActionPanelList<FRCActionBindModel>, InActionPanel, SharedThis(this));
+}
+
+URCBehaviourBind* FRCBehaviourBindModel::GetBindBehaviour() const
+{
+	return Cast<URCBehaviourBind>(GetBehaviour());
+}
+
+#undef LOCTEXT_NAMESPACE
\ No newline at end of file
diff --git a/Engine/Plugins/VirtualProduction/RemoteControl/Source/RemoteControlUI/Private/UI/Behaviour/Builtin/Bind/RCBehaviourBindModel.h b/Engine/Plugins/VirtualProduction/RemoteControl/Source/RemoteControlUI/Private/UI/Behaviour/Builtin/Bind/RCBehaviourBindModel.h
new file mode 100644
--- /dev/null
+++ b/Engine/Plugins/VirtualProduction/RemoteControl/Source/RemoteControlUI/Private/UI/Behaviour/Builtin/Bind/RCBehaviourBindModel.h
@@ -0,0 +1,39 @@
+// Copyright Epic Games, Inc. All Rights Reserved.
+
+#pragma once
+#include "UI/Behaviour/RCBehaviourModel.h"
+
+class SRCActionPanel;
+class SRCBehaviourBind;
+class SRCLogicPanelListBase;
+class SWidget;
+class URCBehaviourBind;
+
+/*
+* ~ FRCBehaviourBindModel ~
+*
+* UI model for Bind Behaviour
+
+*/
+class FRCBehaviourBindModel : public FRCBehaviourModel
+{
+public:
+	FRCBehaviourBindModel(URCBehaviourBind* BindBehaviour, const TSharedPtr<SRemoteControlPanel> InRemoteControlPanel);
+
+	/** Builds a Behaviour specific details widget for Bind Behaviour*/
+	virtual TSharedRef<SWidget> GetBehaviourDetailsWidget() override;
+
+	/** The Actions List widget to be used for Bind Behaviour */
+	virtual TSharedPtr<SRCLogicPanelListBase> GetActionsListWidget(TSharedRef<SRCActionPanel> InActionPanel);
+
+	/** Fetches the Bind Behaviour UObject model*/
+	URCBehaviourBind* GetBindBehaviour() const;
+
+private:
+
+	/** The Bind Behaviour (Data model) associated with us*/
+	TWeakObjectPtr<URCBehaviourBind> BindBehaviourWeakPtr;
+
+	/** Behaviour Details Widget for Bind Behaviour */
+	TSharedPtr<SRCBehaviourBind> BehaviourDetailsWidget;
+};
\ No newline at end of file
diff --git a/Engine/Plugins/VirtualProduction/RemoteControl/Source/RemoteControlUI/Private/UI/Behaviour/Builtin/Bind/SRCBehaviourBind.cpp b/Engine/Plugins/VirtualProduction/RemoteControl/Source/RemoteControlUI/Private/UI/Behaviour/Builtin/Bind/SRCBehaviourBind.cpp
new file mode 100644
--- /dev/null
+++ b/Engine/Plugins/VirtualProduction/RemoteControl/Source/RemoteControlUI/Private/UI/Behaviour/Builtin/Bind/SRCBehaviourBind.cpp
@@ -0,0 +1,70 @@
+// Copyright Epic Games, Inc. All Rights Reserved.
+
+#include "SRCBehaviourBind.h"
+
+#include "Behaviour/Builtin/Bind/RCBehaviourBind.h"
+#include "UI/Behaviour/Builtin/Bind/RCBehaviourBindModel.h"
+#include "UI/SRemoteControlPanel.h"
+#include "UI/Action/SRCActionPanel.h"
+#include "Widgets/Input/SCheckBox.h"
+
+#define LOCTEXT_NAMESPACE "SRCBehaviourBind"
+
+void SRCBehaviourBind::Construct(const FArguments& InArgs, TSharedRef<FRCBehaviourBindModel> InBehaviourItem)
+{
+	BindBehaviourItemWeakPtr = InBehaviourItem;
+
+	ChildSlot
+	[
+		SNew(SHorizontalBox)
+		// Label
+		+ SHorizontalBox::Slot()
+		.AutoWidth()
+		.VAlign(EVerticalAlignment::VAlign_Center)
+		[
+			SNew(STextBlock)
+			.Text(LOCTEXT("AllowNumericInputLabel", "Allow numeric inputs as strings"))
+		]
+
+		// Checkbox
+		+ SHorizontalBox::Slot()
+		.AutoWidth()
+		.HAlign(EHorizontalAlignment::HAlign_Fill)
+		.VAlign(EVerticalAlignment::VAlign_Center)
+		.Padding(FMargin(10.f))
+		[
+			SAssignNew(CheckboxAllowNumericInput, SCheckBox)
+			.IsChecked(false)
+			.OnCheckStateChanged(this, &SRCBehaviourBind::OnAllowNumericCheckboxChanged)
+		]
+	];
+
+	if (const URCBehaviourBind* BindBehaviour = InBehaviourItem->GetBindBehaviour())
+	{
+		CheckboxAllowNumericInput->SetIsChecked(BindBehaviour->AreNumericInputsAllowedAsStrings());
+	}
+}
+
+void SRCBehaviourBind::OnAllowNumericCheckboxChanged(ECheckBoxState NewState)
+{
+	if (TSharedPtr<const FRCBehaviourBindModel> BehaviourItem = BindBehaviourItemWeakPtr.Pin())
+	{
+		if (URCBehaviourBind* BindBehaviour = BehaviourItem->GetBindBehaviour())
+		{
+			// Update flag
+			BindBehaviour->SetAllowNumericInputAsStrings((NewState == ECheckBoxState::Checked));
+
+			// Inform the Action panel that it needs to refresh the Add Action menu list
+			if (TSharedPtr<SRemoteControlPanel> RemoteControlPanel = BehaviourItem->GetRemoteControlPanel())
+			{
+				if (TSharedPtr<SRCActionPanel> ActionPanel = RemoteControlPanel->GetLogicActionPanel())
+				{
+					ActionPanel->RequestRefreshForAddActionsMenu();
+				}
+			}
+		}
+	}
+	
+}
+
+#undef LOCTEXT_NAMESPACE
\ No newline at end of file
diff --git a/Engine/Plugins/VirtualProduction/RemoteControl/Source/RemoteControlUI/Private/UI/Behaviour/Builtin/Bind/SRCBehaviourBind.h b/Engine/Plugins/VirtualProduction/RemoteControl/Source/RemoteControlUI/Private/UI/Behaviour/Builtin/Bind/SRCBehaviourBind.h
new file mode 100644
--- /dev/null
+++ b/Engine/Plugins/VirtualProduction/RemoteControl/Source/RemoteControlUI/Private/UI/Behaviour/Builtin/Bind/SRCBehaviourBind.h
@@ -0,0 +1,42 @@
+// Copyright Epic Games, Inc. All Rights Reserved.
+
+#pragma once
+
+#include "CoreMinimal.h"
+#include "Widgets/SCompoundWidget.h"
+
+enum class ECheckBoxState : uint8;
+class FRCBehaviourBindModel;
+class SCheckBox;
+class URCBehaviourBind;
+
+/*
+* ~ SRCBehaviourConditional ~
+*
+* Behaviour specific details panel for Bind Behaviour
+* 
+* Contains a checkbox which specifies whether string to numeric (and vice versa) conversions are desired
+*/
+class REMOTECONTROLUI_API SRCBehaviourBind : public SCompoundWidget
+{
+public:
+	SLATE_BEGIN_ARGS(SRCBehaviourBind)
+		{
+		}
+
+	SLATE_END_ARGS()
+
+	/** Constructs this widget with InArgs */
+	void Construct(const FArguments& InArgs, TSharedRef<FRCBehaviourBindModel> InBehaviourItem);
+
+private:
+
+	/** Called when the checkbox state changes for "Allow Numeric Input As Strings" */
+	void OnAllowNumericCheckboxChanged(ECheckBoxState NewState);
+
+	/** The Behaviour (UI model) associated with us*/
+	TWeakPtr<FRCBehaviourBindModel> BindBehaviourItemWeakPtr;
+
+	/** Checkbox for "Allow Numeric Input as Strings" */
+	TSharedPtr<SCheckBox> CheckboxAllowNumericInput;
+};
diff --git a/Engine/Plugins/VirtualProduction/RemoteControl/Source/RemoteControlUI/Private/UI/Behaviour/Builtin/Conditional/RCBehaviourConditionalModel.cpp b/Engine/Plugins/VirtualProduction/RemoteControl/Source/RemoteControlUI/Private/UI/Behaviour/Builtin/Conditional/RCBehaviourConditionalModel.cpp
--- a/Engine/Plugins/VirtualProduction/RemoteControl/Source/RemoteControlUI/Private/UI/Behaviour/Builtin/Conditional/RCBehaviourConditionalModel.cpp
+++ b/Engine/Plugins/VirtualProduction/RemoteControl/Source/RemoteControlUI/Private/UI/Behaviour/Builtin/Conditional/RCBehaviourConditionalModel.cpp
@@ -13,6 +13,7 @@
 #include "UI/Action/Conditional/RCActionConditionalModel.h"
 #include "UI/Action/SRCActionPanel.h"
 #include "UI/Action/SRCActionPanelList.h"
+#include "UI/RCUIHelpers.h"
 #include "Widgets/Layout/SBox.h"
 
 #define LOCTEXT_NAMESPACE "FRCBehaviourConditionalModel"
@@ -21,10 +22,6 @@ FRCBehaviourConditionalModel::FRCBehaviourConditionalModel(URCBehaviourCondition
 	: FRCBehaviourModel(ConditionalBehaviour)
 	, ConditionalBehaviourWeakPtr(ConditionalBehaviour)
 {
-	FPropertyRowGeneratorArgs Args;
-	Args.bShouldShowHiddenProperties = true;
-	PropertyRowGenerator = FModuleManager::GetModuleChecked<FPropertyEditorModule>("PropertyEditor").CreatePropertyRowGenerator(Args);
-
 	CreateComparandInputField();
 }
 
@@ -39,7 +36,7 @@ URCAction* FRCBehaviourConditionalModel::AddAction(const TSharedRef<const FRemot
 
 	if (URCBehaviourConditional* ConditionalBehaviour = GetConditionalBehaviour())
 	{
-		NewAction = ConditionalBehaviour->AddAction(InRemoteControlField, Condition, ConditionalBehaviour->Comparand);
+		NewAction = ConditionalBehaviour->AddConditionalAction(InRemoteControlField, Condition, ConditionalBehaviour->Comparand);
 
 		OnActionAdded(NewAction);
 	}
@@ -70,17 +67,7 @@ void FRCBehaviourConditionalModel::CreateComparandInputField()
 			}
 
 			// UI widget (via Property Generator)
-			PropertyRowGenerator->SetStructure(Comparand->CreateStructOnScope());
-			for (const TSharedRef<IDetailTreeNode>& CategoryNode : PropertyRowGenerator->GetRootTreeNodes())
-			{
-				TArray<TSharedRef<IDetailTreeNode>> Children;
-				CategoryNode->GetChildren(Children);
-				for (const TSharedRef<IDetailTreeNode>& Child : Children)
-				{
-					DetailTreeNodeWeakPtr = Child;
-					break;
-				}
-			}
+			DetailTreeNodeWeakPtr = UE::RCUIHelpers::GetDetailTreeNodeForVirtualProperty(Comparand, PropertyRowGenerator);
 		}
 	}
 }
@@ -107,39 +94,7 @@ TSharedRef<SWidget> FRCBehaviourConditionalModel::GetComparandFieldWidget() cons
 {
 	if (const TSharedPtr<IDetailTreeNode> DetailTreeNode = DetailTreeNodeWeakPtr.Pin())
 	{
-		const FNodeWidgets NodeWidgets = DetailTreeNode->CreateNodeWidgets();
-
-		const TSharedRef<SHorizontalBox> FieldWidget = SNew(SHorizontalBox);
-
-		if (NodeWidgets.ValueWidget)
-		{
-			FieldWidget->AddSlot()
-				.Padding(FMargin(3.0f, 2.0f))
-				.HAlign(HAlign_Right)
-				.AutoWidth()
-				[
-					NodeWidgets.ValueWidget.ToSharedRef()
-				];
-		}
-		else if (NodeWidgets.WholeRowWidget)
-		{
-			FieldWidget->AddSlot()
-				.Padding(FMargin(3.0f, 2.0f))
-				.AutoWidth()
-				[
-					NodeWidgets.WholeRowWidget.ToSharedRef()
-				];
-		}
-
-		return SNew(SHorizontalBox)
-			.Clipping(EWidgetClipping::OnDemand)
-			// Property Value Widget
-			+ SHorizontalBox::Slot()
-			.VAlign(VAlign_Center)
-			.AutoWidth()
-			[
-				FieldWidget
-			];
+		return UE::RCUIHelpers::GetGenericFieldWidget(DetailTreeNode);
 	}
 
 	return SNullWidget::NullWidget;
diff --git a/Engine/Plugins/VirtualProduction/RemoteControl/Source/RemoteControlUI/Private/UI/Behaviour/RCBehaviourModel.cpp b/Engine/Plugins/VirtualProduction/RemoteControl/Source/RemoteControlUI/Private/UI/Behaviour/RCBehaviourModel.cpp
--- a/Engine/Plugins/VirtualProduction/RemoteControl/Source/RemoteControlUI/Private/UI/Behaviour/RCBehaviourModel.cpp
+++ b/Engine/Plugins/VirtualProduction/RemoteControl/Source/RemoteControlUI/Private/UI/Behaviour/RCBehaviourModel.cpp
@@ -16,8 +16,10 @@
 
 #define LOCTEXT_NAMESPACE "FRCBehaviourModel"
 
-FRCBehaviourModel::FRCBehaviourModel(URCBehaviour* InBehaviour)
-	: BehaviourWeakPtr(InBehaviour)
+FRCBehaviourModel::FRCBehaviourModel(URCBehaviour* InBehaviour,
+	const TSharedPtr<SRemoteControlPanel> InRemoteControlPanel /*= nullptr*/)
+	: FRCLogicModeBase(InRemoteControlPanel),
+	BehaviourWeakPtr(InBehaviour)
 {
 	RCPanelStyle = &FRemoteControlPanelStyle::Get()->GetWidgetStyle<FRCPanelStyle>("RemoteControlPanel.BehaviourPanel");
 
diff --git a/Engine/Plugins/VirtualProduction/RemoteControl/Source/RemoteControlUI/Private/UI/Behaviour/RCBehaviourModel.h b/Engine/Plugins/VirtualProduction/RemoteControl/Source/RemoteControlUI/Private/UI/Behaviour/RCBehaviourModel.h
--- a/Engine/Plugins/VirtualProduction/RemoteControl/Source/RemoteControlUI/Private/UI/Behaviour/RCBehaviourModel.h
+++ b/Engine/Plugins/VirtualProduction/RemoteControl/Source/RemoteControlUI/Private/UI/Behaviour/RCBehaviourModel.h
@@ -21,7 +21,7 @@ class URCBehaviour;
 class FRCBehaviourModel : public FRCLogicModeBase
 {
 public:
-	FRCBehaviourModel(URCBehaviour* InBehaviour);
+	FRCBehaviourModel(URCBehaviour* InBehaviour, const TSharedPtr<SRemoteControlPanel> InRemoteControlPanel = nullptr);
 
 	/** Add a Logic Action using a remote control field as input */
 	virtual URCAction* AddAction(const TSharedRef<const FRemoteControlField> InRemoteControlField);
diff --git a/Engine/Plugins/VirtualProduction/RemoteControl/Source/RemoteControlUI/Private/UI/Behaviour/SRCBehaviourPanelList.cpp b/Engine/Plugins/VirtualProduction/RemoteControl/Source/RemoteControlUI/Private/UI/Behaviour/SRCBehaviourPanelList.cpp
--- a/Engine/Plugins/VirtualProduction/RemoteControl/Source/RemoteControlUI/Private/UI/Behaviour/SRCBehaviourPanelList.cpp
+++ b/Engine/Plugins/VirtualProduction/RemoteControl/Source/RemoteControlUI/Private/UI/Behaviour/SRCBehaviourPanelList.cpp
@@ -2,6 +2,7 @@
 
 #include "SRCBehaviourPanelList.h"
 
+#include "Behaviour/Builtin/Bind/RCBehaviourBind.h"
 #include "Behaviour/Builtin/Conditional/RCBehaviourConditional.h"
 #include "Behaviour/RCSetAssetByPathBehaviour.h"
 #include "Controller/RCController.h"
@@ -10,6 +11,7 @@
 #include "SRCBehaviourPanel.h"
 #include "SlateOptMacros.h"
 #include "Styling/RemoteControlStyles.h"
+#include "UI/Behaviour/Builtin/Bind/RCBehaviourBindModel.h"
 #include "UI/Behaviour/Builtin/Conditional/RCBehaviourConditionalModel.h"
 #include "UI/Behaviour/Builtin/RCBehaviourIsEqualModel.h"
 #include "UI/Behaviour/Builtin/RCBehaviourSetAssetByPathModel.h"
@@ -81,6 +83,8 @@ void SRCBehaviourPanelList::Reset()
 {
 	BehaviourItems.Empty();
 
+	const TSharedPtr<SRemoteControlPanel> RemoteControlPanel = BehaviourPanelWeakPtr.IsValid() ? BehaviourPanelWeakPtr.Pin()->GetRemoteControlPanel() : nullptr;
+
 	if (TSharedPtr<FRCControllerModel> ControllerItem = ControllerItemWeakPtr.Pin())
 	{
 		if (URCController* Controller = Cast<URCController>(ControllerItem->GetVirtualProperty()))
@@ -99,9 +103,13 @@ void SRCBehaviourPanelList::Reset()
 				{
 					BehaviourItems.Add(MakeShared<FRCSetAssetByPathBehaviourModel>(SetAssetByPathBehaviour));
 				}
+				else if (URCBehaviourBind* BindBehaviour = Cast<URCBehaviourBind>(Behaviour))
+				{
+					BehaviourItems.Add(MakeShared<FRCBehaviourBindModel>(BindBehaviour, RemoteControlPanel));
+				}
 				else
 				{
-					BehaviourItems.Add(MakeShared<FRCBehaviourModel>(Behaviour));
+					BehaviourItems.Add(MakeShared<FRCBehaviourModel>(Behaviour, RemoteControlPanel));
 				}
 			}
 		}
