- Fixing static analysis in media capture and remote control

diff --git a/Engine/Plugins/Media/MediaIOFramework/Source/MediaIOCore/Private/MediaCapture.cpp b/Engine/Plugins/Media/MediaIOFramework/Source/MediaIOCore/Private/MediaCapture.cpp
--- a/Engine/Plugins/Media/MediaIOFramework/Source/MediaIOCore/Private/MediaCapture.cpp
+++ b/Engine/Plugins/Media/MediaIOFramework/Source/MediaIOCore/Private/MediaCapture.cpp
@@ -1693,7 +1693,7 @@ void UMediaCapture::OnEndFrame_GameThread()
 
 					FTexture2DRHIRef SourceTexture = FMediaCaptureHelper::GetSourceTextureForInput(RHICmdList, InCapturingSceneViewport, InTextureRenderTargetResource);
 					InMediaCapture->ProcessCapture_RenderThread(GraphBuilder, InMediaCapture, CapturingFrame, MoveTemp(SourceTexture), InDesiredSize);
-					UE_LOG(LogMediaIOCore, Verbose, TEXT("[%s] - Captured frame %d"), *InMediaCapture->MediaOutputName, CapturingFrame->FrameId);
+					UE_LOG(LogMediaIOCore, Verbose, TEXT("[%s] - Captured frame %d"), *InMediaCapture->MediaOutputName, CapturingFrame ? CapturingFrame->FrameId : -1);
 					
 					GraphBuilder.Execute();
 				}
diff --git a/Engine/Plugins/VirtualProduction/RemoteControl/Source/RemoteControlLogic/Private/Action/Bind/RCPropertyBindAction.cpp b/Engine/Plugins/VirtualProduction/RemoteControl/Source/RemoteControlLogic/Private/Action/Bind/RCPropertyBindAction.cpp
--- a/Engine/Plugins/VirtualProduction/RemoteControl/Source/RemoteControlLogic/Private/Action/Bind/RCPropertyBindAction.cpp
+++ b/Engine/Plugins/VirtualProduction/RemoteControl/Source/RemoteControlLogic/Private/Action/Bind/RCPropertyBindAction.cpp
@@ -73,8 +73,12 @@ void URCPropertyBindAction::Execute() const
 	{
 		const FProperty* ControllerAsProperty = Controller->GetProperty();
 		FProperty* RemoteControlProperty = RemoteControlEntityAsProperty->GetProperty();
+		if (RemoteControlProperty == nullptr)
+		{
+			return;
+		}
 
-		TSharedPtr<IRemoteControlPropertyHandle> Handle = RemoteControlProperty ? RemoteControlEntityAsProperty->GetPropertyHandle() : nullptr;
+		TSharedPtr<IRemoteControlPropertyHandle> Handle = RemoteControlEntityAsProperty->GetPropertyHandle();
 		if (!ensure(Handle))
 		{
 			return;
