Remote Control Logic - Range Behaviour

Range Mapping Behaviour for the Remote Control Logic, which adds...
- Ability to lerp a selection of numerical (float, double) and structs (FVector, FRotator) based on Actions presented.
- Ability to activate non-Lerp Actions based on their distance towards the Controllers normalized Value and a Threshold.

Supports: Float Controller

diff --git a/Engine/Plugins/VirtualProduction/RemoteControl/Source/RemoteControlLogic/Private/Behaviour/Builtin/RangeMap/RCBehaviourRangeMapNode.cpp b/Engine/Plugins/VirtualProduction/RemoteControl/Source/RemoteControlLogic/Private/Behaviour/Builtin/RangeMap/RCBehaviourRangeMapNode.cpp
new file mode 100644
--- /dev/null
+++ b/Engine/Plugins/VirtualProduction/RemoteControl/Source/RemoteControlLogic/Private/Behaviour/Builtin/RangeMap/RCBehaviourRangeMapNode.cpp
@@ -0,0 +1,49 @@
+// Copyright Epic Games, Inc. All Rights Reserved.
+
+#include "Behaviour/Builtin/RangeMap/RCBehaviourRangeMapNode.h"
+
+#include "Behaviour/Builtin/RangeMap/RCRangeMapBehaviour.h"
+#include "Controller/RCController.h"
+#include "RCVirtualProperty.h"
+
+URCBehaviourRangeMapNode::URCBehaviourRangeMapNode()
+{
+	DisplayName = NSLOCTEXT("Remote Control Behaviour", "Behavior Name - Range Mapping", "Range Mapping");
+	BehaviorDescription = NSLOCTEXT("Remote Control Behaviour", "Behavior Desc - Range Mapping", "Behaviour lerps property values between a normalized 0-1 step. Users define min/max values and tolerance threshold.");
+}
+
+bool URCBehaviourRangeMapNode::Execute_Implementation(URCBehaviour* InBehaviour) const
+{
+	if (!ensure(InBehaviour))
+	{
+		return false;
+	}
+	
+	URCRangeMapBehaviour* RangeMapBehaviour = Cast<URCRangeMapBehaviour>(InBehaviour);
+	if (!RangeMapBehaviour)
+	{
+		return false;
+	}
+
+	RangeMapBehaviour->Refresh();
+	return true;
+}
+
+bool URCBehaviourRangeMapNode::IsSupported_Implementation(URCBehaviour* InBehaviour) const
+{
+	static TArray<TPair<EPropertyBagPropertyType, UObject*>> SupportedPropertyBagTypes =
+	{
+		{ EPropertyBagPropertyType::Float, nullptr},
+	};
+	
+	return SupportedPropertyBagTypes.ContainsByPredicate(GetIsSupportedCallback(InBehaviour));
+}
+
+UClass* URCBehaviourRangeMapNode::GetBehaviourClass() const
+{
+	return URCRangeMapBehaviour::StaticClass();
+}
+
+void URCBehaviourRangeMapNode::OnPassed_Implementation(URCBehaviour* InBehaviour) const
+{
+}
diff --git a/Engine/Plugins/VirtualProduction/RemoteControl/Source/RemoteControlLogic/Private/Behaviour/Builtin/RangeMap/RCRangeMapBehaviour.cpp b/Engine/Plugins/VirtualProduction/RemoteControl/Source/RemoteControlLogic/Private/Behaviour/Builtin/RangeMap/RCRangeMapBehaviour.cpp
new file mode 100644
--- /dev/null
+++ b/Engine/Plugins/VirtualProduction/RemoteControl/Source/RemoteControlLogic/Private/Behaviour/Builtin/RangeMap/RCRangeMapBehaviour.cpp
@@ -0,0 +1,458 @@
+﻿// Copyright Epic Games, Inc. All Rights Reserved.
+
+#include "Behaviour/Builtin/RangeMap/RCRangeMapBehaviour.h"
+
+#include "Action/RCAction.h"
+#include "Action/RCActionContainer.h"
+#include "Action/RCFunctionAction.h"
+#include "Action/RCPropertyAction.h"
+#include "Behaviour/Builtin/RangeMap/RCBehaviourRangeMapNode.h"
+#include "Behaviour/RCBehaviourNode.h"
+#include "Containers/Set.h"
+#include "Controller/RCController.h"
+#include "IRemoteControlPropertyHandle.h"
+#include "Kismet/KismetMathLibrary.h"
+#include "RCVirtualProperty.h"
+#include "RCVirtualPropertyContainer.h"
+#include "RemoteControlField.h"
+#include "RemoteControlPreset.h"
+
+namespace UE::RCRangeMapBehaviour
+{
+	const FName MinValue = TEXT("MinValue");
+	const FName MaxValue = TEXT("MaxValue");
+	const FName Threshold = TEXT("Threshold");
+	const FName Step = TEXT("Step");
+}
+
+URCRangeMapBehaviour::URCRangeMapBehaviour()
+{
+	PropertyContainer = CreateDefaultSubobject<URCVirtualPropertyContainerBase>(FName("VirtualPropertyContainer"));
+}
+
+void URCRangeMapBehaviour::Initialize()
+{
+	URCController* RCController = ControllerWeakPtr.Get();
+	if (!RCController)
+	{
+		return;
+	}
+
+	TWeakObjectPtr<URCVirtualPropertyContainerBase> ContainerPtr = PropertyContainer;
+
+	ContainerPtr->AddProperty(UE::RCRangeMapBehaviour::MinValue, URCController::StaticClass(), EPropertyBagPropertyType::Double);
+	ContainerPtr->AddProperty(UE::RCRangeMapBehaviour::MaxValue, URCController::StaticClass(), EPropertyBagPropertyType::Double);
+	ContainerPtr->AddProperty(UE::RCRangeMapBehaviour::Threshold, URCController::StaticClass(), EPropertyBagPropertyType::Double);
+	ContainerPtr->AddProperty(UE::RCRangeMapBehaviour::Step, URCController::StaticClass(), EPropertyBagPropertyType::Double);
+
+	Super::Initialize();
+}
+
+bool URCRangeMapBehaviour::IsSupportedActionLerpType(TObjectPtr<URCAction> InAction) const
+{
+	if (TObjectPtr<URCFunctionAction> InFunctionAction = Cast<URCFunctionAction>(InAction))
+	{
+		// Early false if function. Functions are always non-lerp.
+		return false;
+	}
+	
+	TObjectPtr<URCPropertyAction> InPropertyAction = Cast<URCPropertyAction>(InAction);
+	if (!InPropertyAction)
+	{
+		return false;
+	}
+	
+	const bool bIsNumeric = InPropertyAction->PropertySelfContainer->IsNumericType();
+	const bool bIsVector  = InPropertyAction->PropertySelfContainer->IsVectorType();
+	const bool bIsRotator = InPropertyAction->PropertySelfContainer->IsRotatorType();
+
+	return bIsNumeric || bIsVector || bIsRotator;
+}
+
+void URCRangeMapBehaviour::Refresh()
+{
+	URCController* Controller = ControllerWeakPtr.Get();
+	if (!Controller)
+	{
+		return;
+	}
+
+	// Step 0: Get All required Values, getting the current values set within the DetailPanel setting them up for our header.
+	PropertyContainer->GetVirtualProperty(UE::RCRangeMapBehaviour::MinValue)->GetValueDouble(MinValue);
+	PropertyContainer->GetVirtualProperty(UE::RCRangeMapBehaviour::MaxValue)->GetValueDouble(MaxValue);
+	PropertyContainer->GetVirtualProperty(UE::RCRangeMapBehaviour::Threshold)->GetValueDouble(Threshold);
+
+	Controller->GetValueFloat(ControllerFloatValue);
+	
+	// Step 1: Clamp the Value between Max/Min if necessary
+	if (ControllerFloatValue < MinValue || MaxValue < ControllerFloatValue)
+	{
+		ControllerFloatValue = FMath::Clamp(ControllerFloatValue, MinValue, MaxValue);
+		Controller->SetValueFloat(ControllerFloatValue);
+	}
+}
+
+bool URCRangeMapBehaviour::GetNearestActionByThreshold(TTuple<URCAction*, bool>& OutTuple)
+{
+	// Variables
+	TMap<double, URCAction*> NonLerpActions = GetNonLerpActions();
+	
+	TArray<double> StepActionArray;
+	NonLerpActions.GenerateKeyArray(StepActionArray);
+
+	// Step 1: Convert to double for Kismet Operation and Normalize.
+	double NormalizedControllerValue = UKismetMathLibrary::NormalizeToRange(ControllerFloatValue, MinValue, MaxValue);
+
+	// Step 2: Go through each of the Action Steps and look for the shortest distance
+	for (int StepIndex = 0; StepIndex < StepActionArray.Num(); StepIndex++)
+	{
+		double StepValue = StepActionArray[StepIndex];
+		double ValueDifference = FMath::Abs(StepValue - NormalizedControllerValue);
+		
+		if (StepActionArray.Num() > 1 && StepIndex > 0)
+		{
+			double PrevValueDifference = FMath::Abs(StepActionArray[StepIndex-1] - NormalizedControllerValue);
+			
+			// Break out early in case previous Step had a lesser difference
+			if (PrevValueDifference <= ValueDifference)
+			{
+				break;
+			}
+		}
+		OutTuple = TTuple<URCAction*, bool>(NonLerpActions[StepValue], ValueDifference <= Threshold);
+	}
+	return true;
+}
+
+void URCRangeMapBehaviour::GetLerpActions(TMap<FGuid, TArray<URCAction*>>& OutNumericActionsByField)
+{
+	for (URCAction* Action : ActionContainer->Actions)
+	{
+		TArray<URCAction*>& LerpActionArray = OutNumericActionsByField.FindOrAdd(Action->ExposedFieldId);
+		
+		// Step 01: Find Action
+		if (IsSupportedActionLerpType(Action))
+		{
+			// Step 02: Add Action if it's numerical
+			LerpActionArray.Add(Action);
+		}
+	}
+
+	// Step 03: Sort actions using their StepValue
+	for (TTuple<FGuid, TArray<URCAction*>>& NumericActionTuple : OutNumericActionsByField)
+	{
+		TArray<URCAction*>& ArrayToSort = NumericActionTuple.Value;
+
+		Algo::Sort(ArrayToSort, [this](const URCAction* ActionA, const URCAction* ActionB)
+		{
+			double A = RangeMapActionContainer.Find(ActionA)->StepValue;
+			double B = RangeMapActionContainer.Find(ActionB)->StepValue;
+
+			return A > B;
+		});
+	}
+}
+
+void URCRangeMapBehaviour::Execute()
+{
+	Refresh();
+	const URCBehaviourNode* BehaviourNode = GetBehaviourNode();
+	check(BehaviourNode);
+
+	if (BehaviourNode->GetClass() != URCBehaviourRangeMapNode::StaticClass())
+	{
+		return; // Allow custom Blueprints to drive their own behaviour entirely
+	}
+
+	// Do this beforehand.
+	BehaviourNode->PreExecute(this);
+	
+	URCController* RCController = ControllerWeakPtr.Get();
+	if (!RCController)
+	{
+		ensureMsgf(false, TEXT("Remote Control Controller is not available/nullptr."));
+		return;
+	}
+
+	RCController->GetValueFloat(ControllerFloatValue);
+	ControllerFloatValue = FMath::Clamp(ControllerFloatValue, MinValue, MaxValue);
+
+	TTuple<URCAction*, bool> NearestAction;
+	if (GetNearestActionByThreshold(NearestAction) && NearestAction.Value)
+	{
+		// Execute nearest Action in case its under/passes the Threshold in distance.
+		NearestAction.Key->Execute();
+	}
+
+	// Apply Lerp if possible
+	TMap<FGuid, TTuple<URCAction*, URCAction*>> LerpActions;
+	if (!GetRangeValuePairsForLerp(LerpActions))
+	{
+		return;
+	}
+
+	for (TTuple<FGuid, TTuple<URCAction*, URCAction*>>& LerpAction : LerpActions)
+	{
+		URemoteControlPreset* Preset = RCController->PresetWeakPtr.Get();
+		if (!Preset)
+		{
+			ensureMsgf(false, TEXT("Preset is invalid or unavailable."));
+			return;
+		}
+		FGuid CurrentFieldId = LerpAction.Key;
+
+		// Get the StepValues given for the FGuid
+		URCAction* MinActionOfPair = LerpAction.Value.Key;
+		URCAction* MaxActionOfPair = LerpAction.Value.Value;
+
+		if (!MinActionOfPair || !MaxActionOfPair)
+		{
+			return;
+		}
+
+		FRCRangeMapStep* MinRangeMapStep = RangeMapActionContainer.Find(MinActionOfPair);
+		FRCRangeMapStep* MaxRangeMapStep = RangeMapActionContainer.Find(MaxActionOfPair);
+
+		if (!(MinRangeMapStep && MaxRangeMapStep))
+		{
+			ensureMsgf(false, TEXT("%s is an invalid pointer."), MinRangeMapStep ? TEXT("MaxRangeMapStep") : TEXT("MinRangeMapStep"));
+			return;
+		}
+
+		// Denormalize and Map them based on our Min and Max Value
+		double MappedMinStep = UKismetMathLibrary::Lerp(MinValue, MaxValue, MinRangeMapStep->StepValue);
+		double MappedMaxStep = UKismetMathLibrary::Lerp(MinValue, MaxValue, MaxRangeMapStep->StepValue);
+
+		// Normalize our Controller based on the new MappedMinStep and MappedMaxStep Range.
+		double CustomNormalizedStepValue = UKismetMathLibrary::NormalizeToRange(ControllerFloatValue, MappedMinStep, MappedMaxStep);
+
+		// Apply value based on type
+		switch(MinRangeMapStep->PropertyValue->GetValueType())
+		{
+		case EPropertyBagPropertyType::Double:
+			{
+				double MinRangeDouble;
+				double MaxRangeDouble;
+
+				MinRangeMapStep->PropertyValue->GetValueDouble(MinRangeDouble);
+				MaxRangeMapStep->PropertyValue->GetValueDouble(MaxRangeDouble);
+
+				double ResultDouble = UKismetMathLibrary::Lerp(MinRangeDouble, MaxRangeDouble, CustomNormalizedStepValue);
+				if (TSharedPtr<FRemoteControlProperty> RCProperty = Preset->GetExposedEntity<FRemoteControlProperty>(CurrentFieldId).Pin())
+				{
+					RCProperty->GetPropertyHandle()->SetValue(ResultDouble);
+				}
+				
+				break;
+			}
+		case EPropertyBagPropertyType::Float:
+			{
+				float MinRangeFloat;
+				float MaxRangeFloat;
+
+				MinRangeMapStep->PropertyValue->GetValueFloat(MinRangeFloat);
+				MaxRangeMapStep->PropertyValue->GetValueFloat(MaxRangeFloat);
+
+				float ResultFloat = UKismetMathLibrary::Lerp(MinRangeFloat, MaxRangeFloat, CustomNormalizedStepValue);
+				if (TSharedPtr<FRemoteControlProperty> RCProperty = Preset->GetExposedEntity<FRemoteControlProperty>(CurrentFieldId).Pin())
+				{
+					RCProperty->GetPropertyHandle()->SetValue(ResultFloat);
+				}
+				
+				break;
+			}
+		case EPropertyBagPropertyType::Struct:
+			{
+				ApplyLerpOnStruct(MinRangeMapStep, MaxRangeMapStep, CustomNormalizedStepValue, CurrentFieldId);
+				
+				break;
+			}
+		default:
+			// Shouldn't happen.
+			ensureMsgf(false, TEXT("Unsupported Lerp Type."));
+		}
+	}
+}
+
+URCAction* URCRangeMapBehaviour::AddAction(const TSharedRef<const FRemoteControlField> InRemoteControlField)
+{
+	double InStepValue = 0.0;
+	PropertyContainer->GetVirtualProperty(UE::RCRangeMapBehaviour::Step);
+	
+	const TRCActionUniquenessTest ActionUniquenessTest = [this, InRemoteControlField, InStepValue](const TSet<TObjectPtr<URCAction>>& Actions)
+	{
+		return IsActionUnique(InRemoteControlField, InStepValue, Actions);
+	};
+	
+	return ActionContainer->AddAction(ActionUniquenessTest, InRemoteControlField);
+}
+
+void URCRangeMapBehaviour::OnActionAdded(URCAction* Action, const TObjectPtr<URCVirtualPropertySelfContainer> InPropertyValue)
+{
+	double InStepValue;
+	URCVirtualPropertyBase* StepVirtualProperty = PropertyContainer->GetVirtualProperty(UE::RCRangeMapBehaviour::Step);
+	StepVirtualProperty->GetValueDouble(InStepValue);
+	StepVirtualProperty->SetValueDouble(0.0);
+	
+	FRCRangeMapStep RangeMapStep = FRCRangeMapStep(InStepValue, InPropertyValue);
+
+	RangeMapActionContainer.Add(Action, MoveTemp(RangeMapStep));
+}
+
+bool URCRangeMapBehaviour::GetRangeValuePairsForLerp(TMap<FGuid, TTuple<URCAction*, URCAction*>>& OutPairs)
+{
+	double NormalizedControllerValue = UKismetMathLibrary::NormalizeToRange(ControllerFloatValue, MinValue, MaxValue);
+
+	TMap<FGuid, TArray<URCAction*>> LerpActionMap;
+	GetLerpActions(LerpActionMap);
+	
+	for (TTuple<FGuid, TArray<URCAction*>>& NumericActionTuple : LerpActionMap)
+	{
+		if (NumericActionTuple.Value.Num() < 2)
+		{
+			// Skip and continue early if the Array has less Numeric Values needed to actually Lerp
+			continue;
+		}
+
+		// Intermediary calculations of Lerp for in-between Lerp.
+		URCAction* MinAction;
+		FRCRangeMapStep* MinRangeMap = nullptr;
+		URCAction* MaxAction;
+		FRCRangeMapStep* MaxRangeMap = nullptr;
+		
+		for (URCAction* Action : NumericActionTuple.Value)
+		{
+			FRCRangeMapStep* RangeMap = RangeMapActionContainer.Find(Action);
+			if (!RangeMap)
+			{
+				// Shouldn't happen, but skip if it does
+				continue;
+			}
+
+			// Deal with the nullptr
+			if (NormalizedControllerValue >= RangeMap->StepValue && (!MinRangeMap || NumericActionTuple.Value.Last() != Action))
+			{
+				MinAction = Action;
+				MinRangeMap = RangeMap;
+				
+				continue;
+			}
+
+			if (NormalizedControllerValue <= RangeMap->StepValue)
+			{
+				MaxAction = Action;
+				MaxRangeMap = RangeMap;
+
+				continue;
+			}
+
+			// Ensure there are no nullptr
+			if (!MinRangeMap || !MaxRangeMap)
+			{
+				return false;
+			}
+		}
+
+		OutPairs.Add(NumericActionTuple.Key, TTuple<URCAction*, URCAction*>(MinAction, MaxAction));
+	}
+	
+	return true;
+}
+
+TMap<double, URCAction*> URCRangeMapBehaviour::GetNonLerpActions()
+{
+	TMap<double, URCAction*> ActionMap;
+	
+	for (URCAction* Action : ActionContainer->Actions)
+	{
+		
+		if (!IsSupportedActionLerpType(Action))
+		{
+			FRCRangeMapStep* RangeMapStep = RangeMapActionContainer.Find(Action);
+
+			if (RangeMapStep)
+			{
+				ActionMap.Add(RangeMapStep->StepValue, Action);
+			}
+		}
+	}
+
+	return ActionMap;
+}
+
+bool URCRangeMapBehaviour::CanHaveActionForField(const TSharedPtr<FRemoteControlField> InRemoteControlField) const
+{
+	// We can select all of the actions, and the check whether its actually added is done elsewhere.
+	return true;
+}
+
+void URCRangeMapBehaviour::ApplyLerpOnStruct(const FRCRangeMapStep* MinRangeStep, const FRCRangeMapStep* MaxRangeStep, const double& StepAlpha, const FGuid& FieldId)
+{
+	URCController* RCController = ControllerWeakPtr.Get();
+	if (!RCController)
+	{
+		ensureMsgf(false, TEXT("Remote Controller is invalid or unavailable."));
+		return;
+	}
+	
+	URemoteControlPreset* Preset = RCController->PresetWeakPtr.Get();
+	if (!Preset)
+	{
+		ensureMsgf(false, TEXT("Preset is invalid or unavailable."));
+		return;
+	}
+	
+	// Check whether its a Rotator or FVector based on the MinRangeStep
+	if (MinRangeStep->PropertyValue->IsVectorType())
+	{
+		FVector MinRangeVector;
+		FVector MaxRangeVector;
+
+		MinRangeStep->PropertyValue->GetValueVector(MinRangeVector);
+		MaxRangeStep->PropertyValue->GetValueVector(MaxRangeVector);
+
+		FVector ResultVector = UKismetMathLibrary::VLerp(MinRangeVector, MaxRangeVector, StepAlpha);
+		if (TSharedPtr<FRemoteControlProperty> RCProperty = Preset->GetExposedEntity<FRemoteControlProperty>(FieldId).Pin())
+		{
+			RCProperty->GetPropertyHandle()->SetValue(ResultVector);
+		}
+	}
+
+	if (MinRangeStep->PropertyValue->IsRotatorType())
+	{
+		FRotator MinRangeRotator;
+		FRotator MaxRangeRotator;
+
+		MinRangeStep->PropertyValue->GetValueRotator(MinRangeRotator);
+		MaxRangeStep->PropertyValue->GetValueRotator(MaxRangeRotator);
+
+		FRotator ResultVector = UKismetMathLibrary::RLerp(MinRangeRotator, MaxRangeRotator, StepAlpha, true);
+		if (TSharedPtr<FRemoteControlProperty> RCProperty = Preset->GetExposedEntity<FRemoteControlProperty>(FieldId).Pin())
+		{
+			 RCProperty->GetPropertyHandle()->SetValue(ResultVector);
+		}
+	}
+}
+
+bool URCRangeMapBehaviour::IsActionUnique(const TSharedRef<const FRemoteControlField> InRemoteControlField, const double& InStepValue, const TSet<TObjectPtr<URCAction>>& InActions)
+{
+	const FGuid FieldId = InRemoteControlField->GetId();
+
+	for (const URCAction* Action : InActions)
+	{
+		// Check whether or not we have a RangeStep.
+		if (const FRCRangeMapStep* RangeMapStep = this->RangeMapActionContainer.Find(Action))
+		{
+			// Check if the Step itself is already used
+			const double ActionStepValue = RangeMapStep->StepValue;
+
+			// Only Important bit is StepValue and FieldId
+			if (Action->ExposedFieldId == FieldId && FMath::IsNearlyEqual(ActionStepValue, InStepValue))
+			{
+				return false; // Not Unique.
+			}
+		}
+	}
+	
+	return true;
+}
diff --git a/Engine/Plugins/VirtualProduction/RemoteControl/Source/RemoteControlLogic/Public/Behaviour/Builtin/RangeMap/RCBehaviourRangeMapNode.h b/Engine/Plugins/VirtualProduction/RemoteControl/Source/RemoteControlLogic/Public/Behaviour/Builtin/RangeMap/RCBehaviourRangeMapNode.h
new file mode 100644
--- /dev/null
+++ b/Engine/Plugins/VirtualProduction/RemoteControl/Source/RemoteControlLogic/Public/Behaviour/Builtin/RangeMap/RCBehaviourRangeMapNode.h
@@ -0,0 +1,32 @@
+// Copyright Epic Games, Inc. All Rights Reserved.
+
+#pragma once
+
+#include "Behaviour/RCBehaviourNode.h"
+#include "RCBehaviourRangeMapNode.generated.h"
+
+/**
+ * Uses the given properties and values to limit the values of the Controller and perform lerp operations for the values.
+ */
+UCLASS()
+class REMOTECONTROLLOGIC_API URCBehaviourRangeMapNode : public URCBehaviourNode
+{
+	GENERATED_BODY()
+
+public:
+	
+	URCBehaviourRangeMapNode();
+
+	//~ Begin URCBehaviourNode interface
+	UFUNCTION(BlueprintNativeEvent, BlueprintCallable, Category = "Remote Control Behaviour")
+	bool Execute(URCBehaviour* InBehaviour) const override;
+
+	UFUNCTION(BlueprintNativeEvent, BlueprintCallable, Category = "Remote Control Behaviour")
+	bool IsSupported(URCBehaviour* InBehaviour) const override;
+
+	UFUNCTION(BlueprintNativeEvent, BlueprintCallable, Category = "Remote Control Behaviour")
+	void OnPassed(URCBehaviour* InBehaviour) const;
+
+	UClass* GetBehaviourClass() const override;
+	//~ End URCBehaviourNode interface
+};
diff --git a/Engine/Plugins/VirtualProduction/RemoteControl/Source/RemoteControlLogic/Public/Behaviour/Builtin/RangeMap/RCRangeMapBehaviour.h b/Engine/Plugins/VirtualProduction/RemoteControl/Source/RemoteControlLogic/Public/Behaviour/Builtin/RangeMap/RCRangeMapBehaviour.h
new file mode 100644
--- /dev/null
+++ b/Engine/Plugins/VirtualProduction/RemoteControl/Source/RemoteControlLogic/Public/Behaviour/Builtin/RangeMap/RCRangeMapBehaviour.h
@@ -0,0 +1,111 @@
+﻿// Copyright Epic Games, Inc. All Rights Reserved.
+
+#pragma once
+
+#include "Behaviour/RCBehaviour.h"
+#include "RCVirtualProperty.h"
+#include "RCRangeMapBehaviour.generated.h"
+
+class URCAction;
+class URCVirtualPropertyContainerBase;
+
+USTRUCT()
+struct FRCRangeMapStep
+{
+	GENERATED_BODY()
+
+	FRCRangeMapStep() {}
+	FRCRangeMapStep(const double InStepValue, const URCVirtualPropertySelfContainer* InPropertyValue)
+		: StepValue(InStepValue)
+	{
+		if (!PropertyValue)
+		{
+			return;
+		}
+		
+		PropertyValue = MakeObjectPtrUnsafe<URCVirtualPropertySelfContainer>(InPropertyValue);
+	}
+
+	/** The Value which we use represent the action based on a normalized step. */
+	UPROPERTY()
+	double StepValue = 0.0;
+
+	/** The Property this Action holds and will be used for calculations for the lerp. */
+	UPROPERTY()
+	TObjectPtr<URCVirtualPropertySelfContainer> PropertyValue;
+};
+
+/**
+ * Custom behaviour for Set Asset By Path Logic Behaviour 
+ */
+UCLASS()
+class REMOTECONTROLLOGIC_API URCRangeMapBehaviour : public URCBehaviour
+{
+	GENERATED_BODY()
+
+public:
+	/** Pointer to property container */
+	UPROPERTY()
+	TObjectPtr<URCVirtualPropertyContainerBase> PropertyContainer;
+
+	/** Container holding all FRCRangeMapSteps correlating to each Action in the ActionContainer. */
+	UPROPERTY()
+	TMap<URCAction*, FRCRangeMapStep> RangeMapActionContainer;
+
+public:
+	URCRangeMapBehaviour();
+	
+	//~ Begin URCBehaviour interface
+	virtual void Initialize() override;
+
+	/** Executes the Behaviour */
+	virtual void Execute() override;
+	//~ End URCBehaviour interface
+
+	/** Refresh function being called whenever either the Controller or the Properties of the Behaviour Details Panel change */
+	void Refresh();
+	
+	/** Called whenever a new action is added into the ActionContainer. Will add a FRCRangeMapStep into RangeMapActionContainer corresponding to the Action. */
+	void OnActionAdded(URCAction* Action, const TObjectPtr<URCVirtualPropertySelfContainer> InPropertyValue);
+
+	/** Adds an Action to the ActionContainer, whilst making sure that it is unique. */
+	URCAction* AddAction(const TSharedRef<const FRemoteControlField> InRemoteControlField) override;
+
+	/** Checks whether or not one of the following fields can be added to the AddAction Menu. Makes sure they're unique. */
+	bool CanHaveActionForField(const TSharedPtr<FRemoteControlField> InRemoteControlField) const override;
+
+private:
+	/** Minimum Value which the Range has */
+	double MinValue;
+
+	/** Maximum Value which the Range has */
+	double MaxValue;
+
+	/** Threshold the Behaviour should adhere to for rounding up and down. */
+	double Threshold;
+
+	/** Controller Value of Type Float, which is used for the purpose of readability. */
+	float ControllerFloatValue;
+	
+private:
+	/** Boolean Function to help differentiate Custom Actions */
+	bool IsSupportedActionLerpType(TObjectPtr<URCAction> InAction) const;
+
+	/** Return the nearest Action and whether or not the distance between is less than the threshold.  */
+	bool GetNearestActionByThreshold(TTuple<URCAction*, bool>& OutTuple);
+
+	/** Returns a Map of Numeric Actions, mapped by a double value, bound to a given FieldId */
+	void GetLerpActions(TMap<FGuid, TArray<URCAction*>>& OutNumericActionsByField);
+	
+	/** Gives out all pairs per Unique Exposed Field, which are applicable for Lerp. */
+	bool GetRangeValuePairsForLerp(TMap<FGuid, TTuple<URCAction*, URCAction*>>& OutPairs);
+
+	/** Gives out all Actions which do not fall under as being able to Lerp. Used for Executing them based on Threshold and Distance */
+	TMap<double, URCAction*> GetNonLerpActions();
+
+	/** Auxiliary Function used to calculate and apply the Lerp on Structs, like FVectors or FRotators. */
+	void ApplyLerpOnStruct(const FRCRangeMapStep* MinRangeStep, const FRCRangeMapStep* MaxRangeStep, const double& StepAlpha, const FGuid& FieldId);
+
+	/** Returns whether or not an Action created with a given RemoteControlField can be considered unique. */
+	bool IsActionUnique(const TSharedRef<const FRemoteControlField> InRemoteControlField, const double& InStepValue, const TSet<TObjectPtr<URCAction>>& InActions);
+};
diff --git a/Engine/Plugins/VirtualProduction/RemoteControl/Source/RemoteControlUI/Private/UI/Action/Bind/RCActionBindModel.cpp b/Engine/Plugins/VirtualProduction/RemoteControl/Source/RemoteControlUI/Private/UI/Action/Bind/RCActionBindModel.cpp
--- a/Engine/Plugins/VirtualProduction/RemoteControl/Source/RemoteControlUI/Private/UI/Action/Bind/RCActionBindModel.cpp
+++ b/Engine/Plugins/VirtualProduction/RemoteControl/Source/RemoteControlUI/Private/UI/Action/Bind/RCActionBindModel.cpp
@@ -45,7 +45,7 @@ namespace UE::RCActionPanelBindList
 		
 			if (ColumnName == UE::RCActionPanelBindList::Columns::VariableColor)
 			{
-				return ActionItem->GetVariableColorWidget();
+				return ActionItem->GetTypeColorTagWidget();
 			}
 			else if (ColumnName == UE::RCActionPanelBindList::Columns::Description)
 			{
diff --git a/Engine/Plugins/VirtualProduction/RemoteControl/Source/RemoteControlUI/Private/UI/Action/Conditional/RCActionConditionalModel.cpp b/Engine/Plugins/VirtualProduction/RemoteControl/Source/RemoteControlUI/Private/UI/Action/Conditional/RCActionConditionalModel.cpp
--- a/Engine/Plugins/VirtualProduction/RemoteControl/Source/RemoteControlUI/Private/UI/Action/Conditional/RCActionConditionalModel.cpp
+++ b/Engine/Plugins/VirtualProduction/RemoteControl/Source/RemoteControlUI/Private/UI/Action/Conditional/RCActionConditionalModel.cpp
@@ -28,7 +28,7 @@ namespace UE::RCActionPanelConditionalList
 {
 	namespace Columns
 	{
-		const FName VariableColor = TEXT("VariableColor");
+		const FName TypeColorTag = TEXT("TypeColorTag");
 		const FName DragDropHandle = TEXT("DragDropHandle");
 		const FName Condition = TEXT("Condition");
 		const FName Description = TEXT("Description");
@@ -49,9 +49,9 @@ namespace UE::RCActionPanelConditionalList
 			if (!ensure(ActionItem.IsValid()))
 				return SNullWidget::NullWidget;
 
-			if (ColumnName == UE::RCActionPanelConditionalList::Columns::VariableColor)
+			if (ColumnName == UE::RCActionPanelConditionalList::Columns::TypeColorTag)
 			{
-				return ActionItem->GetVariableColorWidget();
+				return ActionItem->GetTypeColorTagWidget();
 			}
 			else if (ColumnName == UE::RCActionPanelConditionalList::Columns::Condition)
 			{
@@ -83,8 +83,8 @@ TSharedPtr<SHeaderRow> FRCActionConditionalModel::GetHeaderRow()
 	return SNew(SHeaderRow)
 		.Style(&RCPanelStyle->HeaderRowStyle)
 
-		+ SHeaderRow::Column(UE::RCActionPanelConditionalList::Columns::VariableColor)
-		.DefaultLabel(LOCTEXT("RCActionVariableColorColumnHeader", ""))
+		+ SHeaderRow::Column(UE::RCActionPanelConditionalList::Columns::TypeColorTag)
+		.DefaultLabel(LOCTEXT("RCActionTypeColorColumnHeader", ""))
 		.FixedWidth(5.f)
 		.HeaderContentPadding(RCPanelStyle->HeaderRowPadding)
 
diff --git a/Engine/Plugins/VirtualProduction/RemoteControl/Source/RemoteControlUI/Private/UI/Action/RCActionModel.cpp b/Engine/Plugins/VirtualProduction/RemoteControl/Source/RemoteControlUI/Private/UI/Action/RCActionModel.cpp
--- a/Engine/Plugins/VirtualProduction/RemoteControl/Source/RemoteControlUI/Private/UI/Action/RCActionModel.cpp
+++ b/Engine/Plugins/VirtualProduction/RemoteControl/Source/RemoteControlUI/Private/UI/Action/RCActionModel.cpp
@@ -52,7 +52,7 @@ namespace UE::RCActionPanelList
 
 			if (ColumnName == UE::RCActionPanelList::Columns::VariableColor)
 			{
-				return ActionItem->GetVariableColorWidget();
+				return ActionItem->GetTypeColorTagWidget();
 			}
 			else if (ColumnName == UE::RCActionPanelList::Columns::Description)
 			{
@@ -108,11 +108,11 @@ TSharedRef<SWidget> FRCActionModel::GetWidget() const
 	return SNullWidget::NullWidget;
 }
 
-TSharedRef<SWidget> FRCActionModel::GetVariableColorWidget() const
+TSharedRef<SWidget> FRCActionModel::GetTypeColorTagWidget() const
 {
 	const FLinearColor TypeColor = GetActionTypeColor();
 
-	// Variable Color Bar
+	// Type Color Bar
 	return SNew(SBox)
 		.HeightOverride(5.f)
 		[
diff --git a/Engine/Plugins/VirtualProduction/RemoteControl/Source/RemoteControlUI/Private/UI/Action/RCActionModel.h b/Engine/Plugins/VirtualProduction/RemoteControl/Source/RemoteControlUI/Private/UI/Action/RCActionModel.h
--- a/Engine/Plugins/VirtualProduction/RemoteControl/Source/RemoteControlUI/Private/UI/Action/RCActionModel.h
+++ b/Engine/Plugins/VirtualProduction/RemoteControl/Source/RemoteControlUI/Private/UI/Action/RCActionModel.h
@@ -44,7 +44,7 @@ public:
 	virtual TSharedRef<SWidget> GetWidget() const override;
 
 	/** Widget representing Color Coding for the Action Type*/
-	virtual TSharedRef<SWidget> GetVariableColorWidget() const;
+	virtual TSharedRef<SWidget> GetTypeColorTagWidget() const;
 
 	/** Color code for this Action. Customized per action type*/
 	virtual FLinearColor GetActionTypeColor() const = 0;
diff --git a/Engine/Plugins/VirtualProduction/RemoteControl/Source/RemoteControlUI/Private/UI/Action/RangeMap/RCActionRangeMapModel.cpp b/Engine/Plugins/VirtualProduction/RemoteControl/Source/RemoteControlUI/Private/UI/Action/RangeMap/RCActionRangeMapModel.cpp
new file mode 100644
--- /dev/null
+++ b/Engine/Plugins/VirtualProduction/RemoteControl/Source/RemoteControlUI/Private/UI/Action/RangeMap/RCActionRangeMapModel.cpp
@@ -0,0 +1,161 @@
+// Copyright Epic Games, Inc. All Rights Reserved.
+
+#include "RCActionRangeMapModel.h"
+
+#include "Behaviour/Builtin/RangeMap/RCRangeMapBehaviour.h"
+#include "Controller/RCController.h"
+#include "IDetailTreeNode.h"
+#include "IPropertyRowGenerator.h"
+#include "RCVirtualProperty.h"
+#include "RCVirtualPropertyContainer.h"
+#include "RemoteControlField.h"
+#include "RemoteControlPreset.h"
+#include "Styling/RemoteControlStyles.h"
+#include "UI/Behaviour/Builtin/RangeMap/RCBehaviourRangeMapModel.h"
+#include "UI/RCUIHelpers.h"
+#include "UI/RemoteControlPanelStyle.h"
+#include "UI/SRCPanelDragHandle.h"
+#include "UI/SRCPanelExposedEntity.h"
+#include "Widgets/Layout/SBox.h"
+#include "Widgets/Views/SHeaderRow.h"
+#include "Widgets/Views/SListView.h"
+#include "Widgets/Views/STableRow.h"
+
+#define LOCTEXT_NAMESPACE "RCActionModel"
+
+namespace UE::RCActionPanelRangeMapList
+{
+	namespace Columns
+	{
+		const FName TypeColorTag = TEXT("TypeColorTag");
+		const FName DragDropHandle = TEXT("DragDropHandle");
+		const FName Step = TEXT("Step");
+		const FName Name = TEXT("Name");
+		const FName Value = TEXT("Value");
+	}
+
+	class SActionItemListRow : public SMultiColumnTableRow<TSharedRef<FRCActionModel>>
+	{
+	public:
+		void Construct(const FTableRowArgs& InArgs, const TSharedRef<STableViewBase>& OwnerTableView, TSharedRef<FRCActionRangeMapModel> InActionItem)
+		{
+			ActionItem = InActionItem;
+			FSuperRowType::Construct(InArgs, OwnerTableView);
+		}
+
+		TSharedRef<SWidget> GenerateWidgetForColumn(const FName& ColumnName) override
+		{
+			if (!ensure(ActionItem.IsValid()))
+				return SNullWidget::NullWidget;
+
+			if (ColumnName == UE::RCActionPanelRangeMapList::Columns::TypeColorTag)
+			{
+				return ActionItem->GetTypeColorTagWidget();
+			}
+			else if (ColumnName == UE::RCActionPanelRangeMapList::Columns::Step)
+			{
+				return ActionItem->GetStepWidget();
+			}
+			else if (ColumnName == UE::RCActionPanelRangeMapList::Columns::Name)
+			{
+				return ActionItem->GetNameWidget();
+			}
+			else if (ColumnName == UE::RCActionPanelRangeMapList::Columns::Value)
+			{
+				return ActionItem->GetWidget();
+			}
+
+			// TODO: Implement Drag-Drop-handle column with Actions reordering support
+
+			return SNullWidget::NullWidget;
+		}
+
+	private:
+		TSharedPtr<FRCActionRangeMapModel> ActionItem;
+	};
+}
+
+TSharedPtr<SHeaderRow> FRCActionRangeMapModel::GetHeaderRow()
+{
+	const FRCPanelStyle* RCPanelStyle = &FRemoteControlPanelStyle::Get()->GetWidgetStyle<FRCPanelStyle>("RemoteControlPanel.MinorPanel");
+
+	return SNew(SHeaderRow)
+		.Style(&RCPanelStyle->HeaderRowStyle)
+
+		+ SHeaderRow::Column(UE::RCActionPanelRangeMapList::Columns::TypeColorTag)
+		.DefaultLabel(LOCTEXT("RCActionVariableColorColumnHeader", ""))
+		.FixedWidth(5.f)
+		.HeaderContentPadding(RCPanelStyle->HeaderRowPadding)
+
+		+ SHeaderRow::Column(UE::RCActionPanelRangeMapList::Columns::DragDropHandle)
+		.DefaultLabel(LOCTEXT("RCActionDragDropHandleColumnHeader", ""))
+		.FixedWidth(25.f)
+		.HeaderContentPadding(RCPanelStyle->HeaderRowPadding)
+
+		+ SHeaderRow::Column(UE::RCActionPanelRangeMapList::Columns::Step)
+		.DefaultLabel(LOCTEXT("RCActionConditionColumnHeader", "Step"))
+		.FillWidth(0.2f)
+		.HeaderContentPadding(RCPanelStyle->HeaderRowPadding)
+
+		+ SHeaderRow::Column(UE::RCActionPanelRangeMapList::Columns::Name)
+		.DefaultLabel(LOCTEXT("RCActionDescColumnHeader", "Description"))
+		.FillWidth(0.4f)
+		.HeaderContentPadding(RCPanelStyle->HeaderRowPadding)
+
+		+ SHeaderRow::Column(UE::RCActionPanelRangeMapList::Columns::Value)
+		.DefaultLabel(LOCTEXT("RCActionValueColumnHeader", "Value"))
+		.FillWidth(0.4f)
+		.HeaderContentPadding(RCPanelStyle->HeaderRowPadding);
+}
+
+FRCActionRangeMapModel::FRCActionRangeMapModel(URCAction* InAction, const TSharedPtr<FRCBehaviourModel> InBehaviourItem, const TSharedPtr<SRemoteControlPanel> InRemoteControlPanel)
+	: FRCActionModel(InAction, InBehaviourItem, InRemoteControlPanel)
+{
+}
+
+TSharedPtr<FRCActionRangeMapModel> FRCActionRangeMapModel::GetModelByActionType(URCAction* InAction, const TSharedPtr<class FRCBehaviourModel> InBehaviourItem, const TSharedPtr<SRemoteControlPanel> InRemoteControlPanel)
+{
+	if (URCPropertyAction* PropertyAction = Cast<URCPropertyAction>(InAction))
+	{
+		return MakeShared<FRCPropertyActionRangeMapModel>(PropertyAction, InBehaviourItem, InRemoteControlPanel);
+	}
+	else if (URCFunctionAction* FunctionAction = Cast<URCFunctionAction>(InAction))
+	{
+		return MakeShared<FRCFunctionActionRangeMapModel>(FunctionAction, InBehaviourItem, InRemoteControlPanel);
+	}
+
+	return nullptr;
+}
+
+TSharedRef<SWidget> FRCActionRangeMapModel::GetStepWidget() const
+{
+	if (TSharedPtr<FRCRangeMapBehaviourModel> BehaviourItem = StaticCastSharedPtr<FRCRangeMapBehaviourModel>(BehaviourItemWeakPtr.Pin()))
+	{
+		if (URCRangeMapBehaviour* Behaviour = Cast<URCRangeMapBehaviour>(BehaviourItem->GetBehaviour()))
+		{
+			if (FRCRangeMapStep* RangeStep = Behaviour->RangeMapActionContainer.Find(GetAction()))
+			{
+				const FText StepAsText = FText::FromString(FString::SanitizeFloat(RangeStep->StepValue));
+				
+				return SNew(SBox)
+					.Padding(6.0f)
+					[
+						SNew(STextBlock).Text(StepAsText)
+					];
+			}
+		}
+	}
+	
+	return SNullWidget::NullWidget;
+}
+
+TSharedRef<ITableRow> FRCActionRangeMapModel::OnGenerateWidgetForList(TSharedPtr<FRCActionRangeMapModel> InItem, const TSharedRef<STableViewBase>& OwnerTable)
+{
+	typedef UE::RCActionPanelRangeMapList::SActionItemListRow ActionRowType;
+
+	return SNew(ActionRowType, OwnerTable, InItem.ToSharedRef())
+		.Style(&RCPanelStyle->TableRowStyle)
+		.Padding(3.f);
+}
+
+#undef LOCTEXT_NAMESPACE
\ No newline at end of file
diff --git a/Engine/Plugins/VirtualProduction/RemoteControl/Source/RemoteControlUI/Private/UI/Action/RangeMap/RCActionRangeMapModel.h b/Engine/Plugins/VirtualProduction/RemoteControl/Source/RemoteControlUI/Private/UI/Action/RangeMap/RCActionRangeMapModel.h
new file mode 100644
--- /dev/null
+++ b/Engine/Plugins/VirtualProduction/RemoteControl/Source/RemoteControlUI/Private/UI/Action/RangeMap/RCActionRangeMapModel.h
@@ -0,0 +1,78 @@
+// Copyright Epic Games, Inc. All Rights Reserved.
+
+#pragma once
+
+#include "Action/RCFunctionAction.h"
+#include "Action/RCPropertyAction.h"
+#include "UI/Action/RCActionModel.h"
+
+/*
+ * ~ FRCActionRangeMapModel ~
+ *
+ * UI Model for representing Actions as Steps for the RangeMap Behaviour.
+ * Contains and represents, other than name and value, a step value being a value between 0 to 1.
+ *
+ * <Step Value> <Name> <Value>
+ */
+class FRCActionRangeMapModel : public FRCActionModel
+{
+public:
+	FRCActionRangeMapModel(URCAction* InAction, const TSharedPtr<class FRCBehaviourModel> InBehaviourItem, const TSharedPtr<SRemoteControlPanel> InRemoteControlPanel);
+
+	/**
+	 * Simple Widget representing the Normalized Value or Step Value for the Range and this particular Value.
+	 */
+	virtual TSharedRef<SWidget> GetStepWidget() const;
+
+	/** OnGenerateRow delegate for the Actions List View*/
+	TSharedRef<ITableRow> OnGenerateWidgetForList(TSharedPtr<FRCActionRangeMapModel> InItem, const TSharedRef<STableViewBase>& OwnerTable);
+
+	/** Get the Header Row appropriate for this particular Action model */
+	static TSharedPtr<SHeaderRow> GetHeaderRow();
+
+	/** Chooses the appropriate Action model for the current class and field type*/
+	static TSharedPtr<FRCActionRangeMapModel> GetModelByActionType(URCAction* InAction, const TSharedPtr<class FRCBehaviourModel> InBehaviourItem, const TSharedPtr<SRemoteControlPanel> InRemoteControlPanel);
+};
+
+/*
+* ~ FRCPropertyActionRangeMapModel ~
+*
+* UI model for Property based RangeMap Actions
+*/
+class FRCPropertyActionRangeMapModel : public FRCActionRangeMapModel, public FRCPropertyActionType
+{
+public:
+	FRCPropertyActionRangeMapModel(URCPropertyAction* InPropertyAction, const TSharedPtr<class FRCBehaviourModel> InBehaviourItem, const TSharedPtr<SRemoteControlPanel> InRemoteControlPanel)
+	: FRCActionRangeMapModel(InPropertyAction, InBehaviourItem, InRemoteControlPanel),
+		FRCPropertyActionType(InPropertyAction) { }
+
+	virtual TSharedRef<SWidget> GetWidget() const override
+	{
+		return GetPropertyWidget();
+	}
+
+	/** Color code for this Action*/
+	virtual FLinearColor GetActionTypeColor() const override
+	{
+		return GetPropertyTypeColor();
+	}
+};
+
+/*
+* ~ FRCFunctionActionRangeMapModel ~
+*
+* UI model for Function based RangeMap Actions
+*/
+class FRCFunctionActionRangeMapModel : public FRCActionRangeMapModel, public FRCFunctionActionType
+{
+public:
+	FRCFunctionActionRangeMapModel (URCFunctionAction* InFunctionAction, const TSharedPtr<class FRCBehaviourModel> InBehaviourItem, const TSharedPtr<SRemoteControlPanel> InRemoteControlPanel)
+		: FRCActionRangeMapModel(InFunctionAction, InBehaviourItem, InRemoteControlPanel)
+		, FRCFunctionActionType(InFunctionAction) { }
+
+	/** Color code for this Action*/
+	virtual FLinearColor GetActionTypeColor() const override
+	{
+		return GetFunctionTypeColor();
+	}
+};
\ No newline at end of file
diff --git a/Engine/Plugins/VirtualProduction/RemoteControl/Source/RemoteControlUI/Private/UI/Action/SRCActionPanel.cpp b/Engine/Plugins/VirtualProduction/RemoteControl/Source/RemoteControlUI/Private/UI/Action/SRCActionPanel.cpp
--- a/Engine/Plugins/VirtualProduction/RemoteControl/Source/RemoteControlUI/Private/UI/Action/SRCActionPanel.cpp
+++ b/Engine/Plugins/VirtualProduction/RemoteControl/Source/RemoteControlUI/Private/UI/Action/SRCActionPanel.cpp
@@ -378,12 +378,12 @@ FReply SRCActionPanel::OnAddAllFields()
 			{
 				if (const TSharedPtr<FRemoteControlField> RemoteControlField = RemoteControlFieldWeakPtr.Pin())
 				{
-					if (const URCBehaviour* Behaviour = BehaviourItem->GetBehaviour())
+					URCBehaviour* Behaviour = BehaviourItem->GetBehaviour();
+					
+					// Only add the Behaviour if it's listed as addable.
+					if (Behaviour && Behaviour->CanHaveActionForField(RemoteControlField))
 					{
-						if (Behaviour->CanHaveActionForField(RemoteControlField))
-						{
-							AddAction(RemoteControlField.ToSharedRef());
-						}
+						AddAction(RemoteControlField.ToSharedRef());
 					}
 				}
 			}
diff --git a/Engine/Plugins/VirtualProduction/RemoteControl/Source/RemoteControlUI/Private/UI/Behaviour/Builtin/RangeMap/RCBehaviourRangeMapModel.cpp b/Engine/Plugins/VirtualProduction/RemoteControl/Source/RemoteControlUI/Private/UI/Behaviour/Builtin/RangeMap/RCBehaviourRangeMapModel.cpp
new file mode 100644
--- /dev/null
+++ b/Engine/Plugins/VirtualProduction/RemoteControl/Source/RemoteControlUI/Private/UI/Behaviour/Builtin/RangeMap/RCBehaviourRangeMapModel.cpp
@@ -0,0 +1,232 @@
+// Copyright Epic Games, Inc. All Rights Reserved.
+
+#include "RCBehaviourRangeMapModel.h"
+
+#include "Action/RCAction.h"
+#include "Action/RCPropertyAction.h"
+#include "Controller/RCController.h"
+#include "IDetailTreeNode.h"
+#include "IPropertyRowGenerator.h"
+#include "Modules/ModuleManager.h"
+#include "PropertyEditorModule.h"
+#include "RCVirtualProperty.h"
+#include "RCVirtualPropertyContainer.h"
+#include "RemoteControlPreset.h"
+#include "SRCBehaviourRangeMap.h"
+#include "UI/Action/RangeMap/RCActionRangeMapModel.h"
+#include "UI/Action/SRCActionPanel.h"
+#include "UI/Action/SRCActionPanelList.h"
+#include "UI/RemoteControlPanelStyle.h"
+#include "Widgets/Layout/SBorder.h"
+#include "Widgets/Text/STextBlock.h"
+
+#define LOCTEXT_NAMESPACE "FRCRangeMapBehaviourModel"
+
+FRCRangeMapBehaviourModel::FRCRangeMapBehaviourModel(URCRangeMapBehaviour* RangeMapBehaviour)
+	: FRCBehaviourModel(RangeMapBehaviour)
+	, RangeMapBehaviourWeakPtr(RangeMapBehaviour)
+{
+	FPropertyRowGeneratorArgs Args;
+	Args.bShouldShowHiddenProperties = true;
+	PropertyRowGenerator = FModuleManager::GetModuleChecked<FPropertyEditorModule>("PropertyEditor").CreatePropertyRowGenerator(Args);
+
+	if (RangeMapBehaviour)
+	{
+		PropertyRowGenerator->SetStructure(RangeMapBehaviour->PropertyContainer->CreateStructOnScope());
+
+		for (const TSharedRef<IDetailTreeNode>& CategoryNode : PropertyRowGenerator->GetRootTreeNodes())
+		{
+			TArray<TSharedRef<IDetailTreeNode>> Children;
+			CategoryNode->GetChildren(Children);
+			for (const TSharedRef<IDetailTreeNode>& Child : Children)
+			{
+				DetailTreeNodeWeakPtrArray.Add(Child);
+			}
+		}
+	}
+}
+
+TSharedRef<SWidget> FRCRangeMapBehaviourModel::GetBehaviourDetailsWidget()
+{
+	return SNew(SRCBehaviourRangeMap, SharedThis(this));
+}
+
+TSharedRef<SWidget> FRCRangeMapBehaviourModel::GetPropertyWidget() const
+{
+	const TSharedRef<SVerticalBox> FieldWidget = SNew(SVerticalBox);
+	
+	auto DetailTreeItr = DetailTreeNodeWeakPtrArray.CreateConstIterator();
+
+	TSharedPtr<IDetailTreeNode> DetailTreeMinValue = DetailTreeItr->Pin();
+	TSharedPtr<IDetailTreeNode> DetailTreeMaxValue = (++DetailTreeItr)->Pin();
+	FNodeWidgets MinValueWidget = DetailTreeMinValue->CreateNodeWidgets();
+	FNodeWidgets MaxValueWidget = DetailTreeMaxValue->CreateNodeWidgets();
+
+	if (MinValueWidget.ValueWidget && MaxValueWidget.ValueWidget)
+	{
+		FieldWidget->AddSlot()
+			.Padding(FMargin(3.0f, 2.0f))
+			.HAlign(HAlign_Left)
+			.AutoHeight()
+			[
+				CreateMinMaxWidget(DetailTreeMinValue, DetailTreeMaxValue)
+			];
+	}
+	
+	TSharedPtr<IDetailTreeNode> DetailTreeNodeThreshold = (++DetailTreeItr)->Pin();
+	TSharedPtr<IDetailTreeNode> DetailTreeNodeStep = (++DetailTreeItr)->Pin();
+	const FNodeWidgets ThresholdWidget = DetailTreeNodeThreshold->CreateNodeWidgets();
+	const FNodeWidgets StepWidget = DetailTreeNodeStep->CreateNodeWidgets();
+	
+	FieldWidget->AddSlot()
+		.Padding(FMargin(3.0f, 2.0f))
+		.HAlign(HAlign_Center)
+		.AutoHeight()
+		[
+			SNew(SHorizontalBox)
+			+SHorizontalBox::Slot()
+			.Padding(FMargin(3.0f, 2.0f))
+			.HAlign(HAlign_Right)
+			.AutoWidth()
+			[
+				ThresholdWidget.NameWidget.ToSharedRef()
+			]
+			+SHorizontalBox::Slot()
+			.Padding(FMargin(3.0f, 2.0f))
+			.HAlign(HAlign_Right)
+			.AutoWidth()
+			[
+				ThresholdWidget.ValueWidget.ToSharedRef()
+			]
+		];
+
+	FieldWidget->AddSlot()
+		.Padding(FMargin(3.0f, 6.0f))
+		.HAlign(HAlign_Right)
+		.AutoHeight()
+		[
+			SNew(SBorder)
+			[
+				SNew(SHorizontalBox)
+				+SHorizontalBox::Slot()
+				.Padding(FMargin(3.0f, 2.0f))
+				.HAlign(HAlign_Right)
+				.AutoWidth()
+				[
+					SNew(STextBlock)
+					.Text(LOCTEXT("StepInputValue", "Step Value for Action"))
+					.Font(FRemoteControlPanelStyle::Get()->GetFontStyle("RemoteControlPanel.Behaviours.BehaviourDescription"))
+				]
+				+SHorizontalBox::Slot()
+				.Padding(FMargin(3.0f, 2.0f))
+				.HAlign(HAlign_Right)
+				.FillWidth(1.0f)
+				.AutoWidth()
+				[
+					StepWidget.ValueWidget.ToSharedRef()
+				]
+			]
+		];
+
+	return 	SNew(SHorizontalBox)
+		.Clipping(EWidgetClipping::OnDemand)
+		+ SHorizontalBox::Slot()
+		.VAlign(VAlign_Top)
+		.FillWidth(1.0f)
+		.AutoWidth()
+		[
+			FieldWidget
+		];
+}
+
+void FRCRangeMapBehaviourModel::OnActionAdded(URCAction* Action)
+{
+	if (URCRangeMapBehaviour* RangeMapBehaviour = Cast<URCRangeMapBehaviour>(GetBehaviour()))
+	{
+		if (URCPropertyAction* PropertyAction = Cast<URCPropertyAction>(Action))
+		{
+			RangeMapBehaviour->OnActionAdded(Action,PropertyAction->PropertySelfContainer);
+		}
+	}
+}
+
+TSharedPtr<SRCLogicPanelListBase> FRCRangeMapBehaviourModel::GetActionsListWidget(TSharedRef<SRCActionPanel> InActionPanel)
+{
+	return SNew(SRCActionPanelList<FRCActionRangeMapModel>, InActionPanel, SharedThis(this));
+}
+
+URCAction* FRCRangeMapBehaviourModel::AddAction(const TSharedRef<const FRemoteControlField> InRemoteControlField)
+{
+	URCAction* NewAction = nullptr;
+
+	if (URCRangeMapBehaviour* RangeMapBehaviour = Cast<URCRangeMapBehaviour>(GetBehaviour()))
+	{
+		double StepValue;
+		RangeMapBehaviour->PropertyContainer->GetVirtualProperty(FName("Step"))->GetValueDouble(StepValue);
+
+		const FGuid FieldId = InRemoteControlField->GetId();
+		if (TSharedPtr<FRemoteControlProperty> RemoteProperty = RangeMapBehaviour->ControllerWeakPtr.Get()->PresetWeakPtr.Get()->GetExposedEntity<FRemoteControlProperty>(FieldId).Pin())
+		{
+			TObjectPtr<URCVirtualPropertySelfContainer> VirtualPropertySelfContainer = NewObject<URCVirtualPropertySelfContainer>();
+			VirtualPropertySelfContainer->DuplicateProperty(RemoteProperty->GetProperty()->GetFName(), RemoteProperty->GetProperty());
+			
+			NewAction = RangeMapBehaviour->AddAction(InRemoteControlField);
+
+			OnActionAdded(NewAction);
+		}
+	}
+
+	return NewAction;
+}
+
+TSharedRef<SWidget> FRCRangeMapBehaviourModel::CreateMinMaxWidget(TSharedPtr<IDetailTreeNode> MinRangeDetailTree, TSharedPtr<IDetailTreeNode> MaxRangeDetailTree) const
+{
+	const FNodeWidgets MinValueWidget = MinRangeDetailTree->CreateNodeWidgets();
+	const FNodeWidgets MaxValueWidget = MaxRangeDetailTree->CreateNodeWidgets();
+	
+	return SNew(SHorizontalBox)
+		+SHorizontalBox::Slot()
+		.Padding(FMargin(3.0f, 2.0f))
+		.HAlign(HAlign_Right)
+		.AutoWidth()
+		[
+			SNew(SVerticalBox)
+			+SVerticalBox::Slot()
+			.Padding(FMargin(3.0f, 2.0f))
+			.HAlign(HAlign_Right)
+			.AutoHeight()
+			[
+				MinValueWidget.NameWidget.ToSharedRef()
+			]
+			+SVerticalBox::Slot()
+			.Padding(FMargin(3.0f, 2.0f))
+			.HAlign(HAlign_Right)
+			.AutoHeight()
+			[
+				MinValueWidget.ValueWidget.ToSharedRef()
+			]
+		]
+		+SHorizontalBox::Slot()
+		.Padding(FMargin(3.0f, 2.0f))
+		.HAlign(HAlign_Right)
+		.AutoWidth()
+		[
+			SNew(SVerticalBox)
+			+SVerticalBox::Slot()
+			.Padding(FMargin(3.0f, 2.0f))
+			.HAlign(HAlign_Right)
+			.AutoHeight()
+			[
+				MaxValueWidget.NameWidget.ToSharedRef()
+			]
+			+SVerticalBox::Slot()
+			.Padding(FMargin(3.0f, 2.0f))
+			.HAlign(HAlign_Right)
+			.AutoHeight()
+			[
+				MaxValueWidget.ValueWidget.ToSharedRef()
+			]
+		];
+}
+
+#undef LOCTEXT_NAMESPACE
diff --git a/Engine/Plugins/VirtualProduction/RemoteControl/Source/RemoteControlUI/Private/UI/Behaviour/Builtin/RangeMap/RCBehaviourRangeMapModel.h b/Engine/Plugins/VirtualProduction/RemoteControl/Source/RemoteControlUI/Private/UI/Behaviour/Builtin/RangeMap/RCBehaviourRangeMapModel.h
new file mode 100644
--- /dev/null
+++ b/Engine/Plugins/VirtualProduction/RemoteControl/Source/RemoteControlUI/Private/UI/Behaviour/Builtin/RangeMap/RCBehaviourRangeMapModel.h
@@ -0,0 +1,54 @@
+// Copyright Epic Games, Inc. All Rights Reserved.
+
+#pragma once
+#include "Behaviour/Builtin/RangeMap/RCRangeMapBehaviour.h"
+#include "UI/Behaviour/RCBehaviourModel.h"
+
+class IDetailTreeNode;
+class IPropertyRowGenerator;
+class SWidget;
+
+/*
+* ~ FRCRangeMapBehaviourModel ~
+*
+* Child Behaviour class representing the "Range Mapping" Behaviour's UI model
+* 
+* Generates a Property Widget where users can enter the Max, Min, Threshhold, and Step Values.
+* The values of them are used for the interaction implementation of the properties it affects..
+*/
+class FRCRangeMapBehaviourModel : public FRCBehaviourModel
+{
+public:
+	FRCRangeMapBehaviourModel(URCRangeMapBehaviour* RangeMapBehaviour);
+
+	/** Builds a Behaviour specific widget as required for Range Mapping behaviour*/
+	virtual TSharedRef<SWidget> GetBehaviourDetailsWidget() override;
+
+	/** Builds a generic Value Widget of the Controller's type
+	* Use to store user input for performing the Range Mapping Behaviour Interactions */
+	TSharedRef<SWidget> GetPropertyWidget() const;
+
+	/** Add a Logic Action using a remote control field as input */
+	virtual URCAction* AddAction(const TSharedRef<const FRemoteControlField> InRemoteControlField) override;
+
+	/** The Actions List Widget used in the RangeMap Behaviour */
+	virtual TSharedPtr<SRCLogicPanelListBase> GetActionsListWidget(TSharedRef<SRCActionPanel> InActionPanel) override;
+
+protected:
+	/** Invoked by parent class after an action has been added from UI*/
+	virtual void OnActionAdded(URCAction* Action) override;
+
+private:
+	/** The Range Mapping Behaviour (Data model) associated with us*/
+	TWeakObjectPtr<URCRangeMapBehaviour> RangeMapBehaviourWeakPtr;
+
+	/** The row generator used to build a generic Value Widget*/
+	TSharedPtr<IPropertyRowGenerator> PropertyRowGenerator;
+
+	/** Used to create generic Value Widgets for lerping control of values */
+	TArray<TWeakPtr<IDetailTreeNode>> DetailTreeNodeWeakPtrArray;
+
+private:
+	/** Creates the horizontal box containing box the Minimum and Maximum Range Widgets */
+	TSharedRef<SWidget> CreateMinMaxWidget(TSharedPtr<IDetailTreeNode> MinRangeDetailTree, TSharedPtr<IDetailTreeNode> MaxRangeDetailTree) const;
+};
diff --git a/Engine/Plugins/VirtualProduction/RemoteControl/Source/RemoteControlUI/Private/UI/Behaviour/Builtin/RangeMap/SRCBehaviourRangeMap.cpp b/Engine/Plugins/VirtualProduction/RemoteControl/Source/RemoteControlUI/Private/UI/Behaviour/Builtin/RangeMap/SRCBehaviourRangeMap.cpp
new file mode 100644
--- /dev/null
+++ b/Engine/Plugins/VirtualProduction/RemoteControl/Source/RemoteControlUI/Private/UI/Behaviour/Builtin/RangeMap/SRCBehaviourRangeMap.cpp
@@ -0,0 +1,44 @@
+// Copyright Epic Games, Inc. All Rights Reserved.
+
+#include "SRCBehaviourRangeMap.h"
+
+#include "SPositiveActionButton.h"
+#include "SlateOptMacros.h"
+#include "Styling/CoreStyle.h"
+#include "UI/Behaviour/Builtin/RangeMap/RCBehaviourRangeMapModel.h"
+#include "UI/RemoteControlPanelStyle.h"
+#include "Widgets/Images/SImage.h"
+#include "Widgets/Input/SButton.h"
+#include "Widgets/Layout/SSeparator.h"
+
+#define LOCTEXT_NAMESPACE "SRCBehaviourRangeMap"
+
+class SPositiveActionButton;
+
+void SRCBehaviourRangeMap::Construct(const FArguments& InArgs, TSharedRef<const FRCRangeMapBehaviourModel> InBehaviourItem)
+{
+	RangeMapBehaviourItemWeakPtr = InBehaviourItem;
+
+	ChildSlot
+	[
+		SNew(SHorizontalBox)
+
+		// PropertyWidget
+		+ SHorizontalBox::Slot()
+		.HAlign(HAlign_Left)
+		.VAlign(VAlign_Center)
+		.Padding(FMargin(5.f))
+		[
+			SNew(SHorizontalBox)
+			// Value Widget
+			+ SHorizontalBox::Slot()
+			.VAlign(VAlign_Center)
+			.AutoWidth()
+			[
+				InBehaviourItem->GetPropertyWidget()
+			]
+		]
+	];
+}
+
+#undef LOCTEXT_NAMESPACE
\ No newline at end of file
diff --git a/Engine/Plugins/VirtualProduction/RemoteControl/Source/RemoteControlUI/Private/UI/Behaviour/Builtin/RangeMap/SRCBehaviourRangeMap.h b/Engine/Plugins/VirtualProduction/RemoteControl/Source/RemoteControlUI/Private/UI/Behaviour/Builtin/RangeMap/SRCBehaviourRangeMap.h
new file mode 100644
--- /dev/null
+++ b/Engine/Plugins/VirtualProduction/RemoteControl/Source/RemoteControlUI/Private/UI/Behaviour/Builtin/RangeMap/SRCBehaviourRangeMap.h
@@ -0,0 +1,30 @@
+// Copyright Epic Games, Inc. All Rights Reserved.
+
+#pragma once
+
+#include "CoreMinimal.h"
+#include "RemoteControlField.h"
+#include "Widgets/SCompoundWidget.h"
+
+class FRCRangeMapBehaviourModel;
+
+/*
+ * ~ SRCBehaviourRangeMap
+ * 
+ * Behaviour specific details panel for Range Mapping Behaviour.
+ * Contains Input Panels which will be used to calculate and work with in the Behaviour itself.
+ */
+class REMOTECONTROLUI_API SRCBehaviourRangeMap : public SCompoundWidget
+{
+public:
+	SLATE_BEGIN_ARGS(SRCBehaviourRangeMap){}
+
+	SLATE_END_ARGS()
+
+	/** Constructs this widget with InArgs */
+	void Construct(const FArguments& InArgs, TSharedRef<const FRCRangeMapBehaviourModel> InBehaviourItem);
+
+private:
+	/** The Behaviour (UI Model) associated with this. */
+	TWeakPtr<const FRCRangeMapBehaviourModel> RangeMapBehaviourItemWeakPtr;
+};
diff --git a/Engine/Plugins/VirtualProduction/RemoteControl/Source/RemoteControlUI/Private/UI/Behaviour/SRCBehaviourPanelList.cpp b/Engine/Plugins/VirtualProduction/RemoteControl/Source/RemoteControlUI/Private/UI/Behaviour/SRCBehaviourPanelList.cpp
--- a/Engine/Plugins/VirtualProduction/RemoteControl/Source/RemoteControlUI/Private/UI/Behaviour/SRCBehaviourPanelList.cpp
+++ b/Engine/Plugins/VirtualProduction/RemoteControl/Source/RemoteControlUI/Private/UI/Behaviour/SRCBehaviourPanelList.cpp
@@ -16,6 +16,7 @@
 #include "UI/Behaviour/Builtin/RCBehaviourIsEqualModel.h"
 #include "UI/Behaviour/Builtin/RCBehaviourSetAssetByPathModel.h"
 #include "UI/Controller/RCControllerModel.h"
+#include "UI/Behaviour/Builtin/RangeMap/RCBehaviourRangeMapModel.h"
 #include "UI/RemoteControlPanelStyle.h"
 #include "UI/SRemoteControlPanel.h"
 #include "Widgets/Views/SHeaderRow.h"
@@ -107,6 +108,10 @@ void SRCBehaviourPanelList::Reset()
 				{
 					BehaviourItems.Add(MakeShared<FRCBehaviourBindModel>(BindBehaviour, RemoteControlPanel));
 				}
+				else if (URCRangeMapBehaviour* RangeMapBehaviour = Cast<URCRangeMapBehaviour>(Behaviour))
+				{
+					BehaviourItems.Add(MakeShared<FRCRangeMapBehaviourModel>(RangeMapBehaviour));
+				}
 				else
 				{
 					BehaviourItems.Add(MakeShared<FRCBehaviourModel>(Behaviour, RemoteControlPanel));
