Controller Panel improvements.

Notes:
1. Add Controller Menu
    a) Revised list of Controller types
    b) Color coding added for menu entries
    c) Controller default names improved

2. Controller reordering support
    a) Drag handles, drag-and-drop support added
    b) Dispaly index is stored as part of Virtual Property UObject
    c) True reordering support is provided (rows are not swapped, but pushed down)

diff --git a/Engine/Plugins/VirtualProduction/RemoteControl/Config/BaseRemoteControl.ini b/Engine/Plugins/VirtualProduction/RemoteControl/Config/BaseRemoteControl.ini
--- a/Engine/Plugins/VirtualProduction/RemoteControl/Config/BaseRemoteControl.ini
+++ b/Engine/Plugins/VirtualProduction/RemoteControl/Config/BaseRemoteControl.ini
@@ -3,14 +3,10 @@
 +PropertyRedirects=(OldName="RemoteControlField.Id",NewName="RemoteControlEntity.Id")
 
 [/Script/RemoteControlLogic.RemoteControlLogicConfig]
-+SupportedControllerTypes=EPropertyBagPropertyType::Bool
-+SupportedControllerTypes=EPropertyBagPropertyType::Byte
++SupportedControllerTypes=EPropertyBagPropertyType::String
 +SupportedControllerTypes=EPropertyBagPropertyType::Int32
 +SupportedControllerTypes=EPropertyBagPropertyType::Float
-+SupportedControllerTypes=EPropertyBagPropertyType::Double
-+SupportedControllerTypes=EPropertyBagPropertyType::Name
-+SupportedControllerTypes=EPropertyBagPropertyType::String
-+SupportedControllerTypes=EPropertyBagPropertyType::Text
++SupportedControllerTypes=EPropertyBagPropertyType::Bool
 +SupportedControllerTypes=EPropertyBagPropertyType::Struct
 +SupportedControllerStructTypes=Vector
 +SupportedControllerStructTypes=Color
diff --git a/Engine/Plugins/VirtualProduction/RemoteControl/Source/RemoteControl/Private/RCVirtualProperty.cpp b/Engine/Plugins/VirtualProduction/RemoteControl/Source/RemoteControl/Private/RCVirtualProperty.cpp
--- a/Engine/Plugins/VirtualProduction/RemoteControl/Source/RemoteControl/Private/RCVirtualProperty.cpp
+++ b/Engine/Plugins/VirtualProduction/RemoteControl/Source/RemoteControl/Private/RCVirtualProperty.cpp
@@ -625,6 +625,29 @@ FName URCVirtualPropertyBase::GetPropertyName() const
 	return PropertyName;
 }
 
+FName URCVirtualPropertyBase::GetVirtualPropertyTypeDisplayName(const EPropertyBagPropertyType InValueType, UObject* InValueTypeObject)
+{
+	const UEnum* Enum = FindObjectChecked<UEnum>(nullptr, TEXT("/Script/StructUtils.EPropertyBagPropertyType"));
+	FString PropertyTypeString = Enum->GetValueAsName(InValueType).ToString();
+	PropertyTypeString.RemoveFromStart("EPropertyBagPropertyType::");
+
+	// User-friendly typename alias
+	if (PropertyTypeString == "Bool")
+	{
+		PropertyTypeString = "Boolean";
+	}
+	else if (PropertyTypeString == "Int32")
+	{
+		PropertyTypeString = TEXT("Integer");
+	}
+	
+	// For Structs we should derive the prefix from the subobject type (Vector / Color / Rotator / etc) rather than the property bag type (which is always "Struct")
+	if (InValueTypeObject)
+		PropertyTypeString = InValueTypeObject->GetName();
+
+	return *PropertyTypeString;
+}
+
 const FInstancedPropertyBag* URCVirtualPropertyInContainer::GetPropertyBagInstance() const
 {
 	if (ContainerWeakPtr.Get())
diff --git a/Engine/Plugins/VirtualProduction/RemoteControl/Source/RemoteControl/Private/RCVirtualPropertyContainer.cpp b/Engine/Plugins/VirtualProduction/RemoteControl/Source/RemoteControl/Private/RCVirtualPropertyContainer.cpp
--- a/Engine/Plugins/VirtualProduction/RemoteControl/Source/RemoteControl/Private/RCVirtualPropertyContainer.cpp
+++ b/Engine/Plugins/VirtualProduction/RemoteControl/Source/RemoteControl/Private/RCVirtualPropertyContainer.cpp
@@ -23,7 +23,7 @@ URCVirtualPropertyInContainer* URCVirtualPropertyContainerBase::AddProperty(cons
 	MarkPackageDirty();
 #endif
 
-	const FName PropertyName = GenerateUniquePropertyName(InPropertyName, InValueType, this);
+	const FName PropertyName = GenerateUniquePropertyName(InPropertyName, InValueType, InValueTypeObject, this);
 	Bag.AddProperty(PropertyName, InValueType, InValueTypeObject);
 
 	const FPropertyBagPropertyDesc* BagPropertyDesc = Bag.FindPropertyDescByName(InPropertyName);
@@ -152,6 +152,19 @@ URCVirtualPropertyBase* URCVirtualPropertyContainerBase::GetVirtualProperty(cons
 	return nullptr;
 }
 
+URCVirtualPropertyBase* URCVirtualPropertyContainerBase::GetVirtualProperty(const FGuid& InId) const
+{
+	for (URCVirtualPropertyBase* VirtualProperty : VirtualProperties)
+	{
+		if (VirtualProperty->Id == InId)
+		{
+			return VirtualProperty;
+		}
+	}
+
+	return nullptr;
+}
+
 URCVirtualPropertyBase* URCVirtualPropertyContainerBase::GetVirtualPropertyByDisplayName(const FName InDisplayName) const
 {
 	for (URCVirtualPropertyBase* VirtualProperty : VirtualProperties)
@@ -180,7 +193,7 @@ TSharedPtr<FStructOnScope> URCVirtualPropertyContainerBase::CreateStructOnScope(
 	return MakeShared<FStructOnScope>(Bag.GetPropertyBagStruct(), Bag.GetMutableValue().GetMutableMemory());
 }
 
-FName URCVirtualPropertyContainerBase::GenerateUniquePropertyName(const FName& InPropertyName, const EPropertyBagPropertyType InValueType, const URCVirtualPropertyContainerBase* InContainer)
+FName URCVirtualPropertyContainerBase::GenerateUniquePropertyName(const FName& InPropertyName, const EPropertyBagPropertyType InValueType, UObject* InValueTypeObject, const URCVirtualPropertyContainerBase* InContainer)
 {
 	auto GetFinalName = [](const FString& InPrefix, const int32 InIndex = 0)
 	{
@@ -194,12 +207,11 @@ FName URCVirtualPropertyContainerBase::GenerateUniquePropertyName(const FName& I
 		return FinalName;
 	};
 
-	const UEnum* Enum = FindObjectChecked<UEnum>(nullptr, TEXT("/Script/StructUtils.EPropertyBagPropertyType"));
-	FString DefaultName = Enum->GetValueAsName(InValueType).ToString();	
-	DefaultName.RemoveFromStart("EPropertyBagPropertyType::");
+
+	FName DefaultName = URCVirtualPropertyBase::GetVirtualPropertyTypeDisplayName(InValueType, InValueTypeObject);
 	
 	int32 Index = 0;
-	const FString InitialName = InPropertyName.IsNone() ? DefaultName : InPropertyName.ToString();
+	const FString InitialName = InPropertyName.IsNone() ? DefaultName.ToString() : InPropertyName.ToString();
 	FString FinalName = InitialName;
 	
 	// Recursively search for an available name by incrementing suffix till we find one.
diff --git a/Engine/Plugins/VirtualProduction/RemoteControl/Source/RemoteControl/Private/RemoteControlPreset.cpp b/Engine/Plugins/VirtualProduction/RemoteControl/Source/RemoteControl/Private/RemoteControlPreset.cpp
--- a/Engine/Plugins/VirtualProduction/RemoteControl/Source/RemoteControl/Private/RemoteControlPreset.cpp
+++ b/Engine/Plugins/VirtualProduction/RemoteControl/Source/RemoteControl/Private/RemoteControlPreset.cpp
@@ -519,6 +519,16 @@ URCVirtualPropertyBase* URemoteControlPreset::GetVirtualProperty(const FName InP
 	return ControllerContainer->GetVirtualProperty(InPropertyName);
 }
 
+URCVirtualPropertyBase* URemoteControlPreset::GetVirtualProperty(const FGuid& InId) const
+{
+	if (!ensure(ControllerContainer))
+	{
+		return nullptr;
+	}
+
+	return ControllerContainer->GetVirtualProperty(InId);
+}
+
 URCVirtualPropertyBase* URemoteControlPreset::GetVirtualPropertyByDisplayName(const FName InDisplayName) const
 {
 	if (!ensure(ControllerContainer))
@@ -540,7 +550,7 @@ URCVirtualPropertyInContainer* URemoteControlPreset::AddVirtualProperty(TSubclas
 	FName NewPropertyName = InPropertyName;
 	if (NewPropertyName.IsNone())
 	{
-		NewPropertyName = URCVirtualPropertyContainerBase::GenerateUniquePropertyName(TEXT(""), InValueType, ControllerContainer);
+		NewPropertyName = URCVirtualPropertyContainerBase::GenerateUniquePropertyName(TEXT(""), InValueType, InValueTypeObject, ControllerContainer);
 	}	
 
 	return ControllerContainer->AddProperty(NewPropertyName, InPropertyClass, InValueType, InValueTypeObject);
diff --git a/Engine/Plugins/VirtualProduction/RemoteControl/Source/RemoteControl/Public/RCVirtualProperty.h b/Engine/Plugins/VirtualProduction/RemoteControl/Source/RemoteControl/Public/RCVirtualProperty.h
--- a/Engine/Plugins/VirtualProduction/RemoteControl/Source/RemoteControl/Public/RCVirtualProperty.h
+++ b/Engine/Plugins/VirtualProduction/RemoteControl/Source/RemoteControl/Public/RCVirtualProperty.h
@@ -261,6 +261,9 @@ public:
 	UFUNCTION(BlueprintPure, Category = "Remote Control Behaviour")
 	FName GetPropertyName() const;
 
+	/** Fetches a user-friendly representation of the Virtual Property's type name (i.e. Controller type), given a Property Bag type and value object (for structures)*/
+	static FName GetVirtualPropertyTypeDisplayName(const EPropertyBagPropertyType InValueType, UObject* InValueTypeObject);
+
 public:
 	/** Unique property name */
 	UPROPERTY()
@@ -277,6 +280,10 @@ public:
 	/** User friendly name of the Controller*/
 	UPROPERTY()
 	FName DisplayName;
+
+	/** User configurable Display Index for this Virtual Property (as Logic Controller) when represented as a row in the RC Logic Controllers list  */
+	UPROPERTY()
+	int32 DisplayIndex;
 };
 
 /**
diff --git a/Engine/Plugins/VirtualProduction/RemoteControl/Source/RemoteControl/Public/RCVirtualPropertyContainer.h b/Engine/Plugins/VirtualProduction/RemoteControl/Source/RemoteControl/Public/RCVirtualPropertyContainer.h
--- a/Engine/Plugins/VirtualProduction/RemoteControl/Source/RemoteControl/Public/RCVirtualPropertyContainer.h
+++ b/Engine/Plugins/VirtualProduction/RemoteControl/Source/RemoteControl/Public/RCVirtualPropertyContainer.h
@@ -75,6 +75,11 @@ public:
 	 */ 
 	virtual URCVirtualPropertyBase* GetVirtualProperty(const FName InPropertyName) const;
 
+	/**
+	 * Returns virtual property by unique Id.
+	 */
+	URCVirtualPropertyBase* GetVirtualProperty(const FGuid& InId) const;
+
 	/**
 	 * Returns virtual property by user-friendly display name (Controller Name)
 	 */
@@ -93,7 +98,7 @@ public:
 	/**
 	 * Generates unique name for the property for specified property container
 	 */
-	static FName GenerateUniquePropertyName(const FName& InPropertyName, const EPropertyBagPropertyType InValueType, const URCVirtualPropertyContainerBase* InContainer);
+	static FName GenerateUniquePropertyName(const FName& InPropertyName, const EPropertyBagPropertyType InValueType, UObject* InValueTypeObject, const URCVirtualPropertyContainerBase* InContainer);
 
 #if WITH_EDITOR
 	/** Delegate when object changed */
diff --git a/Engine/Plugins/VirtualProduction/RemoteControl/Source/RemoteControl/Public/RemoteControlPreset.h b/Engine/Plugins/VirtualProduction/RemoteControl/Source/RemoteControl/Public/RemoteControlPreset.h
--- a/Engine/Plugins/VirtualProduction/RemoteControl/Source/RemoteControl/Public/RemoteControlPreset.h
+++ b/Engine/Plugins/VirtualProduction/RemoteControl/Source/RemoteControl/Public/RemoteControlPreset.h
@@ -621,6 +621,9 @@ public:
 	/** Fetches a virtual property by internal property name. */
 	URCVirtualPropertyBase* GetVirtualProperty(const FName InPropertyName) const;
 
+	/** Fetches a virtual property by unique Id. */
+	URCVirtualPropertyBase* GetVirtualProperty(const FGuid& InId) const;
+
 	/** Fetches a virtual property by specified name. */
 	URCVirtualPropertyBase* GetVirtualPropertyByDisplayName(const FName InDisplayName) const;
 
diff --git a/Engine/Plugins/VirtualProduction/RemoteControl/Source/RemoteControlUI/Private/UI/Controller/RCControllerModel.cpp b/Engine/Plugins/VirtualProduction/RemoteControl/Source/RemoteControlUI/Private/UI/Controller/RCControllerModel.cpp
--- a/Engine/Plugins/VirtualProduction/RemoteControl/Source/RemoteControlUI/Private/UI/Controller/RCControllerModel.cpp
+++ b/Engine/Plugins/VirtualProduction/RemoteControl/Source/RemoteControlUI/Private/UI/Controller/RCControllerModel.cpp
@@ -13,7 +13,7 @@
 FRCControllerModel::FRCControllerModel(URCVirtualPropertyBase* InVirtualProperty, const TSharedRef<IDetailTreeNode>& InTreeNode, const TSharedPtr<SRemoteControlPanel> InRemoteControlPanel)
 	: FRCLogicModeBase(InRemoteControlPanel)
 	, VirtualPropertyWeakPtr(InVirtualProperty)
-	, DetailTreeNodeWeakPtr(InTreeNode)	
+	, DetailTreeNodeWeakPtr(InTreeNode)
 {
 	if(ensure(InVirtualProperty))
 	{
@@ -24,6 +24,8 @@ FRCControllerModel::FRCControllerModel(URCVirtualPropertyBase* InVirtualProperty
 			.Text(FText::FromName(InVirtualProperty->DisplayName))
 			.OnTextCommitted_Raw(this, &FRCControllerModel::OnControllerNameCommitted);
 	}
+
+	Id = FGuid::NewGuid();
 }
 
 TSharedRef<SWidget> FRCControllerModel::GetWidget() const
diff --git a/Engine/Plugins/VirtualProduction/RemoteControl/Source/RemoteControlUI/Private/UI/Controller/RCControllerModel.h b/Engine/Plugins/VirtualProduction/RemoteControl/Source/RemoteControlUI/Private/UI/Controller/RCControllerModel.h
--- a/Engine/Plugins/VirtualProduction/RemoteControl/Source/RemoteControlUI/Private/UI/Controller/RCControllerModel.h
+++ b/Engine/Plugins/VirtualProduction/RemoteControl/Source/RemoteControlUI/Private/UI/Controller/RCControllerModel.h
@@ -53,6 +53,9 @@ public:
 	/** User-friendly Name of the underlying Controller */
 	FName GetControllerDisplayName();
 
+	/**Fetches the unique Id for this UI item */
+	FGuid GetId() const { return Id; }
+
 private:
 	/* Text commit event for Controller Name text box*/
 	void OnControllerNameCommitted(const FText& InNewControllerName, ETextCommit::Type InCommitInfo);
@@ -68,4 +71,7 @@ private:
 
 	/** Controller name - editable text box */
 	TSharedPtr<SInlineEditableTextBlock> ControllerNameTextBox;
+
+	/**Unique Id for this UI item*/
+	FGuid Id;
 };
diff --git a/Engine/Plugins/VirtualProduction/RemoteControl/Source/RemoteControlUI/Private/UI/Controller/SRCControllerPanel.cpp b/Engine/Plugins/VirtualProduction/RemoteControl/Source/RemoteControlUI/Private/UI/Controller/SRCControllerPanel.cpp
--- a/Engine/Plugins/VirtualProduction/RemoteControl/Source/RemoteControlUI/Private/UI/Controller/SRCControllerPanel.cpp
+++ b/Engine/Plugins/VirtualProduction/RemoteControl/Source/RemoteControlUI/Private/UI/Controller/SRCControllerPanel.cpp
@@ -15,6 +15,7 @@
 #include "UI/BaseLogicUI/RCLogicModeBase.h"
 #include "UI/Panels/SRCDockPanel.h"
 #include "UI/RemoteControlPanelStyle.h"
+#include "UI/RCUIHelpers.h"
 #include "UI/SRemoteControlPanel.h"
 #include "Widgets/SBoxPanel.h"
 #include "Widgets/Input/SComboButton.h"
@@ -180,29 +181,50 @@ TSharedRef<SWidget> SRCControllerPanel::GetControllerMenuContentWidget() const
 	for (const TPair<EPropertyBagPropertyType, UObject*>& Pair : VirtualPropertyFieldClassNames)
 	{
 		// Display Name
-		const UEnum* Enum = FindObjectChecked<UEnum>(nullptr, TEXT("/Script/StructUtils.EPropertyBagPropertyType"));
-		FString PropertyTypeString = Enum->GetValueAsName(Pair.Key).ToString();
-		PropertyTypeString.RemoveFromStart("EPropertyBagPropertyType::");
+		const FName DefaultName = URCVirtualPropertyBase::GetVirtualPropertyTypeDisplayName(Pair.Key, Pair.Value);
 
-		if (PropertyTypeString.StartsWith(TEXT("Bool")))
-		{
-			PropertyTypeString = "Boolean";
-		}
-		
-		if (PropertyTypeString.StartsWith(TEXT("Int")))
+		// Type Color
+		FLinearColor TypeColor = FColor::White;
+
+		// Generate a transient virtual property for deducing type color:
+		FInstancedPropertyBag Bag;
+		Bag.AddProperty(DefaultName, Pair.Key, Pair.Value);
+		const FPropertyBagPropertyDesc* BagPropertyDesc = Bag.FindPropertyDescByName(DefaultName);
+		if (ensure(BagPropertyDesc))
 		{
-			PropertyTypeString = PropertyTypeString.Replace(TEXT("Int"), TEXT("Integer"));
+			TypeColor = UE::RCUIHelpers::GetFieldClassTypeColor(BagPropertyDesc->CachedProperty);
 		}
 
-		const FName DefaultName = Pair.Value ? Pair.Value->GetFName() : *PropertyTypeString;
+		// Variable Color Bar
+		TSharedRef<SWidget> MenuItemWidget = 
+			SNew(SHorizontalBox)
+			+ SHorizontalBox::Slot()
+			.AutoWidth()
+			.HAlign(HAlign_Left)
+			.VAlign(VAlign_Fill)
+			[
+				SNew(SBox)
+				.HeightOverride(5.f)
+				[
+					SNew(SBorder)
+					.Visibility(EVisibility::HitTestInvisible)
+					.BorderImage(FAppStyle::Get().GetBrush("NumericEntrySpinBox.NarrowDecorator"))
+					.BorderBackgroundColor(TypeColor)
+					.Padding(FMargin(5.0f, 0.0f, 0.0f, 0.f))
+				]
+			]
+		+ SHorizontalBox::Slot()
+			.AutoWidth()
+			.Padding(FMargin(6.f, 0.f))
+			[
+				SNew(STextBlock)
+				.Text(FText::FromName(DefaultName))
+			];
 
 		// Menu Item
 		FUIAction Action(FExecuteAction::CreateSP(this, &SRCControllerPanel::OnAddControllerClicked, Pair.Key, Pair.Value));
-		MenuBuilder.AddMenuEntry(
-			FText::Format(LOCTEXT("AddController", "{0}"), FText::FromName(DefaultName)),
-			FText::Format(LOCTEXT("AddControllerTooltip", "Add a controller of type {0}."), FText::FromName(DefaultName)),
-			FSlateIcon(),
-			Action);
+
+		MenuBuilder.AddMenuEntry(Action, MenuItemWidget);
 	}
 
 	return MenuBuilder.MakeWidget();
@@ -215,6 +237,11 @@ void SRCControllerPanel::OnAddControllerClicked(const EPropertyBagPropertyType I
 	{
 		URCVirtualPropertyInContainer* NewVirtualProperty = Preset->AddVirtualProperty(URCController::StaticClass(), InValueType, InValueTypeObject);
 
+		if (ControllerPanelList.IsValid())
+		{
+			NewVirtualProperty->DisplayIndex = ControllerPanelList->GetNumControllerItems();
+		}
+
 		// Refresh list
 		const TSharedPtr<SRemoteControlPanel> RemoteControlPanel = GetRemoteControlPanel();
 		check(RemoteControlPanel);
diff --git a/Engine/Plugins/VirtualProduction/RemoteControl/Source/RemoteControlUI/Private/UI/Controller/SRCControllerPanelList.cpp b/Engine/Plugins/VirtualProduction/RemoteControl/Source/RemoteControlUI/Private/UI/Controller/SRCControllerPanelList.cpp
--- a/Engine/Plugins/VirtualProduction/RemoteControl/Source/RemoteControlUI/Private/UI/Controller/SRCControllerPanelList.cpp
+++ b/Engine/Plugins/VirtualProduction/RemoteControl/Source/RemoteControlUI/Private/UI/Controller/SRCControllerPanelList.cpp
@@ -2,17 +2,22 @@
 
 #include "SRCControllerPanelList.h"
 
+#include "Controller/RCControllerContainer.h"
 #include "IDetailTreeNode.h"
 #include "IPropertyRowGenerator.h"
 #include "RCControllerModel.h"
 #include "RCVirtualProperty.h"
 #include "RCVirtualPropertyContainer.h"
 #include "RemoteControlPreset.h"
+#include "SDropTarget.h"
 #include "SlateOptMacros.h"
 #include "SRCControllerPanel.h"
 #include "Styling/RemoteControlStyles.h"
+#include "UI/Action/SRCActionPanelList.h"
 #include "UI/BaseLogicUI/RCLogicModeBase.h"
 #include "UI/RemoteControlPanelStyle.h"
+#include "UI/SRCPanelExposedEntity.h"
+#include "UI/SRCPanelDragHandle.h"
 #include "UI/SRemoteControlPanel.h"
 #include "Widgets/Views/SHeaderRow.h"
 
@@ -20,18 +25,20 @@
 
 namespace UE::RCControllerPanelList
 {
-	static const FName ControllerNameColumn(TEXT("Controller Name"));
-	static const FName ControllerValueColumn(TEXT("Controller Value"));
+	namespace Columns
+	{
+		const FName Name = TEXT("Controller Name");
+		const FName Value = TEXT("Controller Value");
+		const FName DragHandle = TEXT("Drag Handle");
+	}
 
 	class SControllerItemListRow : public SMultiColumnTableRow<TSharedRef<FRCControllerModel>>
 	{
-	private:
-		TSharedPtr<FRCControllerModel> ControllerItem;
-
 	public:
-		void Construct(const FTableRowArgs& InArgs, const TSharedRef<STableViewBase>& OwnerTableView, TSharedRef<FRCControllerModel> InControllerItem)
+		void Construct(const FTableRowArgs& InArgs, const TSharedRef<STableViewBase>& OwnerTableView, TSharedRef<FRCControllerModel> InControllerItem, TSharedRef<SRCControllerPanelList> InControllerPanelList)
 		{
 			ControllerItem = InControllerItem;
+			ControllerPanelList = InControllerPanelList;
 			FSuperRowType::Construct(InArgs, OwnerTableView);
 		}
 
@@ -40,20 +47,95 @@ namespace UE::RCControllerPanelList
 			if (!ensure(ControllerItem.IsValid()))
 				return SNullWidget::NullWidget;
 
-			if (ColumnName == UE::RCControllerPanelList::ControllerNameColumn)
+			if (ColumnName == UE::RCControllerPanelList::Columns::Name)
 			{
-				return ControllerItem->GetNameWidget();
+				return WrapWithDropTarget(ControllerItem->GetNameWidget());
 			}
-			else if (ColumnName == UE::RCControllerPanelList::ControllerValueColumn)
+			else if (ColumnName == UE::RCControllerPanelList::Columns::Value)
 			{
 				return ControllerItem->GetWidget();
 			}
+			else if (ColumnName == UE::RCControllerPanelList::Columns::DragHandle)
+			{
+				SAssignNew(DragDropBorderWidget, SBorder)
+					.BorderImage(FRemoteControlPanelStyle::Get()->GetBrush("RemoteControlPanel.ExposedFieldBorder"));
+
+				TSharedRef<SWidget> DragHandleWidget = 
+					SNew(SBox)
+					.Padding(5.f)
+					[
+						SNew(SRCPanelDragHandle<FRCControllerDragDrop>, ControllerItem->GetId())
+						.Widget(DragDropBorderWidget)
+					];
+
+				return WrapWithDropTarget(DragHandleWidget);
+			}
 
 			return SNullWidget::NullWidget;
 		}
 
+		TSharedPtr<SBorder> DragDropBorderWidget;
+
 	private:
 
+		TSharedRef<SWidget> WrapWithDropTarget(const TSharedRef<SWidget> InWidget)
+		{
+			return SNew(SDropTarget)
+				.VerticalImage(FRemoteControlPanelStyle::Get()->GetBrush("RemoteControlPanel.VerticalDash"))
+				.HorizontalImage(FRemoteControlPanelStyle::Get()->GetBrush("RemoteControlPanel.HorizontalDash"))
+				.OnDropped_Lambda([this](const FGeometry& InGeometry, const FDragDropEvent& InDragDropEvent) { return SControllerItemListRow::OnControllerItemDragDrop(InDragDropEvent.GetOperation()); })
+				.OnAllowDrop(this, &SControllerItemListRow::OnAllowDrop)
+				.OnIsRecognized(this, &SControllerItemListRow::OnAllowDrop)
+				[
+					InWidget
+				];
+		}
+
+		FReply OnControllerItemDragDrop(TSharedPtr<FDragDropOperation> DragDropOperation)
+		{
+			if (!DragDropOperation)
+			{
+				return FReply::Handled();
+			}
+
+			if (DragDropOperation->IsOfType<FRCControllerDragDrop>())
+			{
+				if (TSharedPtr<FRCControllerDragDrop> DragDropOp = StaticCastSharedPtr<FRCControllerDragDrop>(DragDropOperation))
+				{
+					const FGuid DragDropControllerId = DragDropOp->GetId();
+
+					if (ControllerPanelList.IsValid())
+					{
+						TSharedPtr<FRCControllerModel> DragDropControllerItem = ControllerPanelList->FindControllerItemById(DragDropControllerId);
+
+						if (ensure(ControllerItem && DragDropControllerItem))
+						{
+							ControllerPanelList->ReorderControllerItem(DragDropControllerItem.ToSharedRef(), ControllerItem.ToSharedRef());
+						}
+					}
+				}
+			}
+
+			return FReply::Handled();
+		}
+
+		bool OnAllowDrop(TSharedPtr<FDragDropOperation> DragDropOperation)
+		{
+			if (DragDropOperation && ControllerItem)
+			{
+				// Dragging Controllers onto Controllers (Reordering)
+				if (DragDropOperation->IsOfType<FRCControllerDragDrop>())
+				{
+					if (TSharedPtr<FRCControllerDragDrop> DragDropOp = StaticCastSharedPtr<FRCControllerDragDrop>(DragDropOperation))
+					{
+						return true;
+					}
+				}
+			}
+
+			return false;
+		}
+
 		//~ SWidget Interface
 		virtual FReply OnMouseButtonDoubleClick(const FGeometry& InMyGeometry, const FPointerEvent& InMouseEvent) override
 		{
@@ -64,8 +146,12 @@ namespace UE::RCControllerPanelList
 
 			return FSuperRowType::OnMouseButtonDoubleClick(InMyGeometry, InMouseEvent);
 		}
+
+private:
+		TSharedPtr<FRCControllerModel> ControllerItem;
+		TSharedPtr<SRCControllerPanelList> ControllerPanelList;
 	};
-}
+} 
 
 BEGIN_SLATE_FUNCTION_BUILD_OPTIMIZATION
 
@@ -86,14 +172,19 @@ void SRCControllerPanelList::Construct(const FArguments& InArgs, const TSharedRe
 			SNew(SHeaderRow)
 			.Style(&RCPanelStyle->HeaderRowStyle)
 
-			+ SHeaderRow::Column(UE::RCControllerPanelList::ControllerNameColumn)
+			+ SHeaderRow::Column(UE::RCControllerPanelList::Columns::DragHandle)
+			.DefaultLabel(FText::GetEmpty())
+			.FillWidth(0.1f)
+			.HeaderContentPadding(RCPanelStyle->HeaderRowPadding)
+
+			+ SHeaderRow::Column(UE::RCControllerPanelList::Columns::Name)
 			.DefaultLabel(LOCTEXT("ControllerNameColumnName", "Name"))
-			.FillWidth(0.25f)
+			.FillWidth(0.3f)
 			.HeaderContentPadding(RCPanelStyle->HeaderRowPadding)
 
-			+ SHeaderRow::Column(UE::RCControllerPanelList::ControllerValueColumn)
+			+ SHeaderRow::Column(UE::RCControllerPanelList::Columns::Value)
 			.DefaultLabel(LOCTEXT("ControllerValueColumnName", "Input"))
-			.FillWidth(0.75f)
+			.FillWidth(0.6f)
 			.HeaderContentPadding(RCPanelStyle->HeaderRowPadding)
 		);
 
@@ -147,10 +238,15 @@ void SRCControllerPanelList::Reset()
 	PropertyRowGenerator->OnFinishedChangingProperties().AddSP(this, &SRCControllerPanelList::OnFinishedChangingProperties);
 
 	// Generator should be moved to separate class
-	for (const TSharedRef<IDetailTreeNode>& CategoryNode : PropertyRowGenerator->GetRootTreeNodes())
+	TArray<TSharedRef<IDetailTreeNode>> RootTreeNodes = PropertyRowGenerator->GetRootTreeNodes();
+
+	for (const TSharedRef<IDetailTreeNode>& CategoryNode : RootTreeNodes)
 	{
 		TArray<TSharedRef<IDetailTreeNode>> Children;
 		CategoryNode->GetChildren(Children);
+
+		ControllerItems.SetNumZeroed(Children.Num());
+
 		for (TSharedRef<IDetailTreeNode>& Child : Children)
 		{
 			FProperty* Property = Child->CreatePropertyHandle()->GetProperty();
@@ -158,7 +254,8 @@ void SRCControllerPanelList::Reset()
 
 			if (URCVirtualPropertyBase* Controller = Preset->GetVirtualProperty(Property->GetFName()))
 			{
-				ControllerItems.Add(MakeShared<FRCControllerModel>(Controller, Child, RemoteControlPanel));
+				if(ensureAlways(ControllerItems.IsValidIndex(Controller->DisplayIndex)))
+					ControllerItems[Controller->DisplayIndex] = MakeShared<FRCControllerModel>(Controller, Child, RemoteControlPanel);
 			}
 		}
 	}
@@ -170,7 +267,7 @@ TSharedRef<ITableRow> SRCControllerPanelList::OnGenerateWidgetForList(TSharedPtr
 {
 	typedef UE::RCControllerPanelList::SControllerItemListRow SControllerRowType;
 
-	return SNew(SControllerRowType, OwnerTable, InItem.ToSharedRef())
+	return SNew(SControllerRowType, OwnerTable, InItem.ToSharedRef(), SharedThis(this))
 		.Style(&RCPanelStyle->TableRowStyle)
 		.Padding(FMargin(3.f));
 }
@@ -279,4 +376,37 @@ void SRCControllerPanelList::EnterRenameMode()
 	}
 }
 
+TSharedPtr<FRCControllerModel> SRCControllerPanelList::FindControllerItemById(const FGuid& InId) const
+{
+	for (TSharedPtr<FRCControllerModel> ControllerItem : ControllerItems)
+	{
+		if (ControllerItem && ControllerItem->GetId() == InId)
+		{
+			return ControllerItem;
+		}
+	}
+
+	return nullptr;
+}
+
+void SRCControllerPanelList::ReorderControllerItem(TSharedRef<FRCControllerModel> ItemToMove, TSharedRef<FRCControllerModel> AnchorItem)
+{
+	int32 Index = ControllerItems.Find(AnchorItem);
+
+	// Update UI list
+	ControllerItems.RemoveSingle(ItemToMove);
+	ControllerItems.Insert(ItemToMove, Index);
+
+	// Update display indices
+	for (int32 i = Index; i < ControllerItems.Num(); i++)
+	{
+		if(URCVirtualPropertyBase * Controller = ControllerItems[i]->GetVirtualProperty())
+		{
+			Controller->DisplayIndex = i;
+		}
+	}
+
+	ListView->RequestListRefresh();
+}
+
 #undef LOCTEXT_NAMESPACE
diff --git a/Engine/Plugins/VirtualProduction/RemoteControl/Source/RemoteControlUI/Private/UI/Controller/SRCControllerPanelList.h b/Engine/Plugins/VirtualProduction/RemoteControl/Source/RemoteControlUI/Private/UI/Controller/SRCControllerPanelList.h
--- a/Engine/Plugins/VirtualProduction/RemoteControl/Source/RemoteControlUI/Private/UI/Controller/SRCControllerPanelList.h
+++ b/Engine/Plugins/VirtualProduction/RemoteControl/Source/RemoteControlUI/Private/UI/Controller/SRCControllerPanelList.h
@@ -1,7 +1,10 @@
 ﻿// Copyright Epic Games, Inc. All Rights Reserved.
 
 #pragma once
+#include "DragAndDrop/DecoratedDragDropOp.h"
 #include "UI/BaseLogicUI/SRCLogicPanelListBase.h"
+#include "UI/RemoteControlPanelStyle.h"
+#include "Widgets/Layout/SBorder.h"
 
 struct FRCPanelStyle;
 class FRCControllerModel;
@@ -15,6 +18,37 @@ class URCController;
 class URemoteControlPreset;
 template <typename ItemType> class SListView;
 
+/*
+* ~ FRCControllerDragDrop ~
+*
+* Facilitates drag-drop operation for Controller row drag handles
+*/
+class FRCControllerDragDrop final : public FDecoratedDragDropOp
+{
+public:
+	DRAG_DROP_OPERATOR_TYPE(FRCControllerDragDropOp, FDecoratedDragDropOp)
+
+	using WidgetType = SWidget;
+
+	FRCControllerDragDrop(TSharedPtr<SWidget> InWidget, const FGuid& InId)
+		: Id(InId)
+	{
+	}
+
+	FGuid GetId() const
+	{
+		return Id;
+	}
+
+	virtual void OnDrop(bool bDropWasHandled, const FPointerEvent& MouseEvent) override
+	{
+		FDecoratedDragDropOp::OnDrop(bDropWasHandled, MouseEvent);
+	}
+
+private:
+	FGuid Id;
+};
+
 /*
 * ~ SRCControllerPanelList ~
 *
@@ -47,6 +81,18 @@ public:
 
 	void EnterRenameMode();
 
+	int32 GetNumControllerItems() const
+	{
+		return ControllerItems.Num();
+	}
+
+	/** Finds a Controller UI model by unique Id*/
+	TSharedPtr<FRCControllerModel> FindControllerItemById(const FGuid& InId) const;
+
+	/** Given an item to move and an anchor row this function moves the item to the position of the anchor
+	* and pushes all other rows below */
+	void ReorderControllerItem(TSharedRef<FRCControllerModel> ItemToMove, TSharedRef<FRCControllerModel> AnchorItem);
+
 private:
 
 	/** OnGenerateRow delegate for the Actions List View */
