Range Behaviour - Fix LOCTEXT Warning

Fixes an LOCTEXT issue in the Behaviour/namespace

diff --git a/Engine/Plugins/VirtualProduction/RemoteControl/Source/RemoteControlUI/Private/UI/Action/RangeMap/RCActionRangeMapModel.cpp b/Engine/Plugins/VirtualProduction/RemoteControl/Source/RemoteControlUI/Private/UI/Action/RangeMap/RCActionRangeMapModel.cpp
--- a/Engine/Plugins/VirtualProduction/RemoteControl/Source/RemoteControlUI/Private/UI/Action/RangeMap/RCActionRangeMapModel.cpp
+++ b/Engine/Plugins/VirtualProduction/RemoteControl/Source/RemoteControlUI/Private/UI/Action/RangeMap/RCActionRangeMapModel.cpp
@@ -93,7 +93,7 @@ TSharedPtr<SHeaderRow> FRCActionRangeMapModel::GetHeaderRow()
 		.HeaderContentPadding(RCPanelStyle->HeaderRowPadding)
 
 		+ SHeaderRow::Column(UE::RCActionPanelRangeMapList::Columns::Step)
-		.DefaultLabel(LOCTEXT("RCActionConditionColumnHeader", "Step"))
+		.DefaultLabel(LOCTEXT("RCActionStepColumnHeader", "Step"))
 		.FillWidth(0.2f)
 		.HeaderContentPadding(RCPanelStyle->HeaderRowPadding)
 
