Remote Control Logic - Step Fix for Numeric Actions

Fix issues accidently added:  Nullptr Crash + Value wrongly set (or not set at all).

diff --git a/Engine/Plugins/VirtualProduction/RemoteControl/Source/RemoteControlLogic/Private/Behaviour/Builtin/RangeMap/RCRangeMapBehaviour.cpp b/Engine/Plugins/VirtualProduction/RemoteControl/Source/RemoteControlLogic/Private/Behaviour/Builtin/RangeMap/RCRangeMapBehaviour.cpp
--- a/Engine/Plugins/VirtualProduction/RemoteControl/Source/RemoteControlLogic/Private/Behaviour/Builtin/RangeMap/RCRangeMapBehaviour.cpp
+++ b/Engine/Plugins/VirtualProduction/RemoteControl/Source/RemoteControlLogic/Private/Behaviour/Builtin/RangeMap/RCRangeMapBehaviour.cpp
@@ -145,8 +145,8 @@ void URCRangeMapBehaviour::GetLerpActions(TMap<FGuid, TArray<URCAction*>>& OutNu
 
 		Algo::Sort(ArrayToSort, [this](const URCAction* ActionA, const URCAction* ActionB)
 		{
-			double A = RangeMapActionContainer.Find(ActionA)->StepValue;
-			double B = RangeMapActionContainer.Find(ActionB)->StepValue;
+			const double A = RangeMapActionContainer.Find(ActionA)->StepValue;
+			const double B = RangeMapActionContainer.Find(ActionB)->StepValue;
 
 			return A > B;
 		});
@@ -191,7 +191,7 @@ void URCRangeMapBehaviour::Execute()
 		return;
 	}
 
-	for (TTuple<FGuid, TTuple<URCAction*, URCAction*>>& LerpAction : LerpActions)
+	for (const TTuple<FGuid, TTuple<URCAction*, URCAction*>>& LerpAction : LerpActions)
 	{
 		URemoteControlPreset* Preset = RCController->PresetWeakPtr.Get();
 		if (!Preset)
@@ -202,16 +202,16 @@ void URCRangeMapBehaviour::Execute()
 		FGuid CurrentFieldId = LerpAction.Key;
 
 		// Get the StepValues given for the FGuid
-		URCAction* MinActionOfPair = LerpAction.Value.Key;
-		URCAction* MaxActionOfPair = LerpAction.Value.Value;
+		const URCAction* MinActionOfPair = LerpAction.Value.Key;
+		const URCAction* MaxActionOfPair = LerpAction.Value.Value;
 
 		if (!MinActionOfPair || !MaxActionOfPair)
 		{
 			return;
 		}
 
-		FRCRangeMapStep* MinRangeMapStep = RangeMapActionContainer.Find(MinActionOfPair);
-		FRCRangeMapStep* MaxRangeMapStep = RangeMapActionContainer.Find(MaxActionOfPair);
+		const FRCRangeMapStep* MinRangeMapStep = RangeMapActionContainer.Find(MinActionOfPair);
+		const FRCRangeMapStep* MaxRangeMapStep = RangeMapActionContainer.Find(MaxActionOfPair);
 
 		if (!(MinRangeMapStep && MaxRangeMapStep))
 		{
@@ -220,8 +220,8 @@ void URCRangeMapBehaviour::Execute()
 		}
 
 		// Denormalize and Map them based on our Min and Max Value
-		double MappedMinStep = UKismetMathLibrary::Lerp(MinValue, MaxValue, MinRangeMapStep->StepValue);
-		double MappedMaxStep = UKismetMathLibrary::Lerp(MinValue, MaxValue, MaxRangeMapStep->StepValue);
+		const double MappedMinStep = UKismetMathLibrary::Lerp(MinValue, MaxValue, MinRangeMapStep->StepValue);
+		const double MappedMaxStep = UKismetMathLibrary::Lerp(MinValue, MaxValue, MaxRangeMapStep->StepValue);
 
 		// Normalize our Controller based on the new MappedMinStep and MappedMaxStep Range.
 		double CustomNormalizedStepValue = UKismetMathLibrary::NormalizeToRange(ControllerFloatValue, MappedMinStep, MappedMaxStep);
@@ -237,8 +237,8 @@ void URCRangeMapBehaviour::Execute()
 				MinRangeMapStep->PropertyValue->GetValueDouble(MinRangeDouble);
 				MaxRangeMapStep->PropertyValue->GetValueDouble(MaxRangeDouble);
 
-				double ResultDouble = UKismetMathLibrary::Lerp(MinRangeDouble, MaxRangeDouble, CustomNormalizedStepValue);
-				if (TSharedPtr<FRemoteControlProperty> RCProperty = Preset->GetExposedEntity<FRemoteControlProperty>(CurrentFieldId).Pin())
+				const double ResultDouble = UKismetMathLibrary::Lerp(MinRangeDouble, MaxRangeDouble, CustomNormalizedStepValue);
+				if (const TSharedPtr<FRemoteControlProperty> RCProperty = Preset->GetExposedEntity<FRemoteControlProperty>(CurrentFieldId).Pin())
 				{
 					RCProperty->GetPropertyHandle()->SetValue(ResultDouble);
 				}
@@ -253,8 +253,8 @@ void URCRangeMapBehaviour::Execute()
 				MinRangeMapStep->PropertyValue->GetValueFloat(MinRangeFloat);
 				MaxRangeMapStep->PropertyValue->GetValueFloat(MaxRangeFloat);
 
-				float ResultFloat = UKismetMathLibrary::Lerp(MinRangeFloat, MaxRangeFloat, CustomNormalizedStepValue);
-				if (TSharedPtr<FRemoteControlProperty> RCProperty = Preset->GetExposedEntity<FRemoteControlProperty>(CurrentFieldId).Pin())
+				const float ResultFloat = UKismetMathLibrary::Lerp(MinRangeFloat, MaxRangeFloat, CustomNormalizedStepValue);
+				if (const TSharedPtr<FRemoteControlProperty> RCProperty = Preset->GetExposedEntity<FRemoteControlProperty>(CurrentFieldId).Pin())
 				{
 					RCProperty->GetPropertyHandle()->SetValue(ResultFloat);
 				}
@@ -277,7 +277,8 @@ void URCRangeMapBehaviour::Execute()
 URCAction* URCRangeMapBehaviour::AddAction(const TSharedRef<const FRemoteControlField> InRemoteControlField)
 {
 	double InStepValue = 0.0;
-	PropertyContainer->GetVirtualProperty(UE::RCRangeMapBehaviour::Step);
+	URCVirtualPropertyBase* StepVirtualProperty = PropertyContainer->GetVirtualProperty(UE::RCRangeMapBehaviour::Step);
+	StepVirtualProperty->GetValueDouble(InStepValue);
 	
 	const TRCActionUniquenessTest ActionUniquenessTest = [this, InRemoteControlField, InStepValue](const TSet<TObjectPtr<URCAction>>& Actions)
 	{
@@ -287,9 +288,9 @@ URCAction* URCRangeMapBehaviour::AddAction(const TSharedRef<const FRemoteControl
 	return ActionContainer->AddAction(ActionUniquenessTest, InRemoteControlField);
 }
 
-void URCRangeMapBehaviour::OnActionAdded(URCAction* Action, const TObjectPtr<URCVirtualPropertySelfContainer> InPropertyValue)
+void URCRangeMapBehaviour::OnActionAdded(URCAction* Action, URCVirtualPropertySelfContainer* InPropertyValue)
 {
-	double InStepValue;
+	double InStepValue = 0.0;
 	URCVirtualPropertyBase* StepVirtualProperty = PropertyContainer->GetVirtualProperty(UE::RCRangeMapBehaviour::Step);
 	StepVirtualProperty->GetValueDouble(InStepValue);
 	StepVirtualProperty->SetValueDouble(0.0);
@@ -301,7 +302,7 @@ void URCRangeMapBehaviour::OnActionAdded(URCAction* Action, const TObjectPtr<URC
 
 bool URCRangeMapBehaviour::GetRangeValuePairsForLerp(TMap<FGuid, TTuple<URCAction*, URCAction*>>& OutPairs)
 {
-	double NormalizedControllerValue = UKismetMathLibrary::NormalizeToRange(ControllerFloatValue, MinValue, MaxValue);
+	const double NormalizedControllerValue = UKismetMathLibrary::NormalizeToRange(ControllerFloatValue, MinValue, MaxValue);
 
 	TMap<FGuid, TArray<URCAction*>> LerpActionMap;
 	GetLerpActions(LerpActionMap);
@@ -315,10 +316,10 @@ bool URCRangeMapBehaviour::GetRangeValuePairsForLerp(TMap<FGuid, TTuple<URCActio
 		}
 
 		// Intermediary calculations of Lerp for in-between Lerp.
-		URCAction* MinAction;
-		FRCRangeMapStep* MinRangeMap = nullptr;
-		URCAction* MaxAction;
 		FRCRangeMapStep* MaxRangeMap = nullptr;
+		FRCRangeMapStep* MinRangeMap = nullptr;
+		URCAction* MaxAction = nullptr;
+		URCAction* MinAction = nullptr;
 		
 		for (URCAction* Action : NumericActionTuple.Value)
 		{
diff --git a/Engine/Plugins/VirtualProduction/RemoteControl/Source/RemoteControlLogic/Public/Behaviour/Builtin/RangeMap/RCRangeMapBehaviour.h b/Engine/Plugins/VirtualProduction/RemoteControl/Source/RemoteControlLogic/Public/Behaviour/Builtin/RangeMap/RCRangeMapBehaviour.h
--- a/Engine/Plugins/VirtualProduction/RemoteControl/Source/RemoteControlLogic/Public/Behaviour/Builtin/RangeMap/RCRangeMapBehaviour.h
+++ b/Engine/Plugins/VirtualProduction/RemoteControl/Source/RemoteControlLogic/Public/Behaviour/Builtin/RangeMap/RCRangeMapBehaviour.h
@@ -15,7 +15,7 @@ struct FRCRangeMapStep
 	GENERATED_BODY()
 
 	FRCRangeMapStep() {}
-	FRCRangeMapStep(const double InStepValue, const URCVirtualPropertySelfContainer* InPropertyValue)
+	FRCRangeMapStep(const double InStepValue, URCVirtualPropertySelfContainer* InPropertyValue)
 		: StepValue(InStepValue)
 	{
 		if (!PropertyValue)
@@ -23,7 +23,7 @@ struct FRCRangeMapStep
 			return;
 		}
 		
-		PropertyValue = MakeObjectPtrUnsafe<URCVirtualPropertySelfContainer>(InPropertyValue);
+		PropertyValue = InPropertyValue;
 	}
 
 	/** The Value which we use represent the action based on a normalized step. */
@@ -32,7 +32,7 @@ struct FRCRangeMapStep
 
 	/** The Property this Action holds and will be used for calculations for the lerp. */
 	UPROPERTY()
-	TObjectPtr<URCVirtualPropertySelfContainer> PropertyValue;
+	URCVirtualPropertySelfContainer* PropertyValue;
 };
 
 /**
@@ -66,7 +66,7 @@ public:
 	void Refresh();
 	
 	/** Called whenever a new action is added into the ActionContainer. Will add a FRCRangeMapStep into RangeMapActionContainer corresponding to the Action. */
-	void OnActionAdded(URCAction* Action, const TObjectPtr<URCVirtualPropertySelfContainer> InPropertyValue);
+	void OnActionAdded(URCAction* Action, URCVirtualPropertySelfContainer* InPropertyValue);
 
 	/** Adds an Action to the ActionContainer, whilst making sure that it is unique. */
 	URCAction* AddAction(const TSharedRef<const FRemoteControlField> InRemoteControlField) override;
