ObjectPtr upgrade for engine plugins used from UnrealEditor -allmodules.

diff --git a/Engine/Plugins/VirtualProduction/RemoteControl/Source/RemoteControlLogic/Public/Behaviour/Builtin/RangeMap/RCRangeMapBehaviour.h b/Engine/Plugins/VirtualProduction/RemoteControl/Source/RemoteControlLogic/Public/Behaviour/Builtin/RangeMap/RCRangeMapBehaviour.h
--- a/Engine/Plugins/VirtualProduction/RemoteControl/Source/RemoteControlLogic/Public/Behaviour/Builtin/RangeMap/RCRangeMapBehaviour.h
+++ b/Engine/Plugins/VirtualProduction/RemoteControl/Source/RemoteControlLogic/Public/Behaviour/Builtin/RangeMap/RCRangeMapBehaviour.h
@@ -32,7 +32,7 @@ struct FRCRangeMapStep
 
 	/** The Property this Action holds and will be used for calculations for the lerp. */
 	UPROPERTY()
-	URCVirtualPropertySelfContainer* PropertyValue;
+	TObjectPtr<URCVirtualPropertySelfContainer> PropertyValue;
 };
 
 /**
