Remote Control Logic Behaviour - Set Asset By Path, Additional Features

- Add Button to copy the path of a selected asset into the Path Array
- Replace Target Property Name as Target and use more visual representation instead.
- Add Token *INPUT templating into the Path Array.
- Add Preview of Set-up Working Path, including Token Input
- Add External Image Pathing

diff --git a/Engine/Plugins/VirtualProduction/RemoteControl/Source/RemoteControlLogic/Private/Behaviour/Builtin/RCBehaviourSetAssetByPathNode.cpp b/Engine/Plugins/VirtualProduction/RemoteControl/Source/RemoteControlLogic/Private/Behaviour/Builtin/RCBehaviourSetAssetByPathNode.cpp
--- a/Engine/Plugins/VirtualProduction/RemoteControl/Source/RemoteControlLogic/Private/Behaviour/Builtin/RCBehaviourSetAssetByPathNode.cpp
+++ b/Engine/Plugins/VirtualProduction/RemoteControl/Source/RemoteControlLogic/Private/Behaviour/Builtin/RCBehaviourSetAssetByPathNode.cpp
@@ -11,7 +11,7 @@
 URCBehaviourSetAssetByPathNode::URCBehaviourSetAssetByPathNode()
 {
 	DisplayName = NSLOCTEXT("Remote Control Behaviour", "Behaviour Name - Set Asset By Path", "Set Asset By Path");
-	BehaviorDescription = NSLOCTEXT("Remote Control Behaviour", "Behavior Desc - Set Asset By Path", "Triggers an event which sets an object based on the TargetProperty Name.");
+	BehaviorDescription = NSLOCTEXT("Remote Control Behaviour", "Behaviour Desc - Set Asset By Path", "Triggers an event which sets an object based on the selected Exposed Entity.");
 }
 
 bool URCBehaviourSetAssetByPathNode::Execute_Implementation(URCBehaviour* InBehaviour) const
@@ -33,38 +33,12 @@ bool URCBehaviourSetAssetByPathNode::Execute_Implementation(URCBehaviour* InBeha
 		return false;
 	}
 
-	FString ControllerString;
-	FString TargetPropertyString;
 	FString DefaultString;
-	
-	RCController->GetValueString(ControllerString);
-	SetAssetByPathBehaviour->PropertyInContainer->GetVirtualProperty(SetAssetByPathBehaviourHelpers::TargetProperty)->GetValueString(TargetPropertyString);
-	SetAssetByPathBehaviour->PropertyInContainer->GetVirtualProperty(SetAssetByPathBehaviourHelpers::DefaultProperty)->GetValueString(DefaultString);
-
-	// Add Path String Concat
-	FString ConcatPath = SetAssetByPathBehaviourHelpers::ContentFolder;
-	for (FString PathPart : SetAssetByPathBehaviour->PathStruct.PathArray)
-	{
-		if (PathPart.IsEmpty())
-		{
-			continue;
-		}
-		// Add '/' if needed
-		int32 IndexFound;
+	SetAssetByPathBehaviour->PropertyInContainer->GetVirtualProperty(SetAssetByPathBehaviourHelpers::DefaultInput)->GetValueString(DefaultString);
 
-		if (PathPart.FindLastChar('_', IndexFound) && IndexFound == PathPart.Len() - 1)
-		{
-			// In the case there's underscore char in the at the end of one of the Path Strings, do nothing to facilitate more complex pathing behaviours.
-		}
-		else if (!PathPart.FindLastChar('/', IndexFound) || IndexFound < PathPart.Len() - 1)
-		{
-			PathPart = PathPart.AppendChar('/');
-		}
-		
-		ConcatPath += PathPart;
-	}
+	const FString ConcatPath = SetAssetByPathBehaviour->GetCurrentPath();
 	
-	return SetAssetByPathBehaviour->SetAssetByPath(ConcatPath, TargetPropertyString, DefaultString);
+	return SetAssetByPathBehaviour->SetAssetByPath(ConcatPath, DefaultString);
 }
 
 bool URCBehaviourSetAssetByPathNode::IsSupported_Implementation(URCBehaviour* InBehaviour) const
diff --git a/Engine/Plugins/VirtualProduction/RemoteControl/Source/RemoteControlLogic/Private/Behaviour/RCSetAssetByPathBehaviour.cpp b/Engine/Plugins/VirtualProduction/RemoteControl/Source/RemoteControlLogic/Private/Behaviour/RCSetAssetByPathBehaviour.cpp
--- a/Engine/Plugins/VirtualProduction/RemoteControl/Source/RemoteControlLogic/Private/Behaviour/RCSetAssetByPathBehaviour.cpp
+++ b/Engine/Plugins/VirtualProduction/RemoteControl/Source/RemoteControlLogic/Private/Behaviour/RCSetAssetByPathBehaviour.cpp
@@ -8,14 +8,16 @@
 #include "Components/StaticMeshComponent.h"
 #include "Controller/RCController.h"
 #include "Engine/StaticMesh.h"
+#include "Engine/Texture2DDynamic.h"
 #include "Engine/World.h"
 #include "IRemoteControlModule.h"
 #include "IRemoteControlPropertyHandle.h"
+#include "Kismet/KismetRenderingLibrary.h"
 #include "Materials/MaterialInstance.h"
 #include "RCVirtualProperty.h"
 #include "RCVirtualPropertyContainer.h"
+#include "RemoteControlLogger.h"
 #include "RemoteControlPreset.h"
-#include "StructSerializer.h"
 
 URCSetAssetByPathBehaviour::URCSetAssetByPathBehaviour()
 {
@@ -24,15 +26,15 @@ URCSetAssetByPathBehaviour::URCSetAssetByPathBehaviour()
 
 void URCSetAssetByPathBehaviour::Initialize()
 {
-	PropertyInContainer->AddProperty(SetAssetByPathBehaviourHelpers::TargetProperty, URCController::StaticClass(), EPropertyBagPropertyType::String);
-	PropertyInContainer->AddProperty(SetAssetByPathBehaviourHelpers::DefaultProperty, URCController::StaticClass(), EPropertyBagPropertyType::String);
+	PropertyInContainer->AddProperty(SetAssetByPathBehaviourHelpers::DefaultInput, URCController::StaticClass(), EPropertyBagPropertyType::String);
+	UpdateTargetEntity();
 	
 	Super::Initialize();
 }
 
-bool URCSetAssetByPathBehaviour::SetAssetByPath(const FString& AssetPath, const FString& TargetPropertyString, const FString& DefaultString)
+bool URCSetAssetByPathBehaviour::SetAssetByPath(const FString& AssetPath, const FString& DefaultString)
 {
-	URCController* Controller = ControllerWeakPtr.Get();
+	const URCController* Controller = ControllerWeakPtr.Get();
 	
 	if (AssetPath.IsEmpty() || !Controller)
 	{
@@ -44,113 +46,125 @@ bool URCSetAssetByPathBehaviour::SetAssetByPath(const FString& AssetPath, const
 	{
 		return false;
 	}
+
+	if (!bInternal)
+	{
+		FRemoteControlLogger::Get().Log(SetAssetByPathBehaviourHelpers::SetAssetByPathBehaviour, [AssetPath, ControllerString]
+		{
+			return FText::FromString(FString("Path Behaviour attempts to set Asset to %s") + *AssetPath + *ControllerString);
+		});
+		if (SetExternalAsset(AssetPath + ControllerString))
+		{
+			return true;
+		}
+
+		FRemoteControlLogger::Get().Log(SetAssetByPathBehaviourHelpers::SetAssetByPathBehaviour, [AssetPath, DefaultString]
+		{
+			return FText::FromString(FString("Path Behaviour attempts to set Asset to %s") + *AssetPath + *DefaultString);
+		});
+		if (!DefaultString.IsEmpty() && SetExternalAsset(AssetPath + DefaultString))
+		{
+			return true;
+		}
+
+		return false;
+	}
 	
 	FSoftObjectPath MainObjectRef(AssetPath + ControllerString);
 	
 	if (UObject* Object = MainObjectRef.TryLoad())
 	{
-		return SetAsset(Object, TargetPropertyString);
+		FRemoteControlLogger::Get().Log(SetAssetByPathBehaviourHelpers::SetAssetByPathBehaviour, [MainObjectRef]
+		{
+			return FText::FromString(FString("Path Behaviour attempts to set Asset to %s") + *MainObjectRef.GetAssetPathString());
+		});
+		return SetInternalAsset(Object);
 	}
 	
 	if (!DefaultString.IsEmpty())
 	{
 		FSoftObjectPath DefaultObjectRef(AssetPath + DefaultString);
 		
-		return SetAsset(DefaultObjectRef.TryLoad(), TargetPropertyString);
+		FRemoteControlLogger::Get().Log(SetAssetByPathBehaviourHelpers::SetAssetByPathBehaviour, [DefaultObjectRef]
+		{
+			return FText::FromString(FString("Path Behaviour attempts to set Asset to %s") + *DefaultObjectRef.GetAssetPathString());
+		});
+		return SetInternalAsset(DefaultObjectRef.TryLoad());
 	}
 	
 	return false;
 }
 
-bool URCSetAssetByPathBehaviour::SetAsset(UObject* SetterObject, const FString& PropertyString)
+bool URCSetAssetByPathBehaviour::SetInternalAsset(UObject* SetterObject)
 {
-	URCController* Controller = ControllerWeakPtr.Get();
-	if (!SetterObject || PropertyString.IsEmpty() || !Controller)
+	const URCController* Controller = ControllerWeakPtr.Get();
+	
+	if (!SetterObject || !Controller || !TargetEntity)
 	{
+		FRemoteControlLogger::Get().Log(SetAssetByPathBehaviourHelpers::SetAssetByPathBehaviour, [SetterObject]
+		{
+			return FText::FromString(FString("Path Behaviour fails to set Asset: %s") + *SetterObject->GetName());
+		}, EVerbosityLevel::Error);
 		return false;
 	}
 	
-	const FName TargetPropertyLabel = FName(*PropertyString);
-	TWeakObjectPtr<URemoteControlPreset> PresetPtr = Controller->PresetWeakPtr;
-	TArray<TWeakPtr<FRemoteControlProperty>> ExposedProperties = PresetPtr->GetExposedEntities<FRemoteControlProperty>();
-	TArray<TWeakPtr<FRemoteControlEntity>> ExposedEntities = PresetPtr->GetExposedEntities<FRemoteControlEntity>();
-	
-	TWeakPtr<FRemoteControlEntity>* ExposedEntityPtr = ExposedEntities.FindByPredicate([TargetPropertyLabel](const TWeakPtr<FRemoteControlEntity>& Entity) { return Entity.Pin()->GetLabel() == TargetPropertyLabel;});
-	TWeakPtr<FRemoteControlProperty>* ExposedPropertyPtr = ExposedProperties.FindByPredicate([TargetPropertyLabel](const TWeakPtr<FRemoteControlProperty>& Property){ return Property.Pin()->GetLabel() == TargetPropertyLabel; });
+	const TWeakObjectPtr<URemoteControlPreset> PresetPtr = Controller->PresetWeakPtr;
 	
-	if (ExposedPropertyPtr)
+	TArray<TWeakPtr<FRemoteControlProperty>> ExposedProperties = PresetPtr->GetExposedEntities<FRemoteControlProperty>();
+
+	if (const TWeakPtr<FRemoteControlProperty>* ExposedPropertyPtr = ExposedProperties.FindByPredicate([this](const TWeakPtr<FRemoteControlProperty>& Property){ return Property.Pin()->GetLabel() == TargetEntity->GetLabel(); }))
 	{
-		TWeakPtr<FRemoteControlProperty> ExposedProperty = *ExposedPropertyPtr;
+		const TSharedPtr<FRemoteControlProperty> ExposedProperty = (*ExposedPropertyPtr).Pin();
+		if (!ExposedProperty)
+		{
+			ensureMsgf(false, TEXT("Exposed Property is not valid."));
+			return false;
+		}
+		
 		if (AssetClass == UStaticMesh::StaticClass() && SetterObject->IsA(UStaticMesh::StaticClass()))
 		{
-			UStaticMeshComponent* StaticMeshComponent = Cast<UStaticMeshComponent>(ExposedProperty.Pin()->GetBoundObject());
+			UStaticMeshComponent* StaticMeshComponent = Cast<UStaticMeshComponent>(ExposedProperty->GetBoundObject());
 			StaticMeshComponent->SetStaticMesh(Cast<UStaticMesh>(SetterObject));
 			
+			UE_LOG(LogRemoteControl, Display, TEXT("Path Behaviour sets Static Mesh Asset %s"), *SetterObject->GetName());
 			return true;
 		}
 		
 		if (AssetClass == UMaterial::StaticClass() && SetterObject->IsA(UMaterial::StaticClass()))
 		{
-			if (UMeshComponent* MaterialOwnerComponent = Cast<UMeshComponent>(ExposedProperty.Pin()->GetBoundObject()))
+			if (UMeshComponent* MaterialOwnerComponent = Cast<UMeshComponent>(ExposedProperty->GetBoundObject()))
 			{
 				MaterialOwnerComponent->SetMaterial(0, Cast<UMaterial>(SetterObject));
 			}
 			
+			UE_LOG(LogRemoteControl, Display, TEXT("Path Behaviour sets Material Asset %s"), *SetterObject->GetName());
 			return true;
 		}
-		
-		if (GetSupportedClasses().Contains(AssetClass) && SetterObject->IsA(AssetClass))
+
+		// Only Texture runs through here
+		if (GetSupportedClasses().Contains(AssetClass) && SetterObject->IsA(UTexture::StaticClass()))
 		{
-			FRCObjectReference ObjectRef;
-			TSharedPtr<FRemoteControlProperty> RemoteControlProperty = ExposedProperty.Pin();
-			ObjectRef.Property = RemoteControlProperty->GetProperty();
-			ObjectRef.Access = ExposedProperty.Pin()->GetPropertyHandle()->ShouldGenerateTransaction() ? ERCAccess::WRITE_TRANSACTION_ACCESS : ERCAccess::WRITE_ACCESS;
-			ObjectRef.PropertyPathInfo = RemoteControlProperty->FieldPathInfo.ToString();
-		
-			for (UObject* Object : RemoteControlProperty->GetBoundObjects())
+			if (const TSharedPtr<FRemoteControlProperty> RemoteControlProperty = ExposedProperty)
 			{
-				if (IRemoteControlModule::Get().ResolveObjectProperty(ObjectRef.Access, Object, ObjectRef.PropertyPathInfo, ObjectRef))
-				{
-					/**
-					 * Setting with Pointer and using Serialization of itself afterwards as a workaround with the Texture Asset not updating in the world.
-					 */
-					FObjectPropertyBase* ObjectProperty = CastField<FObjectPropertyBase>(ExposedProperty.Pin()->GetProperty());
-					if (!ObjectProperty)
-					{
-						return false;
-					}
-
-					FProperty* Property = RemoteControlProperty->GetProperty();
-					uint8* ValueAddress = Property->ContainerPtrToValuePtr<uint8>(ObjectRef.ContainerAdress);
-					ObjectProperty->SetObjectPropertyValue(ValueAddress, SetterObject);
-					
-					TArray<uint8> Buffer;
-					FMemoryReader Reader(Buffer);
-					FCborStructDeserializerBackend DeserializerBackend(Reader);
-
-					IRemoteControlModule::Get().SetObjectProperties(ObjectRef, DeserializerBackend, ERCPayloadType::Cbor, Buffer);
-					
-					return true;
-				}
+				SetTextureAsset(RemoteControlProperty, Cast<UTexture>(SetterObject));
 			}
 		}
 		
 		return false;
 	}
 	
-	if (ExposedEntityPtr)
+	if (TargetEntity)
 	{
-		TWeakPtr<FRemoteControlEntity> ExposedEntity = *ExposedEntityPtr;
+		const TWeakPtr<const FRemoteControlEntity> ExposedEntity = TargetEntity;
 		if (AssetClass == UBlueprint::StaticClass() && SetterObject->IsA(UBlueprint::StaticClass()))
 		{
-
 			AActor* OldActor = Cast<AActor>(ExposedEntity.Pin()->GetBoundObject());
 			UWorld* World = OldActor->GetWorld();
 			
 			OldActor->UnregisterAllComponents();
 			
-			FVector OldLocation = OldActor->GetActorLocation();
-			FRotator OldRotation = OldActor->GetActorRotation();
+			const FVector OldLocation = OldActor->GetActorLocation();
+			const FRotator OldRotation = OldActor->GetActorRotation();
 			const FName OldActorName = OldActor->GetFName();
 			
 			FActorSpawnParameters SpawnParams;
@@ -163,14 +177,17 @@ bool URCSetAssetByPathBehaviour::SetAsset(UObject* SetterObject, const FString&
 			{
 				World->DestroyActor(OldActor);
 				
+				UE_LOG(LogRemoteControl, Display, TEXT("Path Behaviour sets Blueprint Asset %s"), *SetterObject->GetName());
 				return true;
 			}
 
+			UE_LOG(LogRemoteControl, Error, TEXT("Path Behaviour unable to delete old Actor %s"), *OldActor->GetName());
 			OldActor->UObject::Rename(*OldActorName.ToString(), OldActor->GetOuter(), REN_DoNotDirty | REN_DontCreateRedirectors | REN_ForceNoResetLoaders);
 			OldActor->RegisterAllComponents();
 		}
 	}
 	
+	UE_LOG(LogRemoteControl, Error, TEXT("Path Behaviour fails to set Asset %s"), *SetterObject->GetName());
 	return false;
 }
 
@@ -186,3 +203,218 @@ TArray<UClass*> URCSetAssetByPathBehaviour::GetSupportedClasses() const
 
 	return SupportedClasses;
 }
+
+FString URCSetAssetByPathBehaviour::GetCurrentPath()
+{
+	FString CurrentPath;
+	URCController* RCController = ControllerWeakPtr.Get();
+	if (!RCController)
+	{
+		ensureMsgf(false, TEXT("Controller is invalid."));
+		return TEXT("Path is invalid!");
+	}
+
+	CurrentPath.Empty();
+	FString ControllerString;
+	RCController->GetValueString(ControllerString);
+
+	// Add Path String Concat
+	CurrentPath = bInternal ? SetAssetByPathBehaviourHelpers::ContentFolder : FString("");
+	for (FString PathPart : PathStruct.PathArray)
+	{
+		if (PathPart.IsEmpty())
+		{
+			continue;
+		}
+		
+		// Change *INPUT String Token to the correct Controller
+		if (PathPart.Contains(SetAssetByPathBehaviourHelpers::InputToken, ESearchCase::CaseSensitive)
+			&& PathPart.Find(SetAssetByPathBehaviourHelpers::InputToken) == 0)
+		{
+			const URemoteControlPreset* RemoteControlPreset = RCController->PresetWeakPtr.Get();
+			if (!RemoteControlPreset)
+			{
+				ensureMsgf(false, TEXT("Remote Control Preset is invalid"));
+				continue;
+			}
+			
+			PathPart.RemoveFromStart(SetAssetByPathBehaviourHelpers::InputToken);
+			// TODO: Test and rewrite to something better
+			int32 CharIndexStart;
+			int32 CharIndexEnd;
+			PathPart.FindChar('(', CharIndexStart);
+			PathPart.FindChar(')', CharIndexEnd);
+			if (CharIndexStart == INDEX_NONE || CharIndexEnd == INDEX_NONE)
+			{
+				// TODO: Add Message that Input token has been done wrongly.
+				ensureMsgf(false, TEXT("Token has been input wrongly"));
+				continue;
+			}
+
+			PathPart.RemoveAt(CharIndexStart);
+			PathPart.RemoveAt(CharIndexEnd-1);
+			const URCVirtualPropertyBase* TokenController = RemoteControlPreset->GetVirtualPropertyByDisplayName(FName(PathPart));
+			if (!TokenController)
+			{
+				ensureMsgf(false, TEXT("No Controller with given name found."));
+				continue;
+			}
+
+			PathPart = TokenController->GetDisplayValueAsString();
+		}
+
+		// Check for certain chars.
+		int32 IndexFound;
+		if (PathPart.FindLastChar('_', IndexFound) && IndexFound == PathPart.Len() - 1)
+		{
+			// In the case there's underscore char in the at the end of one of the Path Strings, do nothing to facilitate more complex pathing behaviours.
+		}
+		else if (!PathPart.FindLastChar('/', IndexFound) || IndexFound < PathPart.Len() - 1)
+		{
+			PathPart = PathPart.AppendChar('/');
+		}
+
+		PathPart.ReplaceCharInline(TCHAR('\\'), TCHAR('/'), ESearchCase::IgnoreCase);
+		
+		CurrentPath += PathPart;
+	}
+
+	return CurrentPath;
+}
+
+void URCSetAssetByPathBehaviour::SetTargetEntity(const TSharedPtr<const FRemoteControlEntity>& InEntity)
+{
+	TargetEntityId = InEntity->GetId();
+	TargetEntity = InEntity;
+}
+
+TWeakPtr<const FRemoteControlEntity> URCSetAssetByPathBehaviour::GetTargetEntity() const
+{
+	return TargetEntity;
+}
+
+bool URCSetAssetByPathBehaviour::SetExternalAsset(FString InExternalPath)
+{
+	URCController* Controller = ControllerWeakPtr.Get();
+	if (!Controller || !TargetEntity)
+	{
+		return false;
+	}
+	
+	URemoteControlPreset* PresetPtr = Controller->PresetWeakPtr.Get();
+	if (!PresetPtr)
+	{
+		return false;
+	}
+	
+	TArray<TWeakPtr<FRemoteControlProperty>> ExposedProperties = PresetPtr->GetExposedEntities<FRemoteControlProperty>();
+	const TWeakPtr<FRemoteControlProperty>* ExposedPropertyPtr = ExposedProperties.FindByPredicate([this](const TWeakPtr<FRemoteControlProperty>& Property)
+	{
+		if (const TSharedPtr<FRemoteControlProperty> PinnedProperty = Property.Pin())
+		{
+			return PinnedProperty->GetLabel() == TargetEntity->GetLabel();
+		}
+		
+		return false;
+	});
+	const TSharedPtr<FRemoteControlProperty> RemoteControlProperty = ExposedPropertyPtr->Pin();
+
+	if (!RemoteControlProperty)
+	{
+		FRemoteControlLogger::Get().Log(SetAssetByPathBehaviourHelpers::SetAssetByPathBehaviour, []
+		{
+			return FText::FromString(TEXT("Path Behaviour Exposed Property not found"));
+		});
+		return false;
+	}
+
+	// TODO: Find better solution.
+	UTexture2D* ImportedTexture = UKismetRenderingLibrary::ImportFileAsTexture2D(this, InExternalPath);
+	if (!ImportedTexture)
+	{
+		FRemoteControlLogger::Get().Log(SetAssetByPathBehaviourHelpers::SetAssetByPathBehaviour, [InExternalPath]
+		{
+			return FText::FromString(FString("Path Behaviour File not found: ") + *InExternalPath);
+		});
+		return false;
+	}
+
+	// Apply Texture
+	if (SetTextureAsset(RemoteControlProperty, ImportedTexture))
+	{
+		FRemoteControlLogger::Get().Log(SetAssetByPathBehaviourHelpers::SetAssetByPathBehaviour, [InExternalPath]
+		{
+			return FText::FromString(FString("Path Behaviour Set external Asset: ") + *InExternalPath);
+		});
+		return true;
+	}
+
+	FRemoteControlLogger::Get().Log(SetAssetByPathBehaviourHelpers::SetAssetByPathBehaviour, []
+	{
+		return FText::FromString(TEXT("Path Behaviour Set external Asset failed"));
+	});
+	return false;
+}
+
+bool URCSetAssetByPathBehaviour::SetTextureAsset(TSharedPtr<FRemoteControlProperty> InRemoteControlPropertyPtr, UTexture* InObject)
+{
+	if (!InObject)
+	{
+		return false;
+	}
+	
+	FRCObjectReference ObjectRef;
+	ObjectRef.Property = InRemoteControlPropertyPtr->GetProperty();
+	ObjectRef.Access = InRemoteControlPropertyPtr->GetPropertyHandle()->ShouldGenerateTransaction() ? ERCAccess::WRITE_TRANSACTION_ACCESS : ERCAccess::WRITE_ACCESS;
+	ObjectRef.PropertyPathInfo = InRemoteControlPropertyPtr->FieldPathInfo.ToString();
+
+	for (UObject* Object : InRemoteControlPropertyPtr->GetBoundObjects())
+	{
+		if (IRemoteControlModule::Get().ResolveObjectProperty(ObjectRef.Access, Object, ObjectRef.PropertyPathInfo, ObjectRef))
+		{
+			/**
+			 * Setting with Pointer and using Serialization of itself afterwards as a workaround with the Texture Asset not updating in the world.
+			 */
+			FObjectPropertyBase* ObjectProperty = CastField<FObjectPropertyBase>(InRemoteControlPropertyPtr->GetProperty());
+			if (!ObjectProperty)
+			{
+				return false;
+			}
+
+			FProperty* Property = InRemoteControlPropertyPtr->GetProperty();
+			uint8* ValueAddress = Property->ContainerPtrToValuePtr<uint8>(ObjectRef.ContainerAdress);
+			ObjectProperty->SetObjectPropertyValue(ValueAddress, InObject);
+
+			TArray<uint8> Buffer;
+			FMemoryReader Reader(Buffer);
+			FCborStructDeserializerBackend DeserializerBackend(Reader);
+
+			IRemoteControlModule::Get().SetObjectProperties(ObjectRef, DeserializerBackend, ERCPayloadType::Cbor, Buffer);
+			return true;
+		}
+	}
+	return false;	
+}
+
+void URCSetAssetByPathBehaviour::UpdateTargetEntity()
+{
+	const URCController* Controller = ControllerWeakPtr.Get();
+	if (!Controller)
+	{
+		return;
+	}
+	
+	const URemoteControlPreset* PresetPtr = Controller->PresetWeakPtr.Get();
+	if (!PresetPtr)
+	{
+		return;
+	}
+
+	TWeakPtr<const FRemoteControlEntity> ExposedEntity = PresetPtr->GetExposedEntity<FRemoteControlEntity>(TargetEntityId);
+	if (!ExposedEntity.Pin())
+	{
+		return;
+	}
+	
+	TargetEntity = ExposedEntity.Pin();
+}
diff --git a/Engine/Plugins/VirtualProduction/RemoteControl/Source/RemoteControlLogic/Public/Behaviour/RCSetAssetByPathBehaviour.h b/Engine/Plugins/VirtualProduction/RemoteControl/Source/RemoteControlLogic/Public/Behaviour/RCSetAssetByPathBehaviour.h
--- a/Engine/Plugins/VirtualProduction/RemoteControl/Source/RemoteControlLogic/Public/Behaviour/RCSetAssetByPathBehaviour.h
+++ b/Engine/Plugins/VirtualProduction/RemoteControl/Source/RemoteControlLogic/Public/Behaviour/RCSetAssetByPathBehaviour.h
@@ -4,17 +4,21 @@
 
 #include "RCBehaviour.h"
 #include "RCVirtualProperty.h"
+#include "RemoteControlEntity.h"
 
 #include "RCSetAssetByPathBehaviour.generated.h"
 
 class URCVirtualPropertyContainerBase;
 class URCUserDefinedStruct;
+class UTexture;
 
 namespace SetAssetByPathBehaviourHelpers
 {
 	const FString ContentFolder = FString(TEXT("/Game/"));
+	const FString InputToken = FString(TEXT("{INPUT}"));
 	const FName TargetProperty = FName(TEXT("Target Property"));
-	const FName DefaultProperty = FName(TEXT("Default Property"));
+	const FName DefaultInput = FName(TEXT("Default Input"));
+	const FName SetAssetByPathBehaviour = FName(TEXT("Set Asset By Path"));
 }
 
 /** Struct to help generate Widgts for the DetailsPanel of the Bahviour */
@@ -44,20 +48,57 @@ public:
 	virtual void Initialize() override;
 	//~ End URCBehaviour interface
 
-	bool SetAssetByPath(const FString& AssetPath, const FString& TargetPropertyString, const FString& DefaultString);
+	/** Given an Input Path, sets the Target Exposed Propterty to the Asset. */
+	bool SetAssetByPath(const FString& AssetPath, const FString& DefaultString);
 
+	/** List of Supported Assets SetAssetByPath Logic Behaviour can set and change. */
 	TArray<UClass*> GetSupportedClasses() const;
+
+	/** Returns the Path as the completed version, including Token Inputs. */
+	FString GetCurrentPath();
+
+	/** Sets the target entity upon which the asset will be set upon. */
+	void SetTargetEntity(const TSharedPtr<const FRemoteControlEntity>& InEntity);
+
+	/** Returns a Pointer to the current Target Entityy */
+	TWeakPtr<const FRemoteControlEntity> GetTargetEntity() const;
+
+	/** Auxialiary Function to apply to update the Target Texture */
+	void UpdateTargetEntity();
+	
 public:
 	/** Pointer to property container */
 	UPROPERTY()
 	TObjectPtr<URCVirtualPropertyContainerBase> PropertyInContainer;
 
+	/** Pointer to the current Class the Asset will use to set the Targeted Exposed Object. */
 	UPROPERTY()
-	TObjectPtr<UClass> AssetClass;
+	UClass* AssetClass = GetSupportedClasses()[0];
 
+	/** Struct holding the Path Information to help load Assets. */
 	UPROPERTY()
 	FRCSetAssetPath PathStruct;
 
+	/** Bool used to help tell if the path given is an internal or external one. */
+	UPROPERTY()
+	bool bInternal = true;
+
 private:
-	bool SetAsset(UObject* SetterObject, const FString& PropertyString);
+	/** Targeted Property Id */
+	UPROPERTY()
+	FGuid TargetEntityId;
+
+	/** Internal Targeted Property, used for any of the setter operations */
+	TSharedPtr<const FRemoteControlEntity> TargetEntity;
+	
+private:
+	/** Auxiliary Function which sets the given SetterObject Asset onto the Exposed Asset with the given PropertyString name. */
+	bool SetInternalAsset(UObject* SetterObject);
+
+	/** Given an External Path towards a external location, loads the asset associated and places it onto an object */
+	bool SetExternalAsset(FString InExternalPath);
+
+	/** Auxiliary Function to apply a Texture onto a given Property */
+	bool SetTextureAsset(TSharedPtr<FRemoteControlProperty> InRemoteControlPropertyPtr, UTexture* InObject);
+
 };
diff --git a/Engine/Plugins/VirtualProduction/RemoteControl/Source/RemoteControlUI/Private/UI/Behaviour/Builtin/RCBehaviourSetAssetByPathModel.cpp b/Engine/Plugins/VirtualProduction/RemoteControl/Source/RemoteControlUI/Private/UI/Behaviour/Builtin/RCBehaviourSetAssetByPathModel.cpp
--- a/Engine/Plugins/VirtualProduction/RemoteControl/Source/RemoteControlUI/Private/UI/Behaviour/Builtin/RCBehaviourSetAssetByPathModel.cpp
+++ b/Engine/Plugins/VirtualProduction/RemoteControl/Source/RemoteControlUI/Private/UI/Behaviour/Builtin/RCBehaviourSetAssetByPathModel.cpp
@@ -2,6 +2,8 @@
 
 #include "RCBehaviourSetAssetByPathModel.h"
 
+#include "Controller/RCController.h"
+#include "Editor.h"
 #include "IDetailTreeNode.h"
 #include "IPropertyRowGenerator.h"
 #include "Modules/ModuleManager.h"
@@ -9,7 +11,12 @@
 #include "RCVirtualPropertyContainer.h"
 #include "RemoteControlPreset.h"
 #include "SRCBehaviourSetAssetByPath.h"
+#include "Selection.h"
+#include "Widgets/Images/SImage.h"
+#include "Widgets/Input/SButton.h"
+#include "Widgets/Input/SComboBox.h"
 #include "Widgets/Layout/SBox.h"
+#include "Widgets/Text/STextBlock.h"
 
 #define LOCTEXT_NAMESPACE "FRCSetAssetByPathBehaviourModel"
 
@@ -38,9 +45,18 @@ FRCSetAssetByPathBehaviourModel::FRCSetAssetByPathBehaviourModel(URCSetAssetByPa
 
 		// Secondary TArray Struct
 		PropertyRowGeneratorArray->SetStructure(MakeShareable(new FStructOnScope(FRCSetAssetPath::StaticStruct(), (uint8*) &SetAssetByPathBehaviour->PathStruct)));
+		PropertyRowGeneratorArray->OnFinishedChangingProperties().AddLambda([this](const FPropertyChangedEvent& InEvent)
+		{
+			if (InEvent.ChangeType != EPropertyChangeType::ValueSet)
+			{
+				return;
+			}
+
+			RefreshPathAndPreview();
+		});
 		PropertyRowGeneratorArray->OnRowsRefreshed().AddLambda([this]()
 		{
-			RegeneratePathArrayWidget();
+			RefreshPathAndPreview();
 		});
 		DetailTreeNodeWeakPtrArray.Empty();
 		for (const TSharedRef<IDetailTreeNode>& CategoryNode : PropertyRowGeneratorArray->GetRootTreeNodes())
@@ -53,8 +69,16 @@ FRCSetAssetByPathBehaviourModel::FRCSetAssetByPathBehaviourModel(URCSetAssetByPa
 			}
 		}
 	}
+
+	PreviewPathWidget = SNew(STextBlock)
+		.AutoWrapText(true);
 	PathArrayWidget = SNew(SBox);
-	RegeneratePathArrayWidget();
+
+	Cast<URCSetAssetByPathBehaviour>(GetBehaviour())->UpdateTargetEntity();
+	TWeakPtr<const FRemoteControlEntity> InitialSelected = Cast<URCSetAssetByPathBehaviour>(GetBehaviour())->GetTargetEntity();
+	
+	RefreshPathAndPreview();
+	SelectorBox = GetSelectorWidget(InitialSelected);
 }
 
 TSharedRef<SWidget> FRCSetAssetByPathBehaviourModel::GetBehaviourDetailsWidget()
@@ -62,16 +86,14 @@ TSharedRef<SWidget> FRCSetAssetByPathBehaviourModel::GetBehaviourDetailsWidget()
 	return SNew(SRCBehaviourSetAssetByPath, SharedThis(this));
 }
 
-TSharedRef<SWidget> FRCSetAssetByPathBehaviourModel::GetPropertyWidget() const
+TSharedRef<SWidget> FRCSetAssetByPathBehaviourModel::GetPropertyWidget()
 {
 	TSharedRef<SVerticalBox> FieldWidget = SNew(SVerticalBox);
 	
-	const TSharedPtr<IDetailTreeNode> PinnedNodeTargetProperty = DetailTreeNodeWeakPtr[0];
-	const TSharedPtr<IDetailTreeNode> PinnedNodeDefault = DetailTreeNodeWeakPtr[1];
+	const TSharedPtr<IDetailTreeNode> PinnedNodeDefault = DetailTreeNodeWeakPtr[0];
 
-	const FNodeWidgets NodeWidgetsTarget = PinnedNodeTargetProperty->CreateNodeWidgets();
 	const FNodeWidgets NodeWidgetsDefault = PinnedNodeDefault->CreateNodeWidgets();
-	if (NodeWidgetsTarget.ValueWidget && NodeWidgetsDefault.ValueWidget)
+	if (NodeWidgetsDefault.ValueWidget)
 	{
 		FieldWidget->AddSlot()
 			.Padding(FMargin(3.0f, 2.0f))
@@ -80,7 +102,8 @@ TSharedRef<SWidget> FRCSetAssetByPathBehaviourModel::GetPropertyWidget() const
 				SNew(SHorizontalBox)
 				+ SHorizontalBox::Slot()
 				[
-					NodeWidgetsTarget.NameWidget.ToSharedRef()
+					SNew(STextBlock)
+						.Text(FText::FromString("Target Exposed Property"))
 				]
 				+ SHorizontalBox::Slot()
 				[
@@ -95,7 +118,7 @@ TSharedRef<SWidget> FRCSetAssetByPathBehaviourModel::GetPropertyWidget() const
 				SNew(SHorizontalBox)
 				+ SHorizontalBox::Slot()
 				[
-					NodeWidgetsTarget.ValueWidget.ToSharedRef()
+					SelectorBox.ToSharedRef()
 				]
 				+ SHorizontalBox::Slot()
 				[
@@ -103,15 +126,8 @@ TSharedRef<SWidget> FRCSetAssetByPathBehaviourModel::GetPropertyWidget() const
 				]
 			];
 	}
-	else if (NodeWidgetsTarget.WholeRowWidget && NodeWidgetsDefault.WholeRowWidget)
+	else if (NodeWidgetsDefault.WholeRowWidget)
 	{
-		FieldWidget->AddSlot()
-			.Padding(FMargin(3.0f, 2.0f))
-			.AutoHeight()
-			[
-				NodeWidgetsTarget.WholeRowWidget.ToSharedRef()
-			];
-		
 		FieldWidget->AddSlot()
 			.Padding(FMargin(3.0f, 2.0f))
 			.AutoHeight()
@@ -121,11 +137,29 @@ TSharedRef<SWidget> FRCSetAssetByPathBehaviourModel::GetPropertyWidget() const
 	}
 
 	FieldWidget->AddSlot()
-		.Padding(FMargin(3.f, 0.f))
+		.Padding(FMargin(3.f, 3.f))
 		.AutoHeight()
 		[
 			PathArrayWidget->AsShared()
 		];
+
+	FieldWidget->AddSlot()
+		.Padding(FMargin(3.f, 3.f))
+		.AutoHeight()
+		[
+			SNew(SHorizontalBox)
+			+ SHorizontalBox::Slot()
+				.AutoWidth()
+				[
+					SNew(STextBlock)
+					.Text(FText::FromString("Preview Path: "))
+				]
+			+ SHorizontalBox::Slot()
+				[
+					PreviewPathWidget->AsShared()
+				]
+		];
+
 	
 	return 	SNew(SHorizontalBox)
 		.Clipping(EWidgetClipping::OnDemand)
@@ -150,9 +184,29 @@ void FRCSetAssetByPathBehaviourModel::RegenerateWeakPtrInternal()
 			DetailTreeNodeWeakPtrArray.Add(Child);
 		}
 	}
+
+	URCSetAssetByPathBehaviour* SetAssetByPathBehaviour = SetAssetByPathBehaviourWeakPtr.Get();
+	if (!SetAssetByPathBehaviour)
+	{
+		return;
+	}
+
+	URCController* RCController = SetAssetByPathBehaviour->ControllerWeakPtr.Get();
+	if (!RCController)
+	{
+		return;
+	}
+
+	URemoteControlPreset* RemoteControlPreset =  RCController->PresetWeakPtr.Get();
+	if (!RemoteControlPreset)
+	{
+		return;
+	}
+
+	ExposedEntities = RemoteControlPreset->GetExposedEntities<const FRemoteControlEntity>();
 }
 
-void FRCSetAssetByPathBehaviourModel::RegeneratePathArrayWidget()
+void FRCSetAssetByPathBehaviourModel::RefreshPathAndPreview()
 {
 	TSharedPtr<SVerticalBox> FieldArrayWidget = SNew(SVerticalBox);
 
@@ -168,6 +222,49 @@ void FRCSetAssetByPathBehaviourModel::RegeneratePathArrayWidget()
 		
 		for (uint8 Counter = 0; Counter < Children.Num(); Counter++)
 		{
+			TSharedPtr<SWidget> Button = SNullWidget::NullWidget;
+			if (Counter == 0)
+			{
+				Button = SNew(SBox)
+					.HAlign(HAlign_Center)
+					.VAlign(VAlign_Center)
+					.WidthOverride(22)
+					.HeightOverride(22)
+					[
+						SNew(SButton)
+						.ButtonStyle( FAppStyle::Get(), "SimpleButton" )
+						.OnClicked_Lambda([this]()
+						{
+							if (URCSetAssetByPathBehaviour* Behaviour = Cast<URCSetAssetByPathBehaviour>(GetBehaviour()))
+							{
+								Behaviour->Modify();
+								CreateAssetPathFromSelection();
+								RefreshPathAndPreview();
+							}
+
+							return FReply::Handled();
+						})
+						.IsFocusable(false)
+						.ContentPadding(0)
+						.Visibility_Lambda([this]()
+						{
+							URCSetAssetByPathBehaviour* Behaviour = Cast<URCSetAssetByPathBehaviour>(GetBehaviour());
+							if (!Behaviour)
+							{
+								return EVisibility::Collapsed;
+							}
+
+							return Behaviour->bInternal ? EVisibility::Visible : EVisibility::Collapsed;
+						})
+						.Content()
+						[
+							SNew(SImage)
+							.Image( FAppStyle::GetBrush("Icons.Use") )
+							.ColorAndOpacity( FSlateColor::UseForeground() )
+						]
+					];
+			}
+			
 			const FNodeWidgets NodeWidgets = Children[Counter]->CreateNodeWidgets();
 			if (NodeWidgets.ValueWidget)
 			{
@@ -184,6 +281,11 @@ void FRCSetAssetByPathBehaviourModel::RegeneratePathArrayWidget()
 						[
 							NodeWidgets.ValueWidget.ToSharedRef()
 						]
+						+ SHorizontalBox::Slot()
+						.AutoWidth()
+						[
+							Button.ToSharedRef()
+						]
 					];
 			}
 			else if (NodeWidgets.WholeRowWidget)
@@ -197,7 +299,119 @@ void FRCSetAssetByPathBehaviourModel::RegeneratePathArrayWidget()
 			}
 		}
 	}
+
+	FString PreviewText = TEXT("Path is invalid!");
+	if (URCSetAssetByPathBehaviour* SetAssetByPathBehaviour = SetAssetByPathBehaviourWeakPtr.Get())
+	{
+		PreviewText = SetAssetByPathBehaviour->GetCurrentPath();
+	}
+	
 	PathArrayWidget->SetContent(FieldArrayWidget.ToSharedRef());
+	PreviewPathWidget->SetText(FText::FromString(PreviewText));
+}
+
+TSharedRef<SWidget> FRCSetAssetByPathBehaviourModel::GetSelectorWidget(TWeakPtr<const FRemoteControlEntity> InInitialSelected)
+{
+	return SNew(SComboBox<TWeakPtr<const FRemoteControlEntity>>)
+		.OptionsSource(&ExposedEntities)
+		.OnGenerateWidget_Lambda([this](TWeakPtr<const FRemoteControlEntity> InItem) -> TSharedRef<SWidget>
+		{
+			TSharedPtr<const FRemoteControlEntity> PinnedItem = InItem.Pin();
+			if (!PinnedItem)
+			{
+				return SNullWidget::NullWidget;
+			}
+
+			return SNew(STextBlock).Text(FText::FromString(*PinnedItem->GetLabel().ToString()));
+		})
+		.OnComboBoxOpening_Lambda([this]()
+		{
+			URCSetAssetByPathBehaviour* PathBehaviour = Cast<URCSetAssetByPathBehaviour>(GetBehaviour());
+			if (PathBehaviour)
+			{
+				return;
+			}
+			
+			URCController* Controller = PathBehaviour->ControllerWeakPtr.Get();
+			if (!Controller)
+			{
+				return;
+			}
+
+			URemoteControlPreset* Preset = Controller->PresetWeakPtr.Get();
+			if (!Preset)
+			{
+				return;
+			}
+			
+			ExposedEntities = Preset->GetExposedEntities<const FRemoteControlEntity>();
+		})
+		.OnSelectionChanged_Lambda([this](TWeakPtr<const FRemoteControlEntity> InItem, ESelectInfo::Type SelectType)
+		{
+			if (TSharedPtr<const FRemoteControlEntity> Entity = InItem.Pin())
+			{
+				if (URCSetAssetByPathBehaviour* PathBehaviour = Cast<URCSetAssetByPathBehaviour>(GetBehaviour()))
+				{
+					PathBehaviour->SetTargetEntity(Entity);
+				}
+			}
+		})
+		.InitiallySelectedItem(InInitialSelected)
+		.Content()
+		[
+			SNew(STextBlock)
+				.Text_Raw(this, &FRCSetAssetByPathBehaviourModel::GetSelectedEntityText)
+		];
 }
 
+FText FRCSetAssetByPathBehaviourModel::GetSelectedEntityText() const
+{
+	FText ReturnText = FText::FromString(*FString("Nothing selected yet"));
+
+	if (URCSetAssetByPathBehaviour* Behaviour = SetAssetByPathBehaviourWeakPtr.Get())
+	{
+		if (Behaviour->GetTargetEntity().Pin())
+		{
+			ReturnText = FText::FromString(*Behaviour->GetTargetEntity().Pin()->GetLabel().ToString());
+		}
+	}
+	return ReturnText;
+}
+
+
+void FRCSetAssetByPathBehaviourModel::CreateAssetPathFromSelection()
+{
+	URCSetAssetByPathBehaviour* SetAssetByPathBehaviour = Cast<URCSetAssetByPathBehaviour>(GetBehaviour());
+	if (!SetAssetByPathBehaviour)
+	{
+		return;
+	}
+	
+	// Clear it, in case it is an already used one.
+	TArray<FString>& PathArray = SetAssetByPathBehaviour->PathStruct.PathArray;
+	PathArray.Empty();
+	
+	TArray<FAssetData> AssetData;
+	GEditor->GetContentBrowserSelections(AssetData);
+	if (AssetData.Num() != 1)
+	{
+		// Only works with one selected Content
+		return;
+	}
+
+	int32 IndexOfLast;
+	const UObject* SelectedAsset = AssetData[0].GetAsset();
+	FString PathString = SelectedAsset->GetPathName();
+
+	// Remove the initial Game
+	PathString.RemoveFromStart("/Game/");
+
+	// Remove anything after the last /
+	PathString.FindLastChar('/', IndexOfLast);
+	PathString.RemoveAt(IndexOfLast, PathString.Len() - IndexOfLast);
+
+	PathArray.Add(PathString);
+}
+
+
 #undef LOCTEXT_NAMESPACE
diff --git a/Engine/Plugins/VirtualProduction/RemoteControl/Source/RemoteControlUI/Private/UI/Behaviour/Builtin/RCBehaviourSetAssetByPathModel.h b/Engine/Plugins/VirtualProduction/RemoteControl/Source/RemoteControlUI/Private/UI/Behaviour/Builtin/RCBehaviourSetAssetByPathModel.h
--- a/Engine/Plugins/VirtualProduction/RemoteControl/Source/RemoteControlUI/Private/UI/Behaviour/Builtin/RCBehaviourSetAssetByPathModel.h
+++ b/Engine/Plugins/VirtualProduction/RemoteControl/Source/RemoteControlUI/Private/UI/Behaviour/Builtin/RCBehaviourSetAssetByPathModel.h
@@ -1,12 +1,17 @@
 // Copyright Epic Games, Inc. All Rights Reserved.
 
 #pragma once
+#include "RemoteControlEntity.h"
 #include "Behaviour/RCSetAssetByPathBehaviour.h"
 #include "UI/Behaviour/RCBehaviourModel.h"
 
 class IDetailTreeNode;
 class SBox;
 
+namespace SetAssetByPathModelHelper
+{
+	const FString InputToken = FString(TEXT("*INPUT "));
+}
 
 /*
  * ~ FRCSetAssetByPathModel ~
@@ -29,7 +34,14 @@ public:
 	 * further two Text Widgets representing elements needed for the Path Behaviour.
 	 * All of them store the user input and use them to perform the SetAssetByPath Behaviour.
 	 */
-	TSharedRef<SWidget> GetPropertyWidget() const;
+	TSharedRef<SWidget> GetPropertyWidget();
+
+	/** Builds the Combobox which the Behaviour will use to select and and apply it's Asset Setting on. */
+	TSharedRef<SWidget> GetSelectorWidget(TWeakPtr<const FRemoteControlEntity> InInitialSelected);
+
+protected:
+	/** Gets the currently selected Entities Name, or a failure Text in case it fails. */
+	FText GetSelectedEntityText() const;
 	
 private:
 	/** The SetAssetByPath Behaviour associated with this Model */
@@ -38,6 +50,12 @@ private:
 	/** Pointer to the Widget holding the PathArray created for this behaviour */
 	TSharedPtr<SBox> PathArrayWidget;
 
+	/** Pointer to the Preview Text Widget for the Path */
+	TSharedPtr<STextBlock> PreviewPathWidget;
+
+	/** POinter to the ComboBox Widget */
+	TSharedPtr<SWidget> SelectorBox;
+
 	/** The row generator used to build the generic value widgets */
 	TSharedPtr<IPropertyRowGenerator> PropertyRowGenerator;
 	
@@ -49,10 +67,16 @@ private:
 	
 	/** Used to create a generic Value Widget based on the Paths Available*/
 	TArray<TSharedPtr<IDetailTreeNode>> DetailTreeNodeWeakPtrArray;
-
+	
+	/** Array of Exposed Entities from the Preset to help create a Widget to choose the Target */
+	TArray<TWeakPtr<const FRemoteControlEntity>> ExposedEntities;
+	
 private:
 	void RegenerateWeakPtrInternal();
 	
 	/** Regenerates and creates a new PathArray Widget if changed */
-	void RegeneratePathArrayWidget();
+	void RefreshPathAndPreview();
+
+	/** Creates a asset path, based on the selected Asset in the content browser. */
+	void CreateAssetPathFromSelection();
 };
diff --git a/Engine/Plugins/VirtualProduction/RemoteControl/Source/RemoteControlUI/Private/UI/Behaviour/Builtin/SRCBehaviourSetAssetByPath.cpp b/Engine/Plugins/VirtualProduction/RemoteControl/Source/RemoteControlUI/Private/UI/Behaviour/Builtin/SRCBehaviourSetAssetByPath.cpp
--- a/Engine/Plugins/VirtualProduction/RemoteControl/Source/RemoteControlUI/Private/UI/Behaviour/Builtin/SRCBehaviourSetAssetByPath.cpp
+++ b/Engine/Plugins/VirtualProduction/RemoteControl/Source/RemoteControlUI/Private/UI/Behaviour/Builtin/SRCBehaviourSetAssetByPath.cpp
@@ -3,8 +3,10 @@
 #include "SRCBehaviourSetAssetByPath.h"
 
 #include "IDetailTreeNode.h"
+#include "RemoteControlPreset.h"
 #include "SPositiveActionButton.h"
 #include "SlateOptMacros.h"
+#include "Styling/StyleColors.h"
 #include "UI/Behaviour/Builtin/RCBehaviourSetAssetByPathModel.h"
 #include "UI/RemoteControlPanelStyle.h"
 #include "Widgets/Input/SButton.h"
@@ -15,12 +17,15 @@ class SPositiveActionButton;
 
 BEGIN_SLATE_FUNCTION_BUILD_OPTIMIZATION
 
-void SRCBehaviourSetAssetByPath::Construct(const FArguments& InArgs, TSharedRef<const FRCSetAssetByPathBehaviourModel> InBehaviourItem)
+void SRCBehaviourSetAssetByPath::Construct(const FArguments& InArgs, TSharedRef<FRCSetAssetByPathBehaviourModel> InBehaviourItem)
 {
-	
-
-	TSharedRef<SHorizontalBox> SelectedClassWidget = SNew(SHorizontalBox);
 	URCSetAssetByPathBehaviour* PathBehaviour = Cast<URCSetAssetByPathBehaviour>(InBehaviourItem->GetBehaviour());
+	TSharedRef<SHorizontalBox> InternalExternalSwitchWidget = SNew(SHorizontalBox);
+	TSharedRef<SHorizontalBox> SelectedClassWidget = SNew(SHorizontalBox)
+		.Visibility_Lambda([PathBehaviour]()
+		{
+			return PathBehaviour->bInternal ? EVisibility::Visible : EVisibility::Collapsed;
+		});
 
 	if (!PathBehaviour)
 	{
@@ -31,6 +36,46 @@ void SRCBehaviourSetAssetByPath::Construct(const FArguments& InArgs, TSharedRef<
 		return;
 	}
 		
+	InternalExternalSwitchWidget->AddSlot()
+	.HAlign(HAlign_Right)
+	.VAlign(VAlign_Center)
+	.Padding(FMargin(0.0f,0.0f))
+	.FillWidth(1.0f)
+	.AutoWidth()
+	[
+		SNew(SButton)
+		.Text(FText::FromString(FString("Internal")))
+		.ButtonColorAndOpacity_Lambda([PathBehaviour]()
+		{
+			FSlateColor Color = PathBehaviour->bInternal ? FAppStyle::Get().GetSlateColor("Colors.Highlight") : FAppStyle::Get().GetSlateColor("Colors.AccentWhite");
+			return Color;
+		})
+		.OnPressed_Lambda([PathBehaviour]()
+		{
+			PathBehaviour->bInternal = true;
+		})
+	];
+	
+	InternalExternalSwitchWidget->AddSlot()
+	.HAlign(HAlign_Right)
+	.VAlign(VAlign_Center)
+	.Padding(FMargin(0.0f,0.0f))
+	.FillWidth(1.0f)
+	.AutoWidth()
+	[
+		SNew(SButton)
+		.Text(FText::FromString(FString("External")))
+		.ButtonColorAndOpacity_Lambda([PathBehaviour]()
+		{
+			FSlateColor Color = PathBehaviour->bInternal ? FAppStyle::Get().GetSlateColor("Colors.AccentWhite") : FAppStyle::Get().GetSlateColor("Colors.Highlight");
+			return Color;
+		})
+		.OnPressed_Lambda([PathBehaviour]()
+		{
+			PathBehaviour->bInternal = false;
+		})
+	];
+	
 	for (UClass* SupportedClass : PathBehaviour->GetSupportedClasses())
 	{
 		SelectedClassWidget->AddSlot()
@@ -44,7 +89,7 @@ void SRCBehaviourSetAssetByPath::Construct(const FArguments& InArgs, TSharedRef<
 			.Text(FText::FromString(SupportedClass->GetName()))
 			.ButtonColorAndOpacity_Lambda([PathBehaviour, SupportedClass]()
 			{
-				FSlateColor Color = PathBehaviour->AssetClass == SupportedClass ? FAppStyle::Get().GetSlateColor("Colors.AccentBlue") : FAppStyle::Get().GetSlateColor("Colors.AccentWhite");
+				FSlateColor Color = PathBehaviour->AssetClass == SupportedClass ? FAppStyle::Get().GetSlateColor("Colors.Select") : FAppStyle::Get().GetSlateColor("Colors.White");
 				return Color;
 			})
 			.OnPressed_Lambda([PathBehaviour, SupportedClass]()
@@ -58,6 +103,11 @@ void SRCBehaviourSetAssetByPath::Construct(const FArguments& InArgs, TSharedRef<
 	[
 		SNew(SVerticalBox)
 		+ SVerticalBox::Slot()
+		[
+			InternalExternalSwitchWidget
+		]
+		+ SVerticalBox::Slot()
+		.Padding(0.f,12.f)
 		[
 			SelectedClassWidget
 		]
diff --git a/Engine/Plugins/VirtualProduction/RemoteControl/Source/RemoteControlUI/Private/UI/Behaviour/Builtin/SRCBehaviourSetAssetByPath.h b/Engine/Plugins/VirtualProduction/RemoteControl/Source/RemoteControlUI/Private/UI/Behaviour/Builtin/SRCBehaviourSetAssetByPath.h
--- a/Engine/Plugins/VirtualProduction/RemoteControl/Source/RemoteControlUI/Private/UI/Behaviour/Builtin/SRCBehaviourSetAssetByPath.h
+++ b/Engine/Plugins/VirtualProduction/RemoteControl/Source/RemoteControlUI/Private/UI/Behaviour/Builtin/SRCBehaviourSetAssetByPath.h
@@ -17,7 +17,7 @@ public:
 	SLATE_END_ARGS()
 
 	/** Constructs this widget with InArgs */
-	void Construct(const FArguments& InArgs, TSharedRef<const FRCSetAssetByPathBehaviourModel> InBehaviourItem);
+	void Construct(const FArguments& InArgs, TSharedRef<FRCSetAssetByPathBehaviourModel> InBehaviourItem);
 
 private:
 	/** The Behaviour (UI model) associated with us*/
