ObjectPtr upgrade for engine plugins used from UnrealEditor -allmodules and -useverse.  Automated upgrade process.

diff --git a/Engine/Plugins/VirtualProduction/RemoteControl/Source/RemoteControlLogic/Public/Behaviour/RCSetAssetByPathBehaviour.h b/Engine/Plugins/VirtualProduction/RemoteControl/Source/RemoteControlLogic/Public/Behaviour/RCSetAssetByPathBehaviour.h
--- a/Engine/Plugins/VirtualProduction/RemoteControl/Source/RemoteControlLogic/Public/Behaviour/RCSetAssetByPathBehaviour.h
+++ b/Engine/Plugins/VirtualProduction/RemoteControl/Source/RemoteControlLogic/Public/Behaviour/RCSetAssetByPathBehaviour.h
@@ -73,7 +73,7 @@ public:
 
 	/** Pointer to the current Class the Asset will use to set the Targeted Exposed Object. */
 	UPROPERTY()
-	UClass* AssetClass = GetSupportedClasses()[0];
+	TObjectPtr<UClass> AssetClass = GetSupportedClasses()[0];
 
 	/** Struct holding the Path Information to help load Assets. */
 	UPROPERTY()
