Fix Static Analyzer Warning for Remote Control Set Asset By Path

Fix a conditional bug.

diff --git a/Engine/Plugins/VirtualProduction/RemoteControl/Source/RemoteControlUI/Private/UI/Behaviour/Builtin/RCBehaviourSetAssetByPathModel.cpp b/Engine/Plugins/VirtualProduction/RemoteControl/Source/RemoteControlUI/Private/UI/Behaviour/Builtin/RCBehaviourSetAssetByPathModel.cpp
--- a/Engine/Plugins/VirtualProduction/RemoteControl/Source/RemoteControlUI/Private/UI/Behaviour/Builtin/RCBehaviourSetAssetByPathModel.cpp
+++ b/Engine/Plugins/VirtualProduction/RemoteControl/Source/RemoteControlUI/Private/UI/Behaviour/Builtin/RCBehaviourSetAssetByPathModel.cpp
@@ -326,19 +326,19 @@ TSharedRef<SWidget> FRCSetAssetByPathBehaviourModel::GetSelectorWidget(TWeakPtr<
 		})
 		.OnComboBoxOpening_Lambda([this]()
 		{
-			URCSetAssetByPathBehaviour* PathBehaviour = Cast<URCSetAssetByPathBehaviour>(GetBehaviour());
-			if (PathBehaviour)
+			const URCSetAssetByPathBehaviour* PathBehaviour = Cast<URCSetAssetByPathBehaviour>(GetBehaviour());
+			if (!PathBehaviour)
 			{
 				return;
 			}
 			
-			URCController* Controller = PathBehaviour->ControllerWeakPtr.Get();
+			const URCController* Controller = PathBehaviour->ControllerWeakPtr.Get();
 			if (!Controller)
 			{
 				return;
 			}
 
-			URemoteControlPreset* Preset = Controller->PresetWeakPtr.Get();
+			const URemoteControlPreset* Preset = Controller->PresetWeakPtr.Get();
 			if (!Preset)
 			{
 				return;
