1. Implement Protocol Specific Dynamic Columns

2. Implement Masking Support for Protocols

diff --git a/Engine/Plugins/VirtualProduction/RemoteControl/Source/RemoteControl/Private/RemoteControlField.cpp b/Engine/Plugins/VirtualProduction/RemoteControl/Source/RemoteControl/Private/RemoteControlField.cpp
--- a/Engine/Plugins/VirtualProduction/RemoteControl/Source/RemoteControl/Private/RemoteControlField.cpp
+++ b/Engine/Plugins/VirtualProduction/RemoteControl/Source/RemoteControl/Private/RemoteControlField.cpp
@@ -94,6 +94,21 @@ bool FRemoteControlField::CanBindObject(const UObject* InObjectToBind) const
 	return false;
 }
 
+void FRemoteControlField::ClearMask(ERCMask InMaskBit)
+{
+	ActiveMasks &= ~InMaskBit;
+}
+
+void FRemoteControlField::EnableMask(ERCMask InMaskBit)
+{
+	ActiveMasks |= InMaskBit;
+}
+
+bool FRemoteControlField::HasMask(ERCMask InMaskBit) const
+{
+	return (ActiveMasks & InMaskBit) != ERCMask::NoMask;
+}
+
 void FRemoteControlField::PostSerialize(const FArchive& Ar)
 {
 	if (Ar.IsLoading())
@@ -306,6 +321,24 @@ bool FRemoteControlProperty::IsBound() const
 	return !!GetProperty();
 }
 
+bool FRemoteControlProperty::HasOptionalMask() const
+{
+	if (const FStructProperty* StructProperty = CastField<FStructProperty>(GetProperty()))
+	{
+		return StructProperty->Struct == TBaseStructure<FLinearColor>::Get() ||
+			StructProperty->Struct == TBaseStructure<FVector4>::Get() ||
+			StructProperty->Struct == TBaseStructure<FQuat>::Get() ||
+			StructProperty->Struct == TBaseStructure<FIntVector4>::Get();
+	}
+
+	return FRemoteControlField::HasOptionalMask();
+}
+
+bool FRemoteControlProperty::SupportsMasking() const
+{
+	return IRemoteControlModule::Get().SupportsMasking(GetProperty());
+}
+
 bool FRemoteControlProperty::CheckIsBoundToPropertyPath(const FString& InPath) const
 {
 	return FieldPathInfo.ToPathPropertyString() == InPath;
diff --git a/Engine/Plugins/VirtualProduction/RemoteControl/Source/RemoteControl/Private/RemoteControlModule.cpp b/Engine/Plugins/VirtualProduction/RemoteControl/Source/RemoteControl/Private/RemoteControlModule.cpp
--- a/Engine/Plugins/VirtualProduction/RemoteControl/Source/RemoteControl/Private/RemoteControlModule.cpp
+++ b/Engine/Plugins/VirtualProduction/RemoteControl/Source/RemoteControl/Private/RemoteControlModule.cpp
@@ -38,6 +38,15 @@
 DEFINE_LOG_CATEGORY(LogRemoteControl);
 #define LOCTEXT_NAMESPACE "RemoteControl"
 
+TSet<UScriptStruct*> FRemoteControlModule::MaskingSupportedStructs = { TBaseStructure<FColor>::Get()
+	, TBaseStructure<FLinearColor>::Get()
+	, TBaseStructure<FRotator>::Get()
+	, TBaseStructure<FVector>::Get()
+	, TBaseStructure<FVector4>::Get()
+	, TBaseStructure<FIntVector>::Get()
+	, TBaseStructure<FIntVector4>::Get()
+};
+
 struct FRCInterceptionPayload
 {
 	TArray<uint8> Payload;
@@ -719,6 +728,47 @@ void FRemoteControlModule::UnregisterEmbeddedPreset(URemoteControlPreset* Preset
 	EmbeddedPresets.Remove(PresetName);
 }
 
+void FRemoteControlModule::PerformMasking(const TSharedRef<FRCMaskingOperation>& InMaskingOperation)
+{
+	TRACE_CPUPROFILER_EVENT_SCOPE(FRemoteControlModule::PerformMasking);
+
+	if (!InMaskingOperation->IsValid())
+	{
+		return;
+	}
+
+	// Since we do not have any masks we do not need to process anything.
+	if (InMaskingOperation->Masks == ERCMask::NoMask || !SupportsMasking(InMaskingOperation->ObjectRef.Property.Get()))
+	{
+		return;
+	}
+
+	if (ActiveMaskingOperations.Contains(InMaskingOperation))
+	{
+		constexpr bool bIsInteractive = true;
+
+		ApplyMaskedValues(InMaskingOperation->OperationId, bIsInteractive);
+
+		ActiveMaskingOperations.Remove(InMaskingOperation);
+	}
+	else
+	{
+		ActiveMaskingOperations.Add(InMaskingOperation);
+
+		CacheRawValues(InMaskingOperation->OperationId);
+	}
+}
+
+bool FRemoteControlModule::SupportsMasking(const FProperty* InProperty) const
+{
+	if (const FStructProperty* StructProperty = CastField<FStructProperty>(InProperty))
+	{
+		return MaskingSupportedStructs.Contains(StructProperty->Struct);
+	}
+
+	return false;
+}
+
 bool FRemoteControlModule::ResolveCall(const FString& ObjectPath, const FString& FunctionName, FRCCallReference& OutCallRef, FString* OutErrorText)
 {
 	bool bSuccess = true;
@@ -1766,6 +1816,245 @@ int32 FRemoteControlModule::EndManualEditorTransaction(const FGuid& TransactionI
 	return INDEX_NONE;
 }
 
+void FRemoteControlModule::ApplyMaskedValues(const FGuid& FromMaskingOperation, bool bIsInteractive)
+{
+	for (const TSharedPtr<FRCMaskingOperation>& MaskingOperation : ActiveMaskingOperations)
+	{
+		if (MaskingOperation->OperationId == FromMaskingOperation)
+		{
+			const bool bHasMaskA = MaskingOperation->HasMask(ERCMask::MaskA);
+			const bool bHasMaskB = MaskingOperation->HasMask(ERCMask::MaskB);
+			const bool bHasMaskC = MaskingOperation->HasMask(ERCMask::MaskC);
+			const bool bHasMaskD = MaskingOperation->HasMask(ERCMask::MaskD);
+
+			if (FStructProperty* ToStructProp = CastField<FStructProperty>(MaskingOperation->ObjectRef.Property.Get()))
+			{
+				if (UObject* OwningObject = MaskingOperation->ObjectRef.Object.Get())
+				{
+#if WITH_EDITOR
+					OwningObject->PreEditChange(ToStructProp);
+					OwningObject->Modify();
+#endif // WITH_EDITOR
+
+					switch (*ToStructProp->Struct->GetFName().ToEName())
+					{
+						case NAME_Vector:
+						{
+							if (const FVector* VectorProp = ToStructProp->ContainerPtrToValuePtr<FVector>(OwningObject))
+							{
+								FVector MaskedVector;
+
+								MaskedVector.X = bHasMaskA ? VectorProp->X : MaskingOperation->PreMaskingCache.X;
+								MaskedVector.Y = bHasMaskB ? VectorProp->Y : MaskingOperation->PreMaskingCache.Y;
+								MaskedVector.Z = bHasMaskC ? VectorProp->Z : MaskingOperation->PreMaskingCache.Z;
+
+								ToStructProp->SetValue_InContainer(OwningObject, &MaskedVector);
+							}
+						}
+						break;
+						
+						case NAME_Vector4:
+						{
+							if (const FVector4* Vector4Prop = ToStructProp->ContainerPtrToValuePtr<FVector4>(OwningObject))
+							{
+								FVector4 MaskedVector4;
+
+								MaskedVector4.X = bHasMaskA ? Vector4Prop->X : MaskingOperation->PreMaskingCache.X;
+								MaskedVector4.Y = bHasMaskB ? Vector4Prop->Y : MaskingOperation->PreMaskingCache.Y;
+								MaskedVector4.Z = bHasMaskC ? Vector4Prop->Z : MaskingOperation->PreMaskingCache.Z;
+								MaskedVector4.W = bHasMaskD ? Vector4Prop->W : MaskingOperation->PreMaskingCache.W;
+
+								ToStructProp->SetValue_InContainer(OwningObject, &MaskedVector4);
+							}
+						}
+						break;
+
+						case NAME_IntVector:
+						{
+							if (const FIntVector* IntVectorProp = ToStructProp->ContainerPtrToValuePtr<FIntVector>(OwningObject))
+							{
+								FIntVector MaskedIntVector;
+
+								MaskedIntVector.X = bHasMaskA ? IntVectorProp->X : MaskingOperation->PreMaskingCache.X;
+								MaskedIntVector.Y = bHasMaskB ? IntVectorProp->Y : MaskingOperation->PreMaskingCache.Y;
+								MaskedIntVector.Z = bHasMaskC ? IntVectorProp->Z : MaskingOperation->PreMaskingCache.Z;
+
+								ToStructProp->SetValue_InContainer(OwningObject, &MaskedIntVector);
+							}
+						}
+						break;
+
+						case NAME_IntVector4:
+						{
+							if (const FIntVector4* IntVector4Prop = ToStructProp->ContainerPtrToValuePtr<FIntVector4>(OwningObject))
+							{
+								FIntVector4 MaskedIntVector4;
+
+								MaskedIntVector4.X = bHasMaskA ? IntVector4Prop->X : MaskingOperation->PreMaskingCache.X;
+								MaskedIntVector4.Y = bHasMaskB ? IntVector4Prop->Y : MaskingOperation->PreMaskingCache.Y;
+								MaskedIntVector4.Z = bHasMaskC ? IntVector4Prop->Z : MaskingOperation->PreMaskingCache.Z;
+								MaskedIntVector4.W = bHasMaskD ? IntVector4Prop->W : MaskingOperation->PreMaskingCache.W;
+
+								ToStructProp->SetValue_InContainer(OwningObject, &MaskedIntVector4);
+							}
+						}
+						break;
+
+						case NAME_Rotator:
+						{
+							if (const FRotator* RotatorProp = ToStructProp->ContainerPtrToValuePtr<FRotator>(OwningObject))
+							{
+								FRotator MaskedRotator;
+
+								MaskedRotator.Roll = bHasMaskA ? RotatorProp->Roll : MaskingOperation->PreMaskingCache.X;
+								MaskedRotator.Pitch = bHasMaskB ? RotatorProp->Pitch : MaskingOperation->PreMaskingCache.Y;
+								MaskedRotator.Yaw = bHasMaskC ? RotatorProp->Yaw : MaskingOperation->PreMaskingCache.Z;
+
+								ToStructProp->SetValue_InContainer(OwningObject, &MaskedRotator);
+							}
+						}
+						break;
+
+						case NAME_Color:
+						{
+							if (const FColor* ColorProp = ToStructProp->ContainerPtrToValuePtr<FColor>(OwningObject))
+							{
+								FColor MaskedColor;
+
+								MaskedColor.R = bHasMaskA ? ColorProp->R : MaskingOperation->PreMaskingCache.X;
+								MaskedColor.G = bHasMaskB ? ColorProp->G : MaskingOperation->PreMaskingCache.Y;
+								MaskedColor.B = bHasMaskC ? ColorProp->B : MaskingOperation->PreMaskingCache.Z;
+
+								ToStructProp->SetValue_InContainer(OwningObject, &MaskedColor);
+							}
+						}
+						break;
+
+						case NAME_LinearColor:
+						{
+							if (const FLinearColor* LinearColorProp = ToStructProp->ContainerPtrToValuePtr<FLinearColor>(OwningObject))
+							{
+								FLinearColor MaskedLinearColor;
+
+								MaskedLinearColor.R = bHasMaskA ? LinearColorProp->R : MaskingOperation->PreMaskingCache.X;
+								MaskedLinearColor.G = bHasMaskB ? LinearColorProp->G : MaskingOperation->PreMaskingCache.Y;
+								MaskedLinearColor.B = bHasMaskC ? LinearColorProp->B : MaskingOperation->PreMaskingCache.Z;
+								MaskedLinearColor.A = bHasMaskD ? LinearColorProp->A : MaskingOperation->PreMaskingCache.W;
+
+								ToStructProp->SetValue_InContainer(OwningObject, &MaskedLinearColor);
+							}
+						}
+						break;
+					}
+
+#if WITH_EDITOR
+					FPropertyChangedEvent ChangeEvent(ToStructProp, bIsInteractive ? EPropertyChangeType::Interactive : EPropertyChangeType::ValueSet);
+					OwningObject->PostEditChangeProperty(ChangeEvent);
+#endif // WITH_EDITOR
+				}
+			}
+		}
+	}
+}
+
+void FRemoteControlModule::CacheRawValues(const FGuid& FromMaskingOperation)
+{
+	for (const TSharedPtr<FRCMaskingOperation>& MaskingOperation : ActiveMaskingOperations)
+	{
+		if (MaskingOperation->OperationId == FromMaskingOperation)
+		{
+			if (const FStructProperty* FromStructProp = CastField<FStructProperty>(MaskingOperation->ObjectRef.Property.Get()))
+			{
+				if (const UObject* OwningObject = MaskingOperation->ObjectRef.Object.Get())
+				{
+					switch (*FromStructProp->Struct->GetFName().ToEName())
+					{
+						case NAME_Vector:
+						{
+							if (const FVector* VectorProp = FromStructProp->ContainerPtrToValuePtr<FVector>(OwningObject))
+							{
+								MaskingOperation->PreMaskingCache.X = VectorProp->X;
+								MaskingOperation->PreMaskingCache.Y = VectorProp->Y;
+								MaskingOperation->PreMaskingCache.Z = VectorProp->Z;
+							}
+						}
+						break;
+						
+						case NAME_Vector4:
+						{
+							if (const FVector4* Vector4Prop = FromStructProp->ContainerPtrToValuePtr<FVector4>(OwningObject))
+							{
+								MaskingOperation->PreMaskingCache.X = Vector4Prop->X;
+								MaskingOperation->PreMaskingCache.Y = Vector4Prop->Y;
+								MaskingOperation->PreMaskingCache.Z = Vector4Prop->Z;
+								MaskingOperation->PreMaskingCache.W = Vector4Prop->W;
+							}
+						}
+						break;
+						
+						case NAME_IntVector:
+						{
+							if (const FIntVector* IntVectorProp = FromStructProp->ContainerPtrToValuePtr<FIntVector>(OwningObject))
+							{
+								MaskingOperation->PreMaskingCache.X = IntVectorProp->X;
+								MaskingOperation->PreMaskingCache.Y = IntVectorProp->Y;
+								MaskingOperation->PreMaskingCache.Z = IntVectorProp->Z;
+							}
+						}
+						break;
+						
+						case NAME_IntVector4:
+						{
+							if (const FIntVector4* IntVector4Prop = FromStructProp->ContainerPtrToValuePtr<FIntVector4>(OwningObject))
+							{
+								MaskingOperation->PreMaskingCache.X = IntVector4Prop->X;
+								MaskingOperation->PreMaskingCache.Y = IntVector4Prop->Y;
+								MaskingOperation->PreMaskingCache.Z = IntVector4Prop->Z;
+								MaskingOperation->PreMaskingCache.W = IntVector4Prop->W;
+							}
+						}
+						break;
+						
+						case NAME_Rotator:
+						{
+							if (const FRotator* RotatorProp = FromStructProp->ContainerPtrToValuePtr<FRotator>(OwningObject))
+							{
+								MaskingOperation->PreMaskingCache.X = RotatorProp->Roll;
+								MaskingOperation->PreMaskingCache.Y = RotatorProp->Pitch;
+								MaskingOperation->PreMaskingCache.Z = RotatorProp->Yaw;
+							}
+						}
+						break;
+
+						case NAME_Color:
+						{
+							if (const FColor* ColorProp = FromStructProp->ContainerPtrToValuePtr<FColor>(OwningObject))
+							{
+								MaskingOperation->PreMaskingCache.X = ColorProp->R;
+								MaskingOperation->PreMaskingCache.Y = ColorProp->G;
+								MaskingOperation->PreMaskingCache.Z = ColorProp->B;
+							}
+						}
+						break;
+
+						case NAME_LinearColor:
+						{
+							if (const FLinearColor* LinearColorProp = FromStructProp->ContainerPtrToValuePtr<FLinearColor>(OwningObject))
+							{
+								MaskingOperation->PreMaskingCache.X = LinearColorProp->R;
+								MaskingOperation->PreMaskingCache.Y = LinearColorProp->G;
+								MaskingOperation->PreMaskingCache.Z = LinearColorProp->B;
+								MaskingOperation->PreMaskingCache.W = LinearColorProp->A;
+							}
+						}
+						break;
+					}
+				}
+			}
+		}
+	}
+}
+
 void FRemoteControlModule::CachePresets() const
 {
 	TArray<FAssetData> Assets;
diff --git a/Engine/Plugins/VirtualProduction/RemoteControl/Source/RemoteControl/Private/RemoteControlModule.h b/Engine/Plugins/VirtualProduction/RemoteControl/Source/RemoteControl/Private/RemoteControlModule.h
--- a/Engine/Plugins/VirtualProduction/RemoteControl/Source/RemoteControl/Private/RemoteControlModule.h
+++ b/Engine/Plugins/VirtualProduction/RemoteControl/Source/RemoteControl/Private/RemoteControlModule.h
@@ -29,6 +29,8 @@ public:
 	virtual bool RegisterEmbeddedPreset(URemoteControlPreset* Preset, bool bReplaceExisting) override;
 	virtual void UnregisterEmbeddedPreset(FName Name) override;
 	virtual void UnregisterEmbeddedPreset(URemoteControlPreset* Preset) override;
+	virtual void PerformMasking(const TSharedRef<FRCMaskingOperation>& InMaskingOperation) override;
+	virtual bool SupportsMasking(const FProperty* InProperty) const override;
 	virtual bool ResolveCall(const FString& ObjectPath, const FString& FunctionName, FRCCallReference& OutCallRef, FString* OutErrorText) override;
 	virtual bool InvokeCall(FRCCall& InCall, ERCPayloadType InPayloadType = ERCPayloadType::Json, const TArray<uint8>& InInterceptPayload = TArray<uint8>()) override;
 	virtual bool ResolveObject(ERCAccess AccessType, const FString& ObjectPath, const FString& PropertyName, FRCObjectReference& OutObjectRef, FString* OutErrorText = nullptr) override;
@@ -62,6 +64,12 @@ public:
 	//~ End IRemoteControlModule
 
 private:
+
+	/** Applies masked values to the given struct property. */
+	void ApplyMaskedValues(const FGuid& FromMaskingOperation, bool bIsInteractive);
+	/** Caches premasking values from the given struct property. */
+	void CacheRawValues(const FGuid& ToMaskingOperation);
+
 	/** Cache all presets in the project for the ResolvePreset function. */
 	void CachePresets() const;
 	
@@ -107,8 +115,8 @@ private:
 	//~ Register/Unregister editor delegates.
 	void RegisterEditorDelegates();
 	void UnregisterEditorDelegates();
-	
-#endif
+
+#endif // WITH_EDITOR
 
 private:
 	/** Cache of preset names to preset assets */
@@ -176,10 +184,17 @@ private:
 
 	/** Delay before we check if a modification is no longer ongoing. */
 	static constexpr float SecondsBetweenOngoingChangeCheck = 0.2f;
-#endif
+
+#endif // WITH_EDITOR
 
 	/** Map of the factories which is responsible for the Remote Control property creation */
 	TMap<FName, TSharedPtr<IRemoteControlPropertyFactory>> EntityFactories;
+
+	/** Holds the set of active masking operations. */
+	TSet<TSharedPtr<FRCMaskingOperation>> ActiveMaskingOperations;
+
+	/** Holds the set of structs supported by masking. */
+	static TSet<UScriptStruct*> MaskingSupportedStructs;
 };
 
 PRAGMA_ENABLE_DEPRECATION_WARNINGS
diff --git a/Engine/Plugins/VirtualProduction/RemoteControl/Source/RemoteControl/Private/RemoteControlProtocolBinding.cpp b/Engine/Plugins/VirtualProduction/RemoteControl/Source/RemoteControl/Private/RemoteControlProtocolBinding.cpp
--- a/Engine/Plugins/VirtualProduction/RemoteControl/Source/RemoteControl/Private/RemoteControlProtocolBinding.cpp
+++ b/Engine/Plugins/VirtualProduction/RemoteControl/Source/RemoteControl/Private/RemoteControlProtocolBinding.cpp
@@ -14,6 +14,14 @@
 #include "UObject/StructOnScope.h"
 #include "UObject/TextProperty.h"
 
+#if WITH_EDITOR
+
+#include "Widgets/Images/SImage.h"
+#include "Widgets/Input/SButton.h"
+
+#endif // WITH_EDITOR
+
+
 #define LOCTEXT_NAMESPACE "RemoteControl"
 
 namespace EntityInterpolation
@@ -691,6 +699,12 @@ bool FRemoteControlProtocolEntity::ApplyProtocolValueToProperty(double InProtoco
 	{
 		IRemoteControlModule::Get().ResolveObjectProperty(ObjectRef.Access, Object, ObjectRef.PropertyPathInfo, ObjectRef);
 
+		TSharedPtr<FRCMaskingOperation> MaskingOperation = MakeShared<FRCMaskingOperation>(ObjectRef.PropertyPathInfo, Object);
+		MaskingOperation->Masks = RemoteControlProperty->GetActiveMasks();
+
+		// Cache the values.
+		IRemoteControlModule::Get().PerformMasking(MaskingOperation.ToSharedRef());
+
 		// Set properties after interpolation
 		TArray<uint8> InterpolatedBuffer;
 		if (GetInterpolatedPropertyBuffer(Property, InProtocolValue, InterpolatedBuffer))
@@ -698,6 +712,9 @@ bool FRemoteControlProtocolEntity::ApplyProtocolValueToProperty(double InProtoco
 			FMemoryReader MemoryReader(InterpolatedBuffer);
 			FCborStructDeserializerBackend CborStructDeserializerBackend(MemoryReader);
 			bSuccess &= IRemoteControlModule::Get().SetObjectProperties(ObjectRef, CborStructDeserializerBackend, ERCPayloadType::Cbor, InterpolatedBuffer);
+
+			// Apply the masked the values.
+			IRemoteControlModule::Get().PerformMasking(MaskingOperation.ToSharedRef());
 		}
 	}
 
@@ -723,6 +740,71 @@ void FRemoteControlProtocolEntity::ResetDefaultBindingState()
 	BindingStatus = ERCBindingStatus::Unassigned;
 }
 
+#if WITH_EDITOR
+
+TSharedRef<SWidget> FRemoteControlProtocolEntity::GetWidget(const FName& ForColumnName)
+{
+	RegisterWidgets();
+
+	TSharedPtr<SWidget> ProtocolWidget = SNullWidget::NullWidget;
+
+	for (const TPair<FName, TSharedPtr<SWidget>>& Widget : Widgets)
+	{
+		if (Widget.Key == ForColumnName && Widget.Value.IsValid())
+		{
+			return Widget.Value.ToSharedRef();
+		}
+	}
+
+	return ProtocolWidget.ToSharedRef();
+}
+
+void FRemoteControlProtocolEntity::RegisterWidgets()
+{
+	if (!Widgets.Contains(FRemoteControlDefaultProtocolColumns::BindingStatus))
+	{
+		TSharedPtr<SWidget> BindingStatusWidget = SNew(SButton)
+			.ButtonStyle(FAppStyle::Get(), "HoverHintOnly")
+			.ToolTipText(LOCTEXT("RecordingButtonToolTip", "Status of the protocol entity binding"))
+			.ForegroundColor(FSlateColor::UseForeground())
+			.OnClicked_Lambda([this]()
+				{
+					ToggleBindingStatus();
+
+					return FReply::Handled();
+				}
+			)
+			.Content()
+			[
+				SNew(SImage)
+				.ColorAndOpacity_Lambda([this]()
+					{
+						const ERCBindingStatus ActiveBindingStatus = GetBindingStatus();
+				
+						switch (ActiveBindingStatus)
+						{
+							case ERCBindingStatus::Awaiting:
+								return FLinearColor::Red;
+							case ERCBindingStatus::Bound:
+								return FLinearColor::Green;
+							case ERCBindingStatus::Unassigned:
+								return FLinearColor::Gray;
+							default:
+								checkNoEntry();
+						}
+
+						return FLinearColor::Black;
+					}
+				)
+				.Image(FAppStyle::Get().GetBrush(TEXT("Icons.FilledCircle")))
+			];
+
+		Widgets.Add(FRemoteControlDefaultProtocolColumns::BindingStatus, BindingStatusWidget);
+	}
+}
+
+#endif // WITH_EDITOR
+
 bool FRemoteControlProtocolEntity::GetInterpolatedPropertyBuffer(FProperty* InProperty, double InProtocolValue, TArray<uint8>& OutBuffer)
 {
 	OutBuffer.Empty();
diff --git a/Engine/Plugins/VirtualProduction/RemoteControl/Source/RemoteControl/Public/IRemoteControlModule.h b/Engine/Plugins/VirtualProduction/RemoteControl/Source/RemoteControl/Public/IRemoteControlModule.h
--- a/Engine/Plugins/VirtualProduction/RemoteControl/Source/RemoteControl/Public/IRemoteControlModule.h
+++ b/Engine/Plugins/VirtualProduction/RemoteControl/Source/RemoteControl/Public/IRemoteControlModule.h
@@ -212,6 +212,61 @@ struct FRCObjectReference
 	FRCFieldPathInfo PropertyPathInfo;
 };
 
+/**
+ * Reference to a UObject or one of its properties for the purpose of masking.
+ */
+struct FRCMaskingOperation
+{
+	FRCMaskingOperation() = default;
+
+	explicit FRCMaskingOperation(FRCFieldPathInfo InPathInfo, UObject* InObject)
+		: OperationId(FGuid::NewGuid())
+		, ObjectRef(ERCAccess::NO_ACCESS, InObject, InPathInfo)
+	{
+		check(InObject);
+	}
+	
+	explicit FRCMaskingOperation(const FRCObjectReference& InObjectRef)
+		: OperationId(FGuid::NewGuid())
+		, ObjectRef(InObjectRef)
+	{
+	}
+
+	bool HasMask(ERCMask InMaskBit) const
+	{
+		return (Masks & InMaskBit) != ERCMask::NoMask;
+	}
+
+	bool IsValid() const
+	{
+		return OperationId.IsValid() && ObjectRef.IsValid();
+	}
+
+	friend bool operator==(const FRCMaskingOperation& LHS, const FRCMaskingOperation& RHS)
+	{
+		return LHS.OperationId == RHS.OperationId && LHS.ObjectRef == RHS.ObjectRef;
+	}
+
+	friend uint32 GetTypeHash(const FRCMaskingOperation& MaskingOperation)
+	{
+		return HashCombine(GetTypeHash(MaskingOperation.OperationId), GetTypeHash(MaskingOperation.ObjectRef));
+	}
+
+public:
+
+	/** Unique identifier of the operation being performed. */
+	FGuid OperationId;
+
+	/** Masks to be applied. */
+	ERCMask Masks = ERCMask::NoMask;
+	
+	/** Holds Object reference. */
+	FRCObjectReference ObjectRef;
+
+	/** Holds the state of this RC property before applying any masking. */
+	FVector4 PreMaskingCache = FVector4::Zero();
+};
+
 /**
  * Interface for the remote control module.
  */
@@ -264,6 +319,17 @@ public:
 	virtual void UnregisterEmbeddedPreset(FName Name) = 0;
 	virtual void UnregisterEmbeddedPreset(URemoteControlPreset* Preset) = 0;
 
+	/**
+	 * Performs the given masking operation.
+	 * @param InMaskingOperation Masking operation to be performed.
+	 */
+	virtual void PerformMasking(const TSharedRef<FRCMaskingOperation>& InMaskingOperation) = 0;
+
+	/**
+	 * Returns true if the given property can be masked, false otherwise.
+	 */
+	virtual bool SupportsMasking(const FProperty* InProperty) const = 0;
+
 	/**
 	 * Resolve a RemoteCall Object and Function.
 	 * This will look for object and function to resolve. 
diff --git a/Engine/Plugins/VirtualProduction/RemoteControl/Source/RemoteControl/Public/RemoteControlField.h b/Engine/Plugins/VirtualProduction/RemoteControl/Source/RemoteControl/Public/RemoteControlField.h
--- a/Engine/Plugins/VirtualProduction/RemoteControl/Source/RemoteControl/Public/RemoteControlField.h
+++ b/Engine/Plugins/VirtualProduction/RemoteControl/Source/RemoteControl/Public/RemoteControlField.h
@@ -10,6 +10,33 @@
 class IRemoteControlPropertyHandle;
 struct FPropertyChangedEvent;
 
+UENUM()
+namespace ERCProtocolBinding
+{
+	enum Op
+	{
+		Added,
+
+		Removed
+	};
+}
+
+UENUM()
+enum class ERCMask : uint8
+{
+	NoMask = 0x00000000,
+
+	MaskA = 0x00000001,
+
+	MaskB = 0x00000002,
+
+	MaskC = 0x00000004,
+
+	MaskD = 0x00000008,
+};
+
+ENUM_CLASS_FLAGS(ERCMask);
+
 /**
  * The type of the exposed field.
  */
@@ -44,6 +71,20 @@ struct REMOTECONTROL_API FRemoteControlField : public FRemoteControlEntity
 	virtual bool CanBindObject(const UObject* InObjectToBind) const override;
 	//~ End FRemoteControlEntity interface
 
+	/** Disables the given mask. */
+	virtual void ClearMask(ERCMask InMaskBit);
+	/** Enables the given mask. */
+	virtual void EnableMask(ERCMask InMaskBit);
+	/** Returns true if the given mask is enabled, false otherwise. */
+	virtual bool HasMask(ERCMask InMaskBit) const;
+	/** Returns true if this field has optional mask. */
+	virtual bool HasOptionalMask() const { return false; }
+	/** Returns true if this field supports masking. */
+	virtual bool SupportsMasking() const { return false; }
+
+	/** Retrieves the enabled masks. */
+	ERCMask GetActiveMasks() const { return ActiveMasks; }
+
 	/**	Returns whether the field is only available in editor. */
 	bool IsEditorOnly() const { return bIsEditorOnly; }
 
@@ -90,6 +131,10 @@ protected:
 	/** Whether the field is only available in editor. */
 	UPROPERTY()
 	bool bIsEditorOnly = false;
+	
+	/** Holds the actively enabled masks. */
+	UPROPERTY()
+	ERCMask ActiveMasks = ERCMask::NoMask;
 
 protected:
 	FRemoteControlField(URemoteControlPreset* InPreset, EExposedFieldType InType, FName InLabel, FRCFieldPathInfo InFieldPathInfo, const TArray<URemoteControlBinding*> InBindings);
@@ -127,6 +172,11 @@ public:
 	virtual UClass* GetSupportedBindingClass() const override;
 	virtual bool IsBound() const override;
 	//~ End FRemoteControlEntity interface
+	
+	//~ Begin FRemoteControlField interface
+	virtual bool HasOptionalMask() const override;
+	virtual bool SupportsMasking() const override;
+	//~ End FRemoteControlField interface
 
 	/**
 	 * Check whether given property path is bound to the property 
@@ -156,11 +206,11 @@ public:
 	virtual void PostLoad() {}
 
 	/**
-	 * Should be called when property chanced on bound object
+	 * Should be called when property changed on bound object
 	 * @param InObject Edited object
 	 * @param InEvent  Change Event
 	 */
-	virtual void OnObjectPropertyChanged(UObject* InObject, FPropertyChangedEvent& InEvent)  {};
+	virtual void OnObjectPropertyChanged(UObject* InObject, FPropertyChangedEvent& InEvent) {}
 
 	/**
 	 * Get the underlying property.
@@ -202,10 +252,11 @@ private:
 	bool bIsEditableInPackaged = false;
 
 #if WITH_EDITOR
+
 	/** Cached edit condition path used to enable the exposed property's edit condition. */
 	FRCFieldPathInfo CachedEditConditionPath;
-#endif
-	
+
+#endif // WITH_EDITOR
 };
 
 /**
diff --git a/Engine/Plugins/VirtualProduction/RemoteControl/Source/RemoteControl/Public/RemoteControlProtocolBinding.h b/Engine/Plugins/VirtualProduction/RemoteControl/Source/RemoteControl/Public/RemoteControlProtocolBinding.h
--- a/Engine/Plugins/VirtualProduction/RemoteControl/Source/RemoteControl/Public/RemoteControlProtocolBinding.h
+++ b/Engine/Plugins/VirtualProduction/RemoteControl/Source/RemoteControl/Public/RemoteControlProtocolBinding.h
@@ -23,6 +23,15 @@ struct FRCObjectReference;
 
 using FGetProtocolMappingCallback = TFunctionRef<void(FRemoteControlProtocolMapping&)>;
 
+#if WITH_EDITOR
+
+namespace FRemoteControlDefaultProtocolColumns
+{
+	static FName BindingStatus = TEXT("BindingStatus");
+}
+
+#endif // WITH_EDITOR
+
 /**
  * Status of the binding 
  */
@@ -34,7 +43,6 @@ enum class ERCBindingStatus : uint8
 	Bound
 };
 
-
 /*
  * Mapping of the range of the values for the protocol
  * This class holds a generic range buffer.
@@ -417,6 +425,16 @@ public:
 	 */
 	virtual bool IsSame(const FRemoteControlProtocolEntity* InOther) { return false; }
 
+#if WITH_EDITOR
+
+	/** Retrieves the widget corresponding to the given column name. */
+	TSharedRef<SWidget> GetWidget(const FName& ForColumnName);
+
+	/** Register(s) all the widgets of this protocol entity. */
+	virtual void RegisterWidgets();
+
+#endif // WITH_EDITOR
+
 public:
 	/** Container for range and mapping value pointers, and an optional number of elements (arrays, strings). */
 	struct FRangeMappingData
@@ -516,6 +534,11 @@ protected:
 	 */
 	UPROPERTY()
 	TSet<FRemoteControlProtocolMapping> Mappings;
+	
+	/** 
+	 * Holds column specific widget(s) for this protocol.
+	 */
+	TMap<FName, TSharedPtr<SWidget>> Widgets;
 
 private:
 	/** Binding status of this protocol entity */
diff --git a/Engine/Plugins/VirtualProduction/RemoteControl/Source/RemoteControl/RemoteControl.Build.cs b/Engine/Plugins/VirtualProduction/RemoteControl/Source/RemoteControl/RemoteControl.Build.cs
--- a/Engine/Plugins/VirtualProduction/RemoteControl/Source/RemoteControl/RemoteControl.Build.cs
+++ b/Engine/Plugins/VirtualProduction/RemoteControl/Source/RemoteControl/RemoteControl.Build.cs
@@ -37,6 +37,8 @@ public class RemoteControl : ModuleRules
 					"AssetTools",
 					"DeveloperSettings",
 					"MessageLog",
+					"Slate",
+					"SlateCore",
 					"UnrealEd",
 				}
 			);
diff --git a/Engine/Plugins/VirtualProduction/RemoteControl/Source/RemoteControlProtocol/Private/RemoteControlProtocol.cpp b/Engine/Plugins/VirtualProduction/RemoteControl/Source/RemoteControlProtocol/Private/RemoteControlProtocol.cpp
--- a/Engine/Plugins/VirtualProduction/RemoteControl/Source/RemoteControlProtocol/Private/RemoteControlProtocol.cpp
+++ b/Engine/Plugins/VirtualProduction/RemoteControl/Source/RemoteControlProtocol/Private/RemoteControlProtocol.cpp
@@ -20,6 +20,15 @@ FRemoteControlProtocol::~FRemoteControlProtocol()
 	FCoreDelegates::OnEndFrame.RemoveAll(this);
 }
 
+void FRemoteControlProtocol::Init()
+{
+#if WITH_EDITOR
+
+	RegisterColumns();
+
+#endif // WITH_EDITOR
+}
+
 FRemoteControlProtocolEntityPtr FRemoteControlProtocol::CreateNewProtocolEntity(FProperty* InProperty, URemoteControlPreset* InOwner, FGuid InPropertyId) const
 {
 	FRemoteControlProtocolEntityPtr NewDataPtr = MakeShared<TStructOnScope<FRemoteControlProtocolEntity>>();
@@ -62,6 +71,31 @@ void FRemoteControlProtocol::OnEndFrame()
 	PreviousTickValuesToApply = MoveTemp(EntityValuesToApply);
 }
 
+#if WITH_EDITOR
+
+FProtocolColumnPtr FRemoteControlProtocol::GetRegisteredColumn(const FName& ByColumnName) const
+{
+	for (const FProtocolColumnPtr& ProtocolColumn : RegisteredColumns)
+	{
+		if (ProtocolColumn->ColumnName == ByColumnName)
+		{
+			return ProtocolColumn.ToSharedRef();
+		}
+	}
+
+	return nullptr;
+}
+
+void FRemoteControlProtocol::GetRegisteredColumns(TSet<FName>& OutColumns)
+{
+	for (const FProtocolColumnPtr& ProtocolColumn : RegisteredColumns)
+	{
+		OutColumns.Add(ProtocolColumn->ColumnName);
+	}
+}
+
+#endif // WITH_EDITOR
+
 TFunction<bool(FRemoteControlProtocolEntityWeakPtr InProtocolEntityWeakPtr)> FRemoteControlProtocol::CreateProtocolComparator(FGuid InPropertyId)
 {
 	return [InPropertyId](FRemoteControlProtocolEntityWeakPtr InProtocolEntityWeakPtr) -> bool
@@ -75,6 +109,17 @@ TFunction<bool(FRemoteControlProtocolEntityWeakPtr InProtocolEntityWeakPtr)> FRe
 	};
 }
 
+#if WITH_EDITOR
+
+void FRemoteControlProtocol::RegisterColumns()
+{
+	REGISTER_COLUMN(FRemoteControlDefaultProtocolColumns::BindingStatus
+		, NSLOCTEXT("RemoteControlProtocol", "RCPresetBindingStatusColumnHeader", "REC")
+		, 33.f);
+}
+
+#endif // WITH_EDITOR
+
 FProperty* IRemoteControlProtocol::GetRangeInputTemplateProperty() const
 {
 	FProperty* RangeInputTemplateProperty = nullptr;
diff --git a/Engine/Plugins/VirtualProduction/RemoteControl/Source/RemoteControlProtocol/Public/IRemoteControlProtocol.h b/Engine/Plugins/VirtualProduction/RemoteControl/Source/RemoteControlProtocol/Public/IRemoteControlProtocol.h
--- a/Engine/Plugins/VirtualProduction/RemoteControl/Source/RemoteControlProtocol/Public/IRemoteControlProtocol.h
+++ b/Engine/Plugins/VirtualProduction/RemoteControl/Source/RemoteControlProtocol/Public/IRemoteControlProtocol.h
@@ -12,6 +12,62 @@ struct FRemoteControlProtocolEntity;
 using FRemoteControlProtocolEntityPtr = TSharedPtr<TStructOnScope<FRemoteControlProtocolEntity>>;
 using FRemoteControlProtocolEntityWeakPtr = TWeakPtr<TStructOnScope<FRemoteControlProtocolEntity>>;
 
+#if WITH_EDITOR
+
+namespace ProtocolColumnConstants
+{
+	static float ColumnSizeMicro = 0.075f;
+	static float ColumnSizeMini = 0.15f;
+	static float ColumnSizeSmall = 0.3f;
+	static float ColumnSizeMedium = 0.5f;
+	static float ColumnSizeLarge = 0.8f;
+	static float ColumnSizeExtraLarge = 1.f;
+}
+
+/**
+ * Simple utility struct to represent a protocol column.
+ */
+struct FProtocolColumn
+{
+	FProtocolColumn() = default;
+
+	FProtocolColumn(const FName& InColumnName, const FText& InDisplayText, float InColumnSize)
+		: ColumnName(InColumnName)
+		, DisplayText(InDisplayText)
+		, ColumnSize(InColumnSize)
+	{
+	}
+
+	friend uint32 GetTypeHash(const FProtocolColumn& OtherColumn)
+	{
+		return GetTypeHash(OtherColumn.ColumnName);
+	}
+
+	bool operator==(const FProtocolColumn& OtherColumn)
+	{
+		return ColumnName == OtherColumn.ColumnName;
+	}
+	
+	bool operator==(const FProtocolColumn& OtherColumn) const
+	{
+		return ColumnName == OtherColumn.ColumnName;
+	}
+
+	/** Holds the name of the column. */
+	FName ColumnName;
+
+	/** Holds the display text of the column. */
+	FText DisplayText;
+
+	/** Holds the size of the column. */
+	float ColumnSize;
+};
+
+using FProtocolColumnPtr = TSharedPtr<FProtocolColumn>;
+using FProtocolColumnRef = TSharedRef<FProtocolColumn>;
+
+#endif // WITH_EDITOR
+
 /**
  * Interface for remote control protocol
  */
@@ -22,7 +78,7 @@ public:
 	virtual ~IRemoteControlProtocol() {}
 
 	/** Add custom initializing of the protocol */
-	virtual void Init() {}
+	virtual void Init() = 0;
 
 	/**
 	 * Add the new protocol entity. That is created based on protocol specific script struct
@@ -65,4 +121,14 @@ public:
 	 * @param InProtocolValue Protocol Value
 	 */
 	virtual void QueueValue(const FRemoteControlProtocolEntityPtr InProtocolEntity, const double InProtocolValue) = 0;
+
+#if WITH_EDITOR
+
+	/** Retrieves the column reference by its name. */
+	virtual FProtocolColumnPtr GetRegisteredColumn(const FName& ByColumnName) const = 0;
+
+	/** Retrieves the set of columns registered by this protocol. */
+	virtual void GetRegisteredColumns(TSet<FName>& OutColumns) = 0;
+	
+#endif // WITH_EDITOR
 };
diff --git a/Engine/Plugins/VirtualProduction/RemoteControl/Source/RemoteControlProtocol/Public/RemoteControlProtocol.h b/Engine/Plugins/VirtualProduction/RemoteControl/Source/RemoteControlProtocol/Public/RemoteControlProtocol.h
--- a/Engine/Plugins/VirtualProduction/RemoteControl/Source/RemoteControlProtocol/Public/RemoteControlProtocol.h
+++ b/Engine/Plugins/VirtualProduction/RemoteControl/Source/RemoteControlProtocol/Public/RemoteControlProtocol.h
@@ -4,6 +4,12 @@
 
 #include "IRemoteControlProtocol.h"
 
+#if WITH_EDITOR
+
+#define REGISTER_COLUMN(ColumnName, DisplayText, ColumnSize) RegisteredColumns.Add(MakeShared<FProtocolColumn>(ColumnName, DisplayText, ColumnSize));
+
+#endif // WITH_EDITOR
+
 /**
  * Base class implementation for remote control protocol
  */
@@ -14,9 +20,19 @@ public:
 	~FRemoteControlProtocol();
 	
 	//~ Begin IRemoteControlProtocol interface
+	virtual void Init() override;
 	virtual FRemoteControlProtocolEntityPtr CreateNewProtocolEntity(FProperty* InProperty, URemoteControlPreset* InOwner, FGuid InPropertyId) const override;
 	virtual void QueueValue(const FRemoteControlProtocolEntityPtr InProtocolEntity, const double InProtocolValue) override;
 	virtual void OnEndFrame() override;
+
+#if WITH_EDITOR
+
+	virtual FProtocolColumnPtr GetRegisteredColumn(const FName& ByColumnName) const override;
+
+	virtual void GetRegisteredColumns(TSet<FName>& OutColumns) override;
+
+#endif // WITH_EDITOR
+
 	//~ End IRemoteControlProtocol interface
 
 	/**
@@ -26,10 +42,26 @@ public:
 	 */
 	static TFunction<bool(FRemoteControlProtocolEntityWeakPtr InProtocolEntityWeakPtr)> CreateProtocolComparator(FGuid InPropertyId);
 
+#if WITH_EDITOR
+
+protected:
+
+	/** Populates protocol specific columns. */
+	virtual void RegisterColumns();
+
+#endif // WITH_EDITOR
+
 protected:
 	/** Current Protocol Name */
 	FName ProtocolName;
 
+#if WITH_EDITOR
+
+	/** Holds a set of protocol specific columns. */
+	TSet<FProtocolColumnPtr> RegisteredColumns;
+
+#endif // WITH_EDITOR
+
 private:
 	/** Map of the entities and protocol values about to apply */
 	TMap<const FRemoteControlProtocolEntityPtr, double> EntityValuesToApply;
diff --git a/Engine/Plugins/VirtualProduction/RemoteControl/Source/RemoteControlProtocolWidgets/Private/RemoteControlProtocolWidgetsModule.cpp b/Engine/Plugins/VirtualProduction/RemoteControl/Source/RemoteControlProtocolWidgets/Private/RemoteControlProtocolWidgetsModule.cpp
--- a/Engine/Plugins/VirtualProduction/RemoteControl/Source/RemoteControlProtocolWidgets/Private/RemoteControlProtocolWidgetsModule.cpp
+++ b/Engine/Plugins/VirtualProduction/RemoteControl/Source/RemoteControlProtocolWidgets/Private/RemoteControlProtocolWidgetsModule.cpp
@@ -7,12 +7,23 @@
 #include "IRemoteControlProtocolModule.h"
 #include "IRemoteControlModule.h"
 #include "RemoteControlPreset.h"
+#include "Styling/ProtocolPanelStyle.h"
 #include "ViewModels/ProtocolEntityViewModel.h"
 #include "Widgets/SNullWidget.h"
 #include "Widgets/SRCProtocolBindingList.h"
 
 DEFINE_LOG_CATEGORY(LogRemoteControlProtocolWidgets);
 
+void FRemoteControlProtocolWidgetsModule::StartupModule()
+{
+	FProtocolPanelStyle::Initialize();
+}
+
+void FRemoteControlProtocolWidgetsModule::ShutdownModule()
+{
+	FProtocolPanelStyle::Shutdown();
+}
+
 TSharedRef<SWidget> FRemoteControlProtocolWidgetsModule::GenerateDetailsForEntity(URemoteControlPreset* InPreset, const FGuid& InFieldId, const EExposedFieldType& InFieldType)
 {
 	check(InPreset);
@@ -29,7 +40,13 @@ TSharedRef<SWidget> FRemoteControlProtocolWidgetsModule::GenerateDetailsForEntit
 	{
 		if (PropertyStruct->IsChildOf(FRemoteControlProperty::StaticStruct()))
 		{
-			return SAssignNew(RCProtocolBindingList, SRCProtocolBindingList, FProtocolEntityViewModel::Create(InPreset, InFieldId));
+			TSharedRef<FProtocolEntityViewModel> EntityViewModel = FProtocolEntityViewModel::Create(InPreset, InFieldId);
+
+			EntityViewModel->OnBindingAdded().AddRaw(this, &FRemoteControlProtocolWidgetsModule::OnBindingAdded);
+
+			EntityViewModel->OnBindingRemoved().AddRaw(this, &FRemoteControlProtocolWidgetsModule::OnBindingRemoved);
+
+			return SAssignNew(RCProtocolBindingList, SRCProtocolBindingList, EntityViewModel);
 		}
 	}
 
@@ -46,4 +63,19 @@ TSharedPtr<IRCProtocolBindingList> FRemoteControlProtocolWidgetsModule::GetProto
 	return RCProtocolBindingList;
 }
 
+FOnProtocolBindingAddedOrRemoved& FRemoteControlProtocolWidgetsModule::OnProtocolBindingAddedOrRemoved()
+{
+	return OnProtocolBindingAddedOrRemovedDelegate;
+}
+
+void FRemoteControlProtocolWidgetsModule::OnBindingAdded(TSharedRef<FProtocolBindingViewModel> InBindingViewModel)
+{
+	OnProtocolBindingAddedOrRemovedDelegate.Broadcast(ERCProtocolBinding::Added);
+}
+
+void FRemoteControlProtocolWidgetsModule::OnBindingRemoved(FGuid InBindingId)
+{
+	OnProtocolBindingAddedOrRemovedDelegate.Broadcast(ERCProtocolBinding::Removed);
+}
+
 IMPLEMENT_MODULE(FRemoteControlProtocolWidgetsModule, RemoteControlProtocolWidgets);
diff --git a/Engine/Plugins/VirtualProduction/RemoteControl/Source/RemoteControlProtocolWidgets/Private/RemoteControlProtocolWidgetsModule.h b/Engine/Plugins/VirtualProduction/RemoteControl/Source/RemoteControlProtocolWidgets/Private/RemoteControlProtocolWidgetsModule.h
--- a/Engine/Plugins/VirtualProduction/RemoteControl/Source/RemoteControlProtocolWidgets/Private/RemoteControlProtocolWidgetsModule.h
+++ b/Engine/Plugins/VirtualProduction/RemoteControl/Source/RemoteControlProtocolWidgets/Private/RemoteControlProtocolWidgetsModule.h
@@ -6,19 +6,38 @@
 
 REMOTECONTROLPROTOCOLWIDGETS_API DECLARE_LOG_CATEGORY_EXTERN(LogRemoteControlProtocolWidgets, Log, All);
 
+class FProtocolBindingViewModel;
 class IRCProtocolBindingList;
 class URemoteControlPreset;
 
 class FRemoteControlProtocolWidgetsModule : public IRemoteControlProtocolWidgetsModule
 {
 public:
+
+	//~ Begin IModuleInterface
+	virtual void StartupModule() override;
+	virtual void ShutdownModule() override;
+	//~ End IModuleInterface
+
 	//~ Begin IRemoteControlProtocolWidgetsModule Interface
 	virtual TSharedRef<SWidget> GenerateDetailsForEntity(URemoteControlPreset* InPreset, const FGuid& InFieldId, const EExposedFieldType& InFieldType) override;
 	virtual void ResetProtocolBindingList() override;
 	virtual TSharedPtr<IRCProtocolBindingList> GetProtocolBindingList() const override;
+	virtual FOnProtocolBindingAddedOrRemoved& OnProtocolBindingAddedOrRemoved() override;
 	//~ End IRemoteControlProtocolWidgetsModule Interface
 
+private:
+
+	/** Called when binding is added. */
+	void OnBindingAdded(TSharedRef<FProtocolBindingViewModel> InBindingViewModel);
+	
+	/** Called when binding is removed. */
+	void OnBindingRemoved(FGuid InBindingId);
+
 private:
 	/** Binding list public interface instance */
 	TSharedPtr<IRCProtocolBindingList> RCProtocolBindingList;
+
+	/** Called when binding is added or removed. */
+	FOnProtocolBindingAddedOrRemoved OnProtocolBindingAddedOrRemovedDelegate;
 };
diff --git a/Engine/Plugins/VirtualProduction/RemoteControl/Source/RemoteControlProtocolWidgets/Private/Styling/ProtocolPanelStyle.cpp b/Engine/Plugins/VirtualProduction/RemoteControl/Source/RemoteControlProtocolWidgets/Private/Styling/ProtocolPanelStyle.cpp
new file mode 100644
--- /dev/null
+++ b/Engine/Plugins/VirtualProduction/RemoteControl/Source/RemoteControlProtocolWidgets/Private/Styling/ProtocolPanelStyle.cpp
@@ -0,0 +1,105 @@
+// Copyright Epic Games, Inc. All Rights Reserved.
+
+#include "Styling/ProtocolPanelStyle.h"
+#include "Styling/AppStyle.h"
+#include "Styling/CoreStyle.h"
+#include "Styling/ProtocolStyles.h"
+#include "Styling/SlateStyleRegistry.h"
+#include "Styling/SlateTypes.h"
+#include "Styling/StarshipCoreStyle.h"
+#include "Styling/StyleColors.h"
+
+#include "Interfaces/IPluginManager.h"
+
+#define IMAGE_BRUSH(RelativePath, ...) FSlateImageBrush(StyleSet->RootToContentDir(RelativePath, TEXT(".png")), __VA_ARGS__)
+#define IMAGE_BRUSH_SVG( RelativePath, ... ) FSlateVectorImageBrush(StyleSet->RootToContentDir(RelativePath, TEXT(".svg")), __VA_ARGS__)
+#define BOX_BRUSH(RelativePath, ...) FSlateBoxBrush(StyleSet->RootToContentDir(RelativePath, TEXT(".png")), __VA_ARGS__)
+#define CORE_BOX_BRUSH( RelativePath, ... ) FSlateBoxBrush(StyleSet->RootToCoreContentDir(RelativePath, TEXT(".png") ), __VA_ARGS__)
+#define BOX_PLUGIN_BRUSH( RelativePath, ... ) FSlateBoxBrush(FProtocolPanelStyle::InContent( RelativePath, ".png" ), __VA_ARGS__)
+#define DEFAULT_FONT(...) FAppStyle::Get().GetDefaultFontStyle(__VA_ARGS__)
+
+TSharedPtr<FSlateStyleSet> FProtocolPanelStyle::StyleSet;
+
+void FProtocolPanelStyle::Initialize()
+{
+	if (StyleSet.IsValid())
+	{
+		return;
+	}
+
+	StyleSet = MakeShared<FSlateStyleSet>(GetStyleSetName());
+	StyleSet->SetContentRoot(FPaths::EngineContentDir() / TEXT("Editor/Slate"));
+	StyleSet->SetCoreContentRoot(FPaths::EngineContentDir() / TEXT("Slate"));
+
+	// Protocol Widget Styles
+	SetupWidgetStyles(StyleSet.ToSharedRef());
+
+	FSlateStyleRegistry::RegisterSlateStyle(*StyleSet.Get());
+}
+
+void FProtocolPanelStyle::Shutdown()
+{
+	if (StyleSet.IsValid())
+	{
+		FSlateStyleRegistry::UnRegisterSlateStyle(*StyleSet.Get());
+		ensure(StyleSet.IsUnique());
+		StyleSet.Reset();
+	}
+}
+
+TSharedPtr<ISlateStyle> FProtocolPanelStyle::Get()
+{
+	return StyleSet;
+}
+
+FName FProtocolPanelStyle::GetStyleSetName()
+{
+	static const FName ProtocolPanelStyleName(TEXT("ProtocolPanelStyle"));
+	return ProtocolPanelStyleName;
+}
+
+FString FProtocolPanelStyle::InContent(const FString& RelativePath, const ANSICHAR* Extension)
+{
+	static const FString ContentDir = IPluginManager::Get().FindPlugin(TEXT("RemoteControl"))->GetBaseDir() + TEXT("/Resources");
+	return (ContentDir / RelativePath) + Extension;
+}
+
+void FProtocolPanelStyle::SetupWidgetStyles(TSharedRef<FSlateStyleSet> InStyle)
+{
+	const ISlateStyle& AppStyle = FAppStyle::Get();
+
+	// Mask Widget Styles
+	FSlateBrush ContentAreaBrushDark = BOX_BRUSH("Common/DarkGroupBorder", FMargin(4.f / 16.f), FLinearColor(0.5f, 0.5f, 0.5f, 1.f));
+	FSlateBrush ContentAreaBrushLight = BOX_BRUSH("Common/LightGroupBorder", FMargin(4.f / 16.f));
+
+	FCheckBoxStyle MaskButtonStyle = AppStyle.GetWidgetStyle<FCheckBoxStyle>("ToggleButtonCheckbox");
+	MaskButtonStyle.SetCheckBoxType(ESlateCheckBoxType::ToggleButton);
+	MaskButtonStyle.SetPadding(0.f);
+
+	// Text Styles
+	const FStyleFonts& StyleFonts = FStyleFonts::Get();
+
+	const FTextBlockStyle& PlainTextStyle = AppStyle.GetWidgetStyle<FTextBlockStyle>("NormalText");
+
+	const FTextBlockStyle& BoldTextStyle = FTextBlockStyle(PlainTextStyle)
+		.SetFont(StyleFonts.NormalBold);
+
+	FProtocolWidgetStyle ProtocolMaskWidgetStyle = FProtocolWidgetStyle()
+		.SetContentAreaBrush(FSlateColorBrush(FStyleColors::Panel))
+		.SetContentAreaBrushDark(ContentAreaBrushDark)
+		.SetContentAreaBrushLight(ContentAreaBrushLight)
+		
+		.SetMaskButtonStyle(MaskButtonStyle)
+
+		.SetBoldTextStyle(BoldTextStyle)
+		.SetPlainTextStyle(PlainTextStyle);
+
+	StyleSet->Set("ProtocolsPanel.Widgets.Mask", ProtocolMaskWidgetStyle);
+}
+
+#undef IMAGE_BRUSH
+#undef IMAGE_BRUSH_SVG
+#undef BOX_BRUSH
+#undef CORE_BOX_BRUSH
+#undef BOX_PLUGIN_BRUSH
+#undef DEFAULT_FONT
diff --git a/Engine/Plugins/VirtualProduction/RemoteControl/Source/RemoteControlProtocolWidgets/Private/Styling/ProtocolPanelStyle.h b/Engine/Plugins/VirtualProduction/RemoteControl/Source/RemoteControlProtocolWidgets/Private/Styling/ProtocolPanelStyle.h
new file mode 100644
--- /dev/null
+++ b/Engine/Plugins/VirtualProduction/RemoteControl/Source/RemoteControlProtocolWidgets/Private/Styling/ProtocolPanelStyle.h
@@ -0,0 +1,48 @@
+// Copyright Epic Games, Inc. All Rights Reserved.
+
+#pragma once
+
+#include "CoreMinimal.h"
+
+#include "Styling/ISlateStyle.h"
+#include "Styling/SlateStyle.h"
+
+/**
+ * 
+ */
+class FProtocolPanelStyle
+{
+public:
+	/**
+	 * Initialize the style set.
+	 */
+	static void Initialize();
+
+	/**
+	 * Shutdown the style set.
+	 */
+	static void Shutdown();
+
+	/**
+	 * Static getter for the style set.
+	 */
+	static TSharedPtr<ISlateStyle> Get();
+
+	/**
+	 * Static getter for the style set name.
+	 */
+	static FName GetStyleSetName();
+
+private:
+
+	/** Get a path to the module's resources folder. */
+	static FString InContent(const FString& RelativePath, const ANSICHAR* Extension);
+
+	/** Initializes Protocol Widget Styles. */
+	static void SetupWidgetStyles(TSharedRef<FSlateStyleSet> InStyle);
+
+private:
+
+	/** Holds the style set. */
+	static TSharedPtr<FSlateStyleSet> StyleSet;
+};
\ No newline at end of file
diff --git a/Engine/Plugins/VirtualProduction/RemoteControl/Source/RemoteControlProtocolWidgets/Private/Styling/ProtocolStyles.cpp b/Engine/Plugins/VirtualProduction/RemoteControl/Source/RemoteControlProtocolWidgets/Private/Styling/ProtocolStyles.cpp
new file mode 100644
--- /dev/null
+++ b/Engine/Plugins/VirtualProduction/RemoteControl/Source/RemoteControlProtocolWidgets/Private/Styling/ProtocolStyles.cpp
@@ -0,0 +1,26 @@
+// Copyright Epic Games, Inc. All Rights Reserved.
+
+#include "Styling/ProtocolStyles.h"
+
+FProtocolWidgetStyle::FProtocolWidgetStyle()
+	: BoldTextStyle(FTextBlockStyle::GetDefault())
+	, PlainTextStyle(FTextBlockStyle::GetDefault())
+{
+}
+
+const FName FProtocolWidgetStyle::TypeName(TEXT("FProtocolWidgetStyle"));
+
+const FProtocolWidgetStyle& FProtocolWidgetStyle::GetDefault()
+{
+	static FProtocolWidgetStyle Default;
+	return Default;
+}
+
+void FProtocolWidgetStyle::GetResources(TArray<const FSlateBrush*>& OutBrushes) const
+{
+	MaskButtonStyle.GetResources(OutBrushes);
+	OutBrushes.Add(&ContentAreaBrush);
+	OutBrushes.Add(&ContentAreaBrushDark);
+	OutBrushes.Add(&ContentAreaBrushLight);
+}
+
diff --git a/Engine/Plugins/VirtualProduction/RemoteControl/Source/RemoteControlProtocolWidgets/Private/Styling/ProtocolStyles.h b/Engine/Plugins/VirtualProduction/RemoteControl/Source/RemoteControlProtocolWidgets/Private/Styling/ProtocolStyles.h
new file mode 100644
--- /dev/null
+++ b/Engine/Plugins/VirtualProduction/RemoteControl/Source/RemoteControlProtocolWidgets/Private/Styling/ProtocolStyles.h
@@ -0,0 +1,63 @@
+// Copyright Epic Games, Inc. All Rights Reserved.
+
+#pragma once
+
+#include "CoreMinimal.h"
+
+#include "Styling/SlateBrush.h"
+#include "Styling/SlateTypes.h"
+#include "Styling/SlateWidgetStyle.h"
+#include "Styling/SlateWidgetStyleContainerBase.h"
+
+#include "ProtocolStyles.generated.h"
+
+/**
+ * Represents all custom protocol widgets.
+ */
+USTRUCT()
+struct FProtocolWidgetStyle : public FSlateWidgetStyle
+{
+	GENERATED_USTRUCT_BODY()
+
+	FProtocolWidgetStyle();
+	virtual ~FProtocolWidgetStyle() {}
+
+	//~ Begin : FSlateWidgetStyle
+	virtual void GetResources(TArray<const FSlateBrush*>& OutBrushes) const override;
+	
+	static const FName TypeName;
+	virtual const FName GetTypeName() const override { return TypeName; };
+	
+	static const FProtocolWidgetStyle& GetDefault();
+	//~ End : FSlateWidgetStyle
+
+	/** Mask Button Bold Text Style */
+	UPROPERTY(EditAnywhere, Category = Appearance)
+		FTextBlockStyle BoldTextStyle;
+	FProtocolWidgetStyle& SetBoldTextStyle(const FTextBlockStyle& InBoldTextStyle) { BoldTextStyle = InBoldTextStyle; return *this; }
+
+	/** Brush used by the RC Panels to draw the content associated with this panel. */
+	UPROPERTY(EditAnywhere, Category = Appearance)
+		FSlateBrush ContentAreaBrush;
+	FProtocolWidgetStyle& SetContentAreaBrush(const FSlateBrush& InContentAreaBrush) { ContentAreaBrush = InContentAreaBrush; return *this; }
+
+	/** Brush used by the RC Panels to draw the content associated with this panel (Dark). */
+	UPROPERTY(EditAnywhere, Category = Appearance)
+		FSlateBrush ContentAreaBrushDark;
+	FProtocolWidgetStyle& SetContentAreaBrushDark(const FSlateBrush& InContentAreaBrushDark) { ContentAreaBrushDark = InContentAreaBrushDark; return *this; }
+
+	/** Brush used by the RC Panels to draw the content associated with this panel (Light). */
+	UPROPERTY(EditAnywhere, Category = Appearance)
+		FSlateBrush ContentAreaBrushLight;
+	FProtocolWidgetStyle& SetContentAreaBrushLight(const FSlateBrush& InContentAreaBrushLight) { ContentAreaBrushLight = InContentAreaBrushLight; return *this; }
+
+	/** Style used for the mask button */
+	UPROPERTY(EditAnywhere, Category = Appearance)
+		FCheckBoxStyle MaskButtonStyle;
+	FProtocolWidgetStyle& SetMaskButtonStyle(const FCheckBoxStyle& InMaskButtonStyle) { MaskButtonStyle = InMaskButtonStyle; return *this; }
+
+	/** Mask Button Plain Text Style */
+	UPROPERTY(EditAnywhere, Category = Appearance)
+		FTextBlockStyle PlainTextStyle;
+	FProtocolWidgetStyle& SetPlainTextStyle(const FTextBlockStyle& InPlainTextStyle) { PlainTextStyle = InPlainTextStyle; return *this; }
+};
diff --git a/Engine/Plugins/VirtualProduction/RemoteControl/Source/RemoteControlProtocolWidgets/Private/Widgets/Masks/SRCProtocolMask.cpp b/Engine/Plugins/VirtualProduction/RemoteControl/Source/RemoteControlProtocolWidgets/Private/Widgets/Masks/SRCProtocolMask.cpp
new file mode 100644
--- /dev/null
+++ b/Engine/Plugins/VirtualProduction/RemoteControl/Source/RemoteControlProtocolWidgets/Private/Widgets/Masks/SRCProtocolMask.cpp
@@ -0,0 +1,112 @@
+// Copyright Epic Games, Inc. All Rights Reserved.
+
+#include "Widgets/Masks/SRCProtocolMask.h"
+
+#include "RemoteControlField.h"
+
+#define LOCTEXT_NAMESPACE "SRCProtocolMask"
+
+namespace RemoteControlProtocolMasking
+{
+	static TMap<UScriptStruct*, EMaskingType> StructsToMaskingTypes = { { TBaseStructure<FColor>::Get(), EMaskingType::Color }
+		, { TBaseStructure<FLinearColor>::Get(), EMaskingType::Color }
+		, { TBaseStructure<FRotator>::Get(), EMaskingType::Rotator }
+		, { TBaseStructure<FVector>::Get(), EMaskingType::Vector }
+		, { TBaseStructure<FIntVector>::Get(), EMaskingType::Vector }
+		, { TBaseStructure<FVector4>::Get(), EMaskingType::Quat }
+		, { TBaseStructure<FIntVector4>::Get(), EMaskingType::Quat }
+		, { TBaseStructure<FQuat>::Get(), EMaskingType::Quat }
+	};
+}
+
+void SRCProtocolMask::Construct(const FArguments& InArgs, TWeakPtr<FRemoteControlField> InField)
+{
+	WeakField = MoveTemp(InField);
+
+	const bool bHasOptionalMask = HasOptionalMask();
+
+	const bool bCanBeMasked = CanBeMasked();
+
+	SRCProtocolMaskTriplet::Construct(SRCProtocolMaskTriplet::FArguments()
+		.MaskA(ERCMask::MaskA)
+		.MaskB(ERCMask::MaskB)
+		.MaskC(ERCMask::MaskC)
+		.OptionalMask(ERCMask::MaskD)
+		.MaskingType(GetMaskingType())
+		.CanBeMasked(bCanBeMasked)
+		.EnableOptionalMask(bHasOptionalMask)
+	);
+}
+
+bool SRCProtocolMask::CanBeMasked() const
+{
+	if (TSharedPtr<FRemoteControlField> RCField = WeakField.Pin())
+	{
+		return RCField->SupportsMasking();
+	}
+
+	return false;
+}
+
+EMaskingType SRCProtocolMask::GetMaskingType()
+{
+	if (CanBeMasked())
+	{
+		if (TSharedPtr<FRemoteControlField> RCField = WeakField.Pin())
+		{
+			if (TSharedPtr<FRemoteControlProperty> RCProperty = StaticCastSharedPtr<FRemoteControlProperty>(RCField))
+			{
+				if (const FStructProperty* StructProperty = CastField<FStructProperty>(RCProperty->GetProperty()))
+				{
+					if (const EMaskingType* StructToMaskingType = RemoteControlProtocolMasking::StructsToMaskingTypes.Find(StructProperty->Struct))
+					{
+						return *StructToMaskingType;
+					}
+				}
+			}
+		}
+	}
+
+	return EMaskingType::Unsupported;
+}
+
+bool SRCProtocolMask::HasOptionalMask() const
+{
+	if (TSharedPtr<FRemoteControlField> RCField = WeakField.Pin())
+	{
+		return RCField->HasOptionalMask();
+	}
+
+	return false;
+}
+
+ECheckBoxState SRCProtocolMask::IsMaskEnabled(ERCMask InMaskBit) const
+{
+	if (TSharedPtr<FRemoteControlField> RCField = WeakField.Pin())
+	{
+		return RCField->HasMask((ERCMask)InMaskBit) ? ECheckBoxState::Checked : ECheckBoxState::Unchecked;
+	}
+
+	return SRCProtocolMaskTriplet::IsMaskEnabled(InMaskBit);
+}
+
+void SRCProtocolMask::SetMaskEnabled(ECheckBoxState NewState, ERCMask NewMaskBit)
+{
+	if (TSharedPtr<FRemoteControlField> RCField = WeakField.Pin())
+	{
+		if (NewState == ECheckBoxState::Checked)
+		{
+			RCField->EnableMask((ERCMask)NewMaskBit);
+		}
+		else
+		{
+			RCField->ClearMask((ERCMask)NewMaskBit);
+		}
+	}
+	else
+	{
+		SRCProtocolMaskTriplet::SetMaskEnabled(NewState, NewMaskBit);
+	}
+}
+
+#undef LOCTEXT_NAMESPACE
diff --git a/Engine/Plugins/VirtualProduction/RemoteControl/Source/RemoteControlProtocolWidgets/Private/Widgets/Masks/SRCProtocolMaskButton.cpp b/Engine/Plugins/VirtualProduction/RemoteControl/Source/RemoteControlProtocolWidgets/Private/Widgets/Masks/SRCProtocolMaskButton.cpp
new file mode 100644
--- /dev/null
+++ b/Engine/Plugins/VirtualProduction/RemoteControl/Source/RemoteControlProtocolWidgets/Private/Widgets/Masks/SRCProtocolMaskButton.cpp
@@ -0,0 +1,150 @@
+// Copyright Epic Games, Inc. All Rights Reserved.
+
+#include "SRCProtocolMaskButton.h"
+
+#include "SlateOptMacros.h"
+
+#include "Styling/ProtocolPanelStyle.h"
+#include "Styling/ProtocolStyles.h"
+
+#include "Widgets/Input/SCheckBox.h"
+#include "Widgets/SOverlay.h"
+#include "Widgets/Text/STextBlock.h"
+
+#define LOCTEXT_NAMESPACE "SRCProtocolMaskButton"
+
+FText SRCProtocolMaskButton::MaskText = LOCTEXT("MaskLabel", "Mask");
+FText SRCProtocolMaskButton::UnmaskText = LOCTEXT("UnmaskLabel", "Unmask");
+
+BEGIN_SLATE_FUNCTION_BUILD_OPTIMIZATION
+
+void SRCProtocolMaskButton::Construct(const FArguments& InArgs)
+{
+	WidgetStyle = &FProtocolPanelStyle::Get()->GetWidgetStyle<FProtocolWidgetStyle>("ProtocolsPanel.Widgets.Mask");
+
+	MaskBit = InArgs._MaskBit;
+
+	MaskedState = ECheckBoxState::Unchecked;
+
+	IsInMaskedState = InArgs._IsMasked;
+
+	OnMaskStateChanged = InArgs._OnMasked;
+
+	DefaultLabel = FText::Format(LOCTEXT("DefaultMaskLabel", "{0}"), { InArgs._DefaultLabel.IsSet() ? InArgs._DefaultLabel.Get() : FText::FromString(TEXT("[LabelMissing]")) });
+
+	OptionalToolTipText = InArgs._ToolTip.IsSet() ? TAttribute<FText>() :
+		InArgs._OptionalTooltip.IsSet() ? InArgs._OptionalTooltip.Get() : TAttribute<FText>();
+
+	SBorder::Construct(SBorder::FArguments()
+		.Padding(0.f)
+		.BorderImage(&WidgetStyle->ContentAreaBrushDark)
+		.Content()
+		[
+			SAssignNew(Mask, SCheckBox)
+			.Style(&WidgetStyle->MaskButtonStyle)
+			.IsChecked(this, &SRCProtocolMaskButton::IsMaskEnabled)
+			.OnCheckStateChanged(this, &SRCProtocolMaskButton::SetMaskEnabled)
+			.HAlign(HAlign_Center)
+			.ToolTip(InArgs._ToolTip)
+			.ToolTipText(this, &SRCProtocolMaskButton::HandleMaskTooltip)
+			[
+				SNew(SOverlay)
+
+				+SOverlay::Slot()
+				.HAlign(HAlign_Fill)
+				.Padding(0.f, 2.f, 0.f, 0.f)
+				[
+					SNew(SBorder)
+					.BorderImage(&WidgetStyle->ContentAreaBrushLight)
+					.BorderBackgroundColor(InArgs._MaskColor)
+				]
+
+				+SOverlay::Slot()
+				.HAlign(HAlign_Center)
+				.Padding(0.f, 0.f, 0.f, 2.f)
+				[
+					SNew(SBorder)
+					.BorderImage(&WidgetStyle->ContentAreaBrushDark)
+					.HAlign(HAlign_Center)
+					.VAlign(VAlign_Center)
+					.Padding(FMargin(8.f, 2.f))
+					[
+						SAssignNew(MaskLabel, STextBlock)
+						.Text(DefaultLabel.Get())
+						.TextStyle(&WidgetStyle->PlainTextStyle)
+					]
+				]
+			]
+		]
+	);
+}
+
+END_SLATE_FUNCTION_BUILD_OPTIMIZATION
+
+bool SRCProtocolMaskButton::IsMasked() const
+{
+	if (Mask.IsValid())
+	{
+		return Mask->IsChecked();
+	}
+
+	return false;
+}
+
+void SRCProtocolMaskButton::ToggleMaskedState(bool bSoftToggle)
+{
+	if (bSoftToggle && MaskLabel.IsValid())
+	{
+		MaskLabel->SetTextStyle(&WidgetStyle->BoldTextStyle);
+
+		MaskLabel->SetColorAndOpacity(FLinearColor::White);
+	}
+	else if (Mask.IsValid())
+	{
+		Mask->ToggleCheckedState();
+	}
+}
+
+FText SRCProtocolMaskButton::HandleMaskTooltip() const
+{
+	if (OptionalToolTipText.Get().IsEmpty())
+	{
+		return FText::Format(LOCTEXT("DefaultMaskTooltip", "{0} {1}"), {  IsMasked() ? UnmaskText : MaskText,  DefaultLabel.Get() });
+	}
+
+	return OptionalToolTipText.Get();
+}
+
+ECheckBoxState SRCProtocolMaskButton::IsMaskEnabled() const
+{
+	if (IsInMaskedState.IsBound())
+	{
+		return IsInMaskedState.Execute(MaskBit);
+	}
+
+	return MaskedState.Get(ECheckBoxState::Unchecked);
+}
+
+void SRCProtocolMaskButton::SetMaskEnabled(ECheckBoxState NewState)
+{
+	const bool bIsChecked = NewState == ECheckBoxState::Checked && MaskLabel.IsValid();
+
+	MaskedState = NewState;
+
+	if (bIsChecked)
+	{
+		MaskLabel->SetTextStyle(&WidgetStyle->BoldTextStyle);
+
+		MaskLabel->SetColorAndOpacity(FLinearColor::White);
+	}
+	else
+	{
+		MaskLabel->SetTextStyle(&WidgetStyle->PlainTextStyle);
+
+		MaskLabel->SetColorAndOpacity(FSlateColor::UseForeground());
+	}
+
+	OnMaskStateChanged.ExecuteIfBound(NewState, MaskBit);
+}
+
+#undef LOCTEXT_NAMESPACE
diff --git a/Engine/Plugins/VirtualProduction/RemoteControl/Source/RemoteControlProtocolWidgets/Private/Widgets/Masks/SRCProtocolMaskButton.h b/Engine/Plugins/VirtualProduction/RemoteControl/Source/RemoteControlProtocolWidgets/Private/Widgets/Masks/SRCProtocolMaskButton.h
new file mode 100644
--- /dev/null
+++ b/Engine/Plugins/VirtualProduction/RemoteControl/Source/RemoteControlProtocolWidgets/Private/Widgets/Masks/SRCProtocolMaskButton.h
@@ -0,0 +1,120 @@
+// Copyright Epic Games, Inc. All Rights Reserved.
+
+#pragma once
+
+#include "CoreMinimal.h"
+#include "Widgets/Layout/SBorder.h"
+
+enum class ERCMask : uint8;
+struct FProtocolWidgetStyle;
+class SCheckBox;
+class STextBlock;
+
+/** Delegate that is executed to retrieve the current mask state */
+DECLARE_DELEGATE_RetVal_OneParam(ECheckBoxState, FIsMasked, ERCMask /* MaskBit */);
+
+/** Delegate that is executed when the mask state changes */
+DECLARE_DELEGATE_TwoParams(FOnMaskStateChanged, ECheckBoxState /* NewState */, ERCMask /* MaskBit */);
+
+/**
+ * A custom widget represents a mask widget.
+ */
+class SRCProtocolMaskButton : public SBorder
+{
+public:
+	SLATE_BEGIN_ARGS(SRCProtocolMaskButton)
+		: _MaskBit()
+		, _MaskColor()
+		, _DefaultLabel()
+		, _OptionalTooltip()
+	{}
+
+		SLATE_ARGUMENT(ERCMask, MaskBit)
+		SLATE_ARGUMENT(FLinearColor, MaskColor)
+
+		SLATE_ATTRIBUTE(FText, DefaultLabel)
+		SLATE_ATTRIBUTE(FText, OptionalTooltip)
+
+		/** Whether the mask button is currently in a active state */
+		SLATE_EVENT(FIsMasked, IsMasked)
+
+		/** Called when the masked state has changed */
+		SLATE_EVENT(FOnMaskStateChanged, OnMasked)
+
+	SLATE_END_ARGS()
+
+	/** Constructs this widget with InArgs */
+	void Construct(const FArguments& InArgs);
+
+	/**
+	 * Retrieves mask bit.
+	 */
+	ERCMask GetMaskBit() const
+	{
+		return MaskBit;
+	}
+
+	/**
+	 * Returns true if the mask is currently enabled
+	 *
+	 * @return	True if masked, otherwise false
+	 */
+	bool IsMasked() const;
+
+	/**
+	 * Toggles the masked state for this mask button, fire events as needed
+	 */
+	void ToggleMaskedState(bool bSoftToggle = false);
+
+protected:
+
+	/**
+	 * Retrieves the tooltip of the mask button based on its state.
+	 */
+	FText HandleMaskTooltip() const;
+
+	/**
+	 * Retrieves the state of the mask.
+	 */
+	ECheckBoxState IsMaskEnabled() const;
+
+	/**
+	 * Handles masked state.
+	 */
+	void SetMaskEnabled(ECheckBoxState NewState);
+
+private:
+
+	/** Holds the label of this mask button. */
+	TAttribute<FText> DefaultLabel;
+	
+	/** Holds the tooltip text of this mask button. */
+	TAttribute<FText> OptionalToolTipText;
+
+	/** Holds the mask that should be affected by this button as an uint8. */
+	ERCMask MaskBit;
+
+	/** Holds a shared reference to the underlying mask button. */
+	TSharedPtr<SCheckBox> Mask;
+	
+	/** Holds a shared reference to the underlying mask label. */
+	TSharedPtr<STextBlock> MaskLabel;
+
+	/** Holds the "Mask" string as text. */
+	static FText MaskText;
+
+	/** Holds the "Unmask" string as text. */
+	static FText UnmaskText;
+
+	/** Are we Masked */
+	TAttribute<ECheckBoxState> MaskedState;
+
+	/** Are we masked */
+	FIsMasked IsInMaskedState;
+
+	/** Delegate called when the masking state changes. */
+	FOnMaskStateChanged OnMaskStateChanged;
+
+	/** Protocol Widget Style reference. */
+	const FProtocolWidgetStyle* WidgetStyle;
+};
diff --git a/Engine/Plugins/VirtualProduction/RemoteControl/Source/RemoteControlProtocolWidgets/Private/Widgets/Masks/SRCProtocolMaskTriplet.cpp b/Engine/Plugins/VirtualProduction/RemoteControl/Source/RemoteControlProtocolWidgets/Private/Widgets/Masks/SRCProtocolMaskTriplet.cpp
new file mode 100644
--- /dev/null
+++ b/Engine/Plugins/VirtualProduction/RemoteControl/Source/RemoteControlProtocolWidgets/Private/Widgets/Masks/SRCProtocolMaskTriplet.cpp
@@ -0,0 +1,227 @@
+// Copyright Epic Games, Inc. All Rights Reserved.
+
+#include "Widgets/Masks/SRCProtocolMaskTriplet.h"
+#include "SRCProtocolMaskButton.h"
+
+#include "Styling/ProtocolPanelStyle.h"
+#include "Styling/ProtocolStyles.h"
+
+#include "Widgets/Layout/SSpacer.h"
+#include "Widgets/SBoxPanel.h"
+
+#define LOCTEXT_NAMESPACE "SRCProtocolMaskTriplet"
+
+void SRCProtocolMaskTriplet::Construct(const FArguments& InArgs)
+{
+	WidgetStyle = &FProtocolPanelStyle::Get()->GetWidgetStyle<FProtocolWidgetStyle>("ProtocolsPanel.Widgets.Mask");
+
+	MaskingType = InArgs._MaskingType;
+
+	const bool bCanBeMasked = InArgs._CanBeMasked.Get(false);
+
+	TSharedPtr<SWidget> ContentWidget = SNullWidget::NullWidget;
+
+	if (bCanBeMasked)
+	{
+		TSharedPtr<SHorizontalBox> ContentHBox;
+
+		const FText MaskALabel = GetMaskLabel(EMasking::Bits::ETM_A);
+		const FText MaskBLabel = GetMaskLabel(EMasking::Bits::ETM_B);
+		const FText MaskCLabel = GetMaskLabel(EMasking::Bits::ETM_C);
+
+		ContentWidget = SNew(SBorder) // Foreground Border
+			.BorderImage(&WidgetStyle->ContentAreaBrushDark)
+			[
+				SAssignNew(ContentHBox, SHorizontalBox)
+
+				// Mask A
+				+SHorizontalBox::Slot()
+				.AutoWidth()
+				.Padding(2.f, 0.f)
+				[
+					SAssignNew(MaskA, SRCProtocolMaskButton)
+					.MaskBit(InArgs._MaskA)
+					.MaskColor(FLinearColor::Red)
+					.DefaultLabel(MaskALabel)
+					.IsMasked(this, &SRCProtocolMaskTriplet::IsMaskEnabled)
+					.OnMasked(this, &SRCProtocolMaskTriplet::SetMaskEnabled)
+				]
+		
+				// Spacer
+				+SHorizontalBox::Slot()
+				.FillWidth(1.f)
+				[
+					SNew(SSpacer)
+				]
+
+				// Mask B
+				+SHorizontalBox::Slot()
+				.AutoWidth()
+				.Padding(2.f, 0.f)
+				[
+					SAssignNew(MaskB, SRCProtocolMaskButton)
+					.MaskBit(InArgs._MaskB)
+					.MaskColor(FLinearColor::Green)
+					.DefaultLabel(MaskBLabel)
+					.IsMasked(this, &SRCProtocolMaskTriplet::IsMaskEnabled)
+					.OnMasked(this, &SRCProtocolMaskTriplet::SetMaskEnabled)
+				]
+				
+				// Spacer
+				+SHorizontalBox::Slot()
+				.FillWidth(1.f)
+				[
+					SNew(SSpacer)
+				]
+
+				// Mask C
+				+SHorizontalBox::Slot()
+				.AutoWidth()
+				.Padding(2.f, 0.f)
+				[
+					SAssignNew(MaskC, SRCProtocolMaskButton)
+					.MaskBit(InArgs._MaskC)
+					.MaskColor(FLinearColor::Blue)
+					.DefaultLabel(MaskCLabel)
+					.IsMasked(this, &SRCProtocolMaskTriplet::IsMaskEnabled)
+					.OnMasked(this, &SRCProtocolMaskTriplet::SetMaskEnabled)
+				]
+			];
+
+		if (InArgs._EnableOptionalMask && ContentHBox.IsValid())
+		{
+			const FText MaskDLabel = GetMaskLabel(EMasking::Bits::ETM_D);
+
+			// Spacer
+			ContentHBox->AddSlot()
+				.FillWidth(1.f)
+				[
+					SNew(SSpacer)
+				];
+
+			// Optional Mask
+			ContentHBox->AddSlot()
+				.AutoWidth()
+				.Padding(2.f, 0.f)
+				[
+					SAssignNew(MaskD, SRCProtocolMaskButton)
+					.MaskBit(InArgs._OptionalMask)
+					.MaskColor(FLinearColor::White)
+					.DefaultLabel(MaskDLabel)
+					.IsMasked(this, &SRCProtocolMaskTriplet::IsMaskEnabled)
+					.OnMasked(this, &SRCProtocolMaskTriplet::SetMaskEnabled)
+				];
+		}
+	}
+
+	// Background Border
+	SBorder::Construct(SBorder::FArguments()
+		.BorderImage(bCanBeMasked ? &WidgetStyle->ContentAreaBrush : FAppStyle::Get().GetBrush("NoBorder"))
+		.HAlign(HAlign_Center)
+		.Padding(2.f)
+		.Content()
+		[
+			ContentWidget.ToSharedRef()
+		]
+	);
+
+	LoadMasks();
+}
+
+FText SRCProtocolMaskTriplet::GetMaskLabel(EMasking::Bits InMaskBitName) const
+{
+	FText MaskLabel = FText::GetEmpty();
+
+	switch (MaskingType)
+	{
+		case EMaskingType::Rotator:
+		case EMaskingType::Vector:
+		case EMaskingType::Quat:
+		{
+			switch (InMaskBitName)
+			{
+				case EMasking::Bits::ETM_A:
+					MaskLabel = LOCTEXT("MaskXLabel", "X");
+					break;
+				case EMasking::Bits::ETM_B:
+					MaskLabel = LOCTEXT("MaskYLabel", "Y");
+					break;
+				case EMasking::Bits::ETM_C:
+					MaskLabel = LOCTEXT("MaskZLabel", "Z");
+					break;
+				case EMasking::Bits::ETM_D:
+					MaskLabel = LOCTEXT("MaskWLabel", "W");
+					break;
+			}
+		}
+		break;
+		case EMaskingType::Polar:
+		{
+			switch (InMaskBitName)
+			{
+				case EMasking::Bits::ETM_A:
+					MaskLabel = LOCTEXT("MaskILabel", "U");
+					break;
+				case EMasking::Bits::ETM_B:
+					MaskLabel = LOCTEXT("MaskJLabel", "V");
+					break;
+				case EMasking::Bits::ETM_C:
+					MaskLabel = LOCTEXT("MaskKLabel", "W");
+					break;
+			}
+		}
+		break;
+		case EMaskingType::Color:
+		{
+			switch (InMaskBitName)
+			{
+				case EMasking::Bits::ETM_A:
+					MaskLabel = LOCTEXT("MaskRLabel", "R");
+					break;
+				case EMasking::Bits::ETM_B:
+					MaskLabel = LOCTEXT("MaskGLabel", "G");
+					break;
+				case EMasking::Bits::ETM_C:
+					MaskLabel = LOCTEXT("MaskBLabel", "B");
+					break;
+				case EMasking::Bits::ETM_D:
+					MaskLabel = LOCTEXT("MaskALabel", "A");
+					break;
+			}
+		}
+		break;
+		case EMaskingType::Unsupported:
+			check(false);
+			break;
+	}
+
+	return MaskLabel;
+}
+
+void SRCProtocolMaskTriplet::LoadMasks()
+{
+	// Soft toggle enabled masks.
+	constexpr bool bSoftToggle = true;
+
+	if (MaskA.IsValid() && MaskA->IsMasked())
+	{
+		MaskA->ToggleMaskedState(bSoftToggle);
+	}
+
+	if (MaskB.IsValid() && MaskB->IsMasked())
+	{
+		MaskB->ToggleMaskedState(bSoftToggle);
+	}
+
+	if (MaskC.IsValid() && MaskC->IsMasked())
+	{
+		MaskC->ToggleMaskedState(bSoftToggle);
+	}
+
+	if (MaskD.IsValid() && MaskD->IsMasked())
+	{
+		MaskD->ToggleMaskedState(bSoftToggle);
+	}
+}
+
+#undef LOCTEXT_NAMESPACE
diff --git a/Engine/Plugins/VirtualProduction/RemoteControl/Source/RemoteControlProtocolWidgets/Public/IRemoteControlProtocolWidgetsModule.h b/Engine/Plugins/VirtualProduction/RemoteControl/Source/RemoteControlProtocolWidgets/Public/IRemoteControlProtocolWidgetsModule.h
--- a/Engine/Plugins/VirtualProduction/RemoteControl/Source/RemoteControlProtocolWidgets/Public/IRemoteControlProtocolWidgetsModule.h
+++ b/Engine/Plugins/VirtualProduction/RemoteControl/Source/RemoteControlProtocolWidgets/Public/IRemoteControlProtocolWidgetsModule.h
@@ -11,6 +11,8 @@ class SWidget;
 class URemoteControlPreset;
 class SWidget;
 
+DECLARE_MULTICAST_DELEGATE_OneParam(FOnProtocolBindingAddedOrRemoved, ERCProtocolBinding::Op /* BindingOperation */);
+
 /** A Remote Control module that provides editor widgets for protocol bindings. */
 class IRemoteControlProtocolWidgetsModule : public IModuleInterface
 {
@@ -35,4 +37,7 @@ public:
 
 	/** Get the binding list public reference */
 	virtual TSharedPtr<IRCProtocolBindingList> GetProtocolBindingList() const = 0;
+
+	/** Called when binding is added or removed. */
+	virtual FOnProtocolBindingAddedOrRemoved& OnProtocolBindingAddedOrRemoved() = 0;
 };
diff --git a/Engine/Plugins/VirtualProduction/RemoteControl/Source/RemoteControlProtocolWidgets/Public/Widgets/Masks/SRCProtocolMask.h b/Engine/Plugins/VirtualProduction/RemoteControl/Source/RemoteControlProtocolWidgets/Public/Widgets/Masks/SRCProtocolMask.h
new file mode 100644
--- /dev/null
+++ b/Engine/Plugins/VirtualProduction/RemoteControl/Source/RemoteControlProtocolWidgets/Public/Widgets/Masks/SRCProtocolMask.h
@@ -0,0 +1,42 @@
+// Copyright Epic Games, Inc. All Rights Reserved.
+
+#pragma once
+
+#include "CoreMinimal.h"
+#include "SRCProtocolMaskTriplet.h"
+
+enum class ERCMask : uint8;
+struct FRemoteControlField;
+
+/**
+ * A widget that represents a mask.
+ */
+class REMOTECONTROLPROTOCOLWIDGETS_API SRCProtocolMask : public SRCProtocolMaskTriplet, public IHasMaskExtensibility
+{
+public:
+	SLATE_BEGIN_ARGS(SRCProtocolMask)
+	{}
+
+	SLATE_END_ARGS()
+
+	/** Constructs this widget with InArgs */
+	void Construct(const FArguments& InArgs, TWeakPtr<FRemoteControlField> InField);
+
+protected:
+
+	//~ Begin : IHasMaskExtensibility Interface
+	virtual bool CanBeMasked() const override;
+	virtual EMaskingType GetMaskingType() override;
+	virtual bool HasOptionalMask() const override;
+	//~ End : IHasMaskExtensibility Interface
+
+	//~ Begin : SRCProtocolMaskTriplet Interface
+	virtual ECheckBoxState IsMaskEnabled(ERCMask InMaskBit) const override;
+	virtual void SetMaskEnabled(ECheckBoxState NewState, ERCMask NewMaskBit) override;
+	//~ End : SRCProtocolMaskTriplet Interface
+
+private:
+
+	/** Holds a weak reference to the housing RC field. */
+	TWeakPtr<FRemoteControlField> WeakField;
+};
diff --git a/Engine/Plugins/VirtualProduction/RemoteControl/Source/RemoteControlProtocolWidgets/Public/Widgets/Masks/SRCProtocolMaskTriplet.h b/Engine/Plugins/VirtualProduction/RemoteControl/Source/RemoteControlProtocolWidgets/Public/Widgets/Masks/SRCProtocolMaskTriplet.h
new file mode 100644
--- /dev/null
+++ b/Engine/Plugins/VirtualProduction/RemoteControl/Source/RemoteControlProtocolWidgets/Public/Widgets/Masks/SRCProtocolMaskTriplet.h
@@ -0,0 +1,133 @@
+// Copyright Epic Games, Inc. All Rights Reserved.
+
+#pragma once
+
+#include "CoreMinimal.h"
+#include "Widgets/Layout/SBorder.h"
+
+enum class ERCMask : uint8;
+struct FProtocolWidgetStyle;
+class SRCProtocolMaskButton;
+
+namespace EMasking
+{
+	enum class Bits : uint8
+	{
+		ETM_A,
+
+		ETM_B,
+
+		ETM_C,
+
+		ETM_D
+	};
+}
+
+enum class EMaskingType : uint8
+{
+	Color,
+
+	Polar,
+
+	Rotator,
+
+	Vector,
+
+	Quat,
+
+	Unsupported
+};
+
+/**
+ * A widget that represents a triplet mask.
+ */
+class SRCProtocolMaskTriplet : public SBorder
+{
+public:
+	SLATE_BEGIN_ARGS(SRCProtocolMaskTriplet)
+		: _MaskA()
+		, _MaskB()
+		, _MaskC()
+		, _OptionalMask()
+		, _EnableOptionalMask(false)
+	{}
+		SLATE_ARGUMENT(ERCMask, MaskA)
+		SLATE_ARGUMENT(ERCMask, MaskB)
+		SLATE_ARGUMENT(ERCMask, MaskC)
+		SLATE_ARGUMENT(ERCMask, OptionalMask)
+
+		SLATE_ARGUMENT(EMaskingType, MaskingType)
+
+		SLATE_ATTRIBUTE(bool, CanBeMasked)
+		SLATE_ARGUMENT(bool, EnableOptionalMask)
+
+	SLATE_END_ARGS()
+
+	/** Constructs this widget with InArgs */
+	void Construct(const FArguments& InArgs);
+
+private:
+
+	/**
+	 * Retrieves the label of the mask button based on type.
+	 */
+	FText GetMaskLabel(EMasking::Bits InMaskBitName) const;
+
+protected:
+
+	/**
+	 * Retrieves the state of the mask.
+	 */
+	virtual ECheckBoxState IsMaskEnabled(ERCMask InMaskBit) const { return ECheckBoxState::Unchecked; }
+
+	/** Soft toggles masks based on active masks. */
+	void LoadMasks();
+
+	/**
+	 * Handles masked state.
+	 */
+	virtual void SetMaskEnabled(ECheckBoxState NewState, ERCMask NewMaskBit) {}
+
+protected:
+
+	/** Protocol Widget Style reference. */
+	const FProtocolWidgetStyle* WidgetStyle;
+
+private:
+
+	/** Holds the type of this triplet mask. */
+	EMaskingType MaskingType;
+
+	/** Holds a shared reference to the underlying mask button for A. */
+	TSharedPtr<SRCProtocolMaskButton> MaskA;
+
+	/** Holds a shared reference to the underlying mask button for B. */
+	TSharedPtr<SRCProtocolMaskButton> MaskB;
+
+	/** Holds a shared reference to the underlying mask button for C. */
+	TSharedPtr<SRCProtocolMaskButton> MaskC;
+
+	/** Holds a shared reference to the underlying mask button for D. */
+	TSharedPtr<SRCProtocolMaskButton> MaskD;
+};
+
+/**
+ * Ensures the maskability of the derived classes.
+ */
+class IHasMaskExtensibility
+{
+	/**
+	 * True when the derived classes supports masking, false otherwise.
+	 */
+	virtual bool CanBeMasked() const = 0;
+
+	/**
+	 * Retrieves the masking type based on the property type.
+	 */
+	virtual EMaskingType GetMaskingType() = 0;
+
+	/**
+	 * True when derived classes has optional mask.
+	 */
+	virtual bool HasOptionalMask() const = 0;
+};
diff --git a/Engine/Plugins/VirtualProduction/RemoteControl/Source/RemoteControlUI/Private/Styling/RemoteControlStyles.cpp b/Engine/Plugins/VirtualProduction/RemoteControl/Source/RemoteControlUI/Private/Styling/RemoteControlStyles.cpp
--- a/Engine/Plugins/VirtualProduction/RemoteControl/Source/RemoteControlUI/Private/Styling/RemoteControlStyles.cpp
+++ b/Engine/Plugins/VirtualProduction/RemoteControl/Source/RemoteControlUI/Private/Styling/RemoteControlStyles.cpp
@@ -26,8 +26,11 @@ void FRCPanelStyle::GetResources(TArray<const FSlateBrush*>& OutBrushes) const
 	FlatButtonStyle.GetResources(OutBrushes);
 	TableRowStyle.GetResources(OutBrushes);
 	TableViewStyle.GetResources(OutBrushes);
+	SwitchButtonStyle.GetResources(OutBrushes);
 	ToggleButtonStyle.GetResources(OutBrushes);
 	OutBrushes.Add(&ContentAreaBrush);
+	OutBrushes.Add(&ContentAreaBrushDark);
+	OutBrushes.Add(&ContentAreaBrushLight);
 	OutBrushes.Add(&SectionHeaderBrush);
 }
 
diff --git a/Engine/Plugins/VirtualProduction/RemoteControl/Source/RemoteControlUI/Private/Styling/RemoteControlStyles.h b/Engine/Plugins/VirtualProduction/RemoteControl/Source/RemoteControlUI/Private/Styling/RemoteControlStyles.h
--- a/Engine/Plugins/VirtualProduction/RemoteControl/Source/RemoteControlUI/Private/Styling/RemoteControlStyles.h
+++ b/Engine/Plugins/VirtualProduction/RemoteControl/Source/RemoteControlUI/Private/Styling/RemoteControlStyles.h
@@ -37,8 +37,18 @@ struct FRCPanelStyle : public FSlateWidgetStyle
 	UPROPERTY(EditAnywhere, Category = Appearance)
 		FSlateBrush ContentAreaBrush;
 	FRCPanelStyle& SetContentAreaBrush(const FSlateBrush& InContentAreaBrush) { ContentAreaBrush = InContentAreaBrush; return *this; }
+	
+	/** Brush used by the RC Panels to draw the content associated with this panel (Dark). */
+	UPROPERTY(EditAnywhere, Category = Appearance)
+		FSlateBrush ContentAreaBrushDark;
+	FRCPanelStyle& SetContentAreaBrushDark(const FSlateBrush& InContentAreaBrushDark) { ContentAreaBrushDark = InContentAreaBrushDark; return *this; }
+	
+	/** Brush used by the RC Panels to draw the content associated with this panel (Light). */
+	UPROPERTY(EditAnywhere, Category = Appearance)
+		FSlateBrush ContentAreaBrushLight;
+	FRCPanelStyle& SetContentAreaBrushLight(const FSlateBrush& InContentAreaBrushLight) { ContentAreaBrushLight = InContentAreaBrushLight; return *this; }
 
-	/** Style used for the add new button */
+	/** Style used for the flat button */
 	UPROPERTY(EditAnywhere, Category = Appearance)
 		FButtonStyle FlatButtonStyle;
 	FRCPanelStyle& SetFlatButtonStyle(const FButtonStyle& InFlatButtonStyle) { FlatButtonStyle = InFlatButtonStyle; return *this; }
@@ -93,7 +103,12 @@ struct FRCPanelStyle : public FSlateWidgetStyle
 		FTableViewStyle TableViewStyle;
 	FRCPanelStyle& SetTableViewStyle(const FTableViewStyle& InTableViewStyle) { TableViewStyle = InTableViewStyle; return *this; }
 
-	/** Style used for the add new button */
+	/** Style used for the switch button */
+	UPROPERTY(EditAnywhere, Category = Appearance)
+		FCheckBoxStyle SwitchButtonStyle;
+	FRCPanelStyle& SetSwitchButtonStyle(const FCheckBoxStyle& InSwitchButtonStyle) { SwitchButtonStyle = InSwitchButtonStyle; return *this; }
+
+	/** Style used for the toggle button */
 	UPROPERTY(EditAnywhere, Category = Appearance)
 		FCheckBoxStyle ToggleButtonStyle;
 	FRCPanelStyle& SetToggleButtonStyle(const FCheckBoxStyle& InToggleButtonStyle) { ToggleButtonStyle = InToggleButtonStyle; return *this; }
diff --git a/Engine/Plugins/VirtualProduction/RemoteControl/Source/RemoteControlUI/Private/UI/Panels/SRCDockPanel.h b/Engine/Plugins/VirtualProduction/RemoteControl/Source/RemoteControlUI/Private/UI/Panels/SRCDockPanel.h
--- a/Engine/Plugins/VirtualProduction/RemoteControl/Source/RemoteControlUI/Private/UI/Panels/SRCDockPanel.h
+++ b/Engine/Plugins/VirtualProduction/RemoteControl/Source/RemoteControlUI/Private/UI/Panels/SRCDockPanel.h
@@ -75,7 +75,7 @@ public:
 
 		SLATE_ARGUMENT(bool, EnableHeader)
 
-		SLATE_ARGUMENT(FText, HeaderLabel)
+		SLATE_ATTRIBUTE(FText, HeaderLabel)
 
 		SLATE_ARGUMENT(EOrientation, Orientation)
 
@@ -140,7 +140,7 @@ public:
 
 		SLATE_ARGUMENT(bool, EnableHeader)
 
-		SLATE_ARGUMENT(FText, HeaderLabel)
+		SLATE_ATTRIBUTE(FText, HeaderLabel)
 
 		SLATE_ARGUMENT(EOrientation, Orientation)
 
diff --git a/Engine/Plugins/VirtualProduction/RemoteControl/Source/RemoteControlUI/Private/UI/RemoteControlPanelStyle.cpp b/Engine/Plugins/VirtualProduction/RemoteControl/Source/RemoteControlUI/Private/UI/RemoteControlPanelStyle.cpp
--- a/Engine/Plugins/VirtualProduction/RemoteControl/Source/RemoteControlUI/Private/UI/RemoteControlPanelStyle.cpp
+++ b/Engine/Plugins/VirtualProduction/RemoteControl/Source/RemoteControlUI/Private/UI/RemoteControlPanelStyle.cpp
@@ -181,6 +181,14 @@ void FRemoteControlPanelStyle::SetupPanelStyles(TSharedRef<FSlateStyleSet> InSty
 
 	StyleSet->Set("RemoteControlPanel.GroupRow", TableRowStyle);
 
+	// Mode Switcher Styles
+	FSlateBrush ContentAreaBrushDark = BOX_BRUSH("Common/DarkGroupBorder", FMargin(4.f / 16.f), FLinearColor(0.5f, 0.5f, 0.5f, 1.f));
+	FSlateBrush ContentAreaBrushLight = BOX_BRUSH("Common/LightGroupBorder", FMargin(4.f / 16.f));
+
+	FCheckBoxStyle SwitchButtonStyle = AppStyle.GetWidgetStyle<FCheckBoxStyle>("ToggleButtonCheckbox");
+	SwitchButtonStyle.SetCheckBoxType(ESlateCheckBoxType::ToggleButton);
+	SwitchButtonStyle.SetPadding(FMargin(4.f, 2.f));
+
 	// Text Styles
 	const FStyleFonts& StyleFonts = FStyleFonts::Get();
 
@@ -191,9 +199,12 @@ void FRemoteControlPanelStyle::SetupPanelStyles(TSharedRef<FSlateStyleSet> InSty
 
 	FRCPanelStyle RCMinorPanelStyle = FRCPanelStyle()
 		.SetFlatButtonStyle(FlatButtonStyle)
+		.SetSwitchButtonStyle(SwitchButtonStyle)
 		.SetToggleButtonStyle(ToggleButtonStyle)
 
 		.SetContentAreaBrush(FSlateColorBrush(FStyleColors::Panel))
+		.SetContentAreaBrushDark(ContentAreaBrushDark)
+		.SetContentAreaBrushLight(ContentAreaBrushLight)
 		.SetSectionHeaderBrush(FSlateColorBrush(FStyleColors::Foldout))
 
 		.SetHeaderRowStyle(HeaderRowStyle)
diff --git a/Engine/Plugins/VirtualProduction/RemoteControl/Source/RemoteControlUI/Private/UI/SRCModeSwitcher.cpp b/Engine/Plugins/VirtualProduction/RemoteControl/Source/RemoteControlUI/Private/UI/SRCModeSwitcher.cpp
new file mode 100644
--- /dev/null
+++ b/Engine/Plugins/VirtualProduction/RemoteControl/Source/RemoteControlUI/Private/UI/SRCModeSwitcher.cpp
@@ -0,0 +1,267 @@
+// Copyright Epic Games, Inc. All Rights Reserved.
+
+#include "UI/SRCModeSwitcher.h"
+
+#include "RemoteControlPanelStyle.h"
+#include "Styling/RemoteControlStyles.h"
+
+#include "Widgets/Images/SLayeredImage.h"
+#include "Widgets/Input/SCheckBox.h"
+#include "Widgets/Layout/SBox.h"
+#include "Widgets/SBoxPanel.h"
+#include "Widgets/Text/STextBlock.h"
+
+#define LOCTEXT_NAMESPACE "SRCModeSwitcher"
+
+void SRCModeSwitcher::Construct(const FArguments& InArgs)
+{
+	RCPanelStyle = &FRemoteControlPanelStyle::Get()->GetWidgetStyle<FRCPanelStyle>("RemoteControlPanel.MinorPanel");
+
+	ActiveMode = InArgs._DefaultMode;
+
+	ModeSwitched = InArgs._OnModeSwitched;
+
+	if (InArgs._OnModesChanged.IsBound())
+	{
+		ModesChanged.Add(InArgs._OnModesChanged);
+	}
+
+	SBorder::Construct(SBorder::FArguments()
+		.Padding(RCPanelStyle->PanelPadding)
+		.BorderImage(&RCPanelStyle->ContentAreaBrushDark)
+	);
+
+	// Copy all the mode info from the declaration
+	for (FRCMode* const Mode : InArgs.Slots)
+	{
+		RCModes.Add(*Mode);
+
+		Mode->bIsVisible = true;
+	}
+
+	// Generate widgets for all modes
+	RegenerateWidgets();
+}
+
+ECheckBoxState SRCModeSwitcher::IsModeEnabled(FRCMode* InMode) const
+{
+	check(InMode);
+
+	return ActiveMode == InMode->ModeId ? ECheckBoxState::Checked : ECheckBoxState::Unchecked;
+}
+
+void SRCModeSwitcher::RegenerateWidgets()
+{
+	TSharedRef<SHorizontalBox> ModeSwitcher = SNew(SHorizontalBox)
+		.Clipping(EWidgetClipping::ClipToBounds);
+
+	// Construct widgets for all modes
+	{
+		for (int32 SlotIndex = 0; SlotIndex < RCModes.Num(); ++SlotIndex)
+		{
+			FRCMode& SomeMode = RCModes[SlotIndex];
+
+			if (SomeMode.ShouldGenerateWidget.Get(true) && SomeMode.bIsVisible)
+			{
+				const FText ModeLabel = FText::Format(LOCTEXT("ModeLabel", "{0}"), { SomeMode.DefaultText.IsSet() ? SomeMode.DefaultText.Get() : FText::FromString(SomeMode.ModeId.ToString() + TEXT("[LabelMissing]")) });
+
+				TSharedRef<SWidget> ContentWidget = SNullWidget::NullWidget;
+
+				if (SomeMode.bShowIcon)
+				{
+					// If we were supplied an image than go ahead and use that.
+					TSharedRef<SLayeredImage> IconWidget = SNew(SLayeredImage)
+						.ColorAndOpacity(FSlateColor::UseForeground())
+						.Visibility(EVisibility::HitTestInvisible)
+						.Image(SomeMode.OptionalIcon.Get().GetIcon());
+
+					ContentWidget = SNew(SHorizontalBox)
+
+					// Icon Widget
+					+ SHorizontalBox::Slot()
+					.AutoWidth()
+					.VAlign(VAlign_Center)
+					.HAlign(HAlign_Center)
+					.Padding(2.f, 4.f)
+					[
+						IconWidget
+					]
+					
+					// Label Text
+					+ SHorizontalBox::Slot()
+					.AutoWidth()
+					.Padding(RCPanelStyle->HeaderRowPadding)
+					.VAlign(VAlign_Center)
+					[
+						SNew(STextBlock)
+						.Text(ModeLabel)
+						.TextStyle(&RCPanelStyle->PanelTextStyle)
+					];
+				}
+				else
+				{
+					ContentWidget = SNew(STextBlock)
+						.Text(ModeLabel)
+						.TextStyle(&RCPanelStyle->PanelTextStyle);
+				}
+
+				TSharedRef<SWidget> NewModeSwitcher = SNew(SCheckBox)
+					.Style(&RCPanelStyle->SwitchButtonStyle)
+					.IsChecked(this, &SRCModeSwitcher::IsModeEnabled, &SomeMode)
+					.OnCheckStateChanged(this, &SRCModeSwitcher::SetModeEnabled, &SomeMode)
+					.HAlign(HAlign_Center)
+					.ToolTip(SomeMode.ToolTip)
+					.ToolTipText(SomeMode.ToolTip.IsSet() ? TAttribute<FText>() :
+						SomeMode.DefaultTooltip.IsSet() ? SomeMode.DefaultTooltip : ModeLabel)
+					.Padding(4.f)
+					[
+						ContentWidget
+					];
+
+				switch (SomeMode.SizeRule)
+				{
+					case EModeWidgetSizing::Fill:
+					default:
+					{
+						// Add resizable cell
+						ModeSwitcher->AddSlot()
+							.FillWidth(SomeMode.GetWidth())
+							.HAlign(SomeMode.CellHAlignment)
+							.VAlign(SomeMode.CellVAlignment)
+							[
+								NewModeSwitcher
+							];
+					}
+					break;
+					case EModeWidgetSizing::Fixed:
+					{
+						// Add fixed size cell
+						ModeSwitcher->AddSlot()
+							.AutoWidth()
+							.HAlign(SomeMode.CellHAlignment)
+							.VAlign(SomeMode.CellVAlignment)
+							[
+								SNew(SBox)
+								.WidthOverride(SomeMode.GetWidth())
+								[
+									NewModeSwitcher
+								]
+							];
+					}
+					break;
+					case EModeWidgetSizing::Auto:
+					{
+						// Add size to content cell
+						ModeSwitcher->AddSlot()
+							.AutoWidth()
+							.HAlign(SomeMode.CellHAlignment)
+							.VAlign(SomeMode.CellVAlignment)
+							[
+								NewModeSwitcher
+							];
+					}
+					break;
+				}
+			}
+		}
+	}
+
+	SetContent(ModeSwitcher);
+}
+
+void SRCModeSwitcher::SetModeEnabled(ECheckBoxState NewState, FRCMode* NewMode)
+{
+	check(NewMode);
+
+	if (ActiveMode != NewMode->ModeId)
+	{
+		if (NewState == ECheckBoxState::Checked)
+		{
+			ActiveMode = NewMode->ModeId;
+
+			OnModeSwitched().ExecuteIfBound(*NewMode);
+		}
+	}
+}
+
+const TArray<SRCModeSwitcher::FRCMode>& SRCModeSwitcher::GetModes() const
+{
+	return RCModes;
+}
+
+void SRCModeSwitcher::AddMode(const FRCMode::FArguments& NewModeArgs)
+{
+	SRCModeSwitcher::FRCMode NewMode = SRCModeSwitcher::FRCMode(NewModeArgs);
+
+	AddMode(NewMode);
+}
+
+void SRCModeSwitcher::AddMode(FRCMode& NewMode)
+{
+	int32 InsertIdx = RCModes.Num();
+
+	InsertMode(NewMode, InsertIdx);
+}
+
+void SRCModeSwitcher::InsertMode(const FRCMode::FArguments& NewModeArgs, int32 InsertIdx)
+{
+	SRCModeSwitcher::FRCMode NewMode = SRCModeSwitcher::FRCMode(NewModeArgs);
+
+	InsertMode(NewMode, InsertIdx);
+}
+
+void SRCModeSwitcher::InsertMode(FRCMode& NewMode, int32 InsertIdx)
+{
+	check(NewMode.ModeId != NAME_None);
+
+	if (NewMode.bIsDefault)
+	{
+		ActiveMode = NewMode.ModeId;
+	}
+
+	RCModes.Insert(NewMode, InsertIdx);
+
+	ModesChanged.Broadcast(SharedThis(this));
+
+	RegenerateWidgets();
+}
+
+void SRCModeSwitcher::RemoveMode(const FName& InModeId)
+{
+	check(InModeId != NAME_None);
+
+	for (int32 SlotIndex = RCModes.Num() - 1; SlotIndex >= 0; --SlotIndex)
+	{
+		FRCMode& RCMode = RCModes[SlotIndex];
+
+		if (RCMode.ModeId == InModeId)
+		{
+			RCModes.RemoveAt(SlotIndex);
+		}
+	}
+
+	ModesChanged.Broadcast(SharedThis(this));
+
+	RegenerateWidgets();
+}
+
+void SRCModeSwitcher::RefreshModes()
+{
+	RegenerateWidgets();
+}
+
+void SRCModeSwitcher::ClearModes()
+{
+	const bool bHadModeSwitches = RCModes.Num() > 0;
+	
+	RCModes.Empty();
+	
+	if (bHadModeSwitches)
+	{
+		ModesChanged.Broadcast(SharedThis(this));
+	}
+
+	RegenerateWidgets();
+}
+
+#undef LOCTEXT_NAMESPACE
diff --git a/Engine/Plugins/VirtualProduction/RemoteControl/Source/RemoteControlUI/Private/UI/SRCModeSwitcher.h b/Engine/Plugins/VirtualProduction/RemoteControl/Source/RemoteControlUI/Private/UI/SRCModeSwitcher.h
new file mode 100644
--- /dev/null
+++ b/Engine/Plugins/VirtualProduction/RemoteControl/Source/RemoteControlUI/Private/UI/SRCModeSwitcher.h
@@ -0,0 +1,247 @@
+// Copyright Epic Games, Inc. All Rights Reserved.
+
+#pragma once
+
+#include "CoreMinimal.h"
+#include "Textures/SlateIcon.h"
+#include "Widgets/Layout/SBorder.h"
+
+struct FRCPanelStyle;
+class SWidget;
+
+namespace EModeWidgetSizing
+{
+	enum Rule
+	{
+		/** Mode widget stretches to a fraction of the mode switcher. */
+		Fill = 0,
+
+		/**	Mode widget is fixed width and cannot be resized. */
+		Fixed = 1,
+
+		/** Mode widget will be sized to fit the content. */
+		Auto = 2,
+	};
+};
+
+/**
+ * A custom mode switcher for Remote Control Panel.
+ */
+class SRCModeSwitcher : public SBorder
+{
+public:
+	/** Describes a single mode switch. */
+	class FRCMode
+	{
+	public:
+
+		SLATE_BEGIN_ARGS(FRCMode)
+			: _ModeId()
+			, _DefaultLabel()
+			, _DefaultTooltip()
+			, _ToolTip()
+			, _FillWidth(1.0f)
+			, _FixedWidth()
+			, _HAlignCell(HAlign_Fill)
+			, _VAlignCell(VAlign_Fill)
+			, _IsDefault(false)
+		{}
+		
+			SLATE_ARGUMENT(FName, ModeId)
+
+			SLATE_ATTRIBUTE(FText, DefaultLabel)
+			SLATE_ATTRIBUTE(FText, DefaultTooltip)
+			SLATE_ATTRIBUTE(TSharedPtr< IToolTip >, ToolTip)
+
+			/** Set the Mode Size Mode to Fill. It's a fraction between 0 and 1 */
+			SLATE_ATTRIBUTE(float, FillWidth)
+			/** Set the Mode Size Mode to Fixed. */
+			SLATE_ARGUMENT(TOptional<float>, FixedWidth)
+
+			SLATE_ARGUMENT(EHorizontalAlignment, HAlignCell)
+			SLATE_ARGUMENT(EVerticalAlignment, VAlignCell)
+
+			SLATE_ATTRIBUTE(FSlateIcon, OptionalIcon)
+
+			SLATE_ARGUMENT(bool, IsDefault)
+			SLATE_ATTRIBUTE(bool, ShouldGenerateWidget)
+
+		SLATE_END_ARGS()
+
+		FRCMode(const FArguments& InArgs)
+			: bIsDefault(InArgs._IsDefault)
+			, bIsVisible(true)
+			, bShowIcon(false)
+			, CellHAlignment(InArgs._HAlignCell)
+			, CellVAlignment(InArgs._VAlignCell)
+			, DefaultText(InArgs._DefaultLabel)
+			, DefaultTooltip(InArgs._DefaultTooltip)
+			, ModeId(InArgs._ModeId)
+			, OptionalIcon(InArgs._OptionalIcon)
+			, ToolTip(InArgs._ToolTip)
+			, ShouldGenerateWidget(InArgs._ShouldGenerateWidget)
+			, SizeRule(EModeWidgetSizing::Fill)
+			, Width(1.f)
+		{
+			if ( InArgs._FixedWidth.IsSet() )
+			{
+				Width = InArgs._FixedWidth.GetValue();
+				SizeRule = EModeWidgetSizing::Fixed;
+			}
+			else if ( InArgs._FillWidth.IsSet() )
+			{
+				Width = InArgs._FillWidth;
+				SizeRule = EModeWidgetSizing::Fill;
+			}
+			else
+			{
+				SizeRule = EModeWidgetSizing::Auto;
+			}
+
+			if (OptionalIcon.IsSet())
+			{
+				bShowIcon = true;
+			}
+		}
+
+		bool operator==(const FRCMode& OtherMode)
+		{
+			return ModeId == OtherMode.ModeId;
+		}
+		
+		bool operator==(const FRCMode& OtherMode) const
+		{
+			return ModeId == OtherMode.ModeId;
+		}
+
+	public:
+
+		void SetWidth(float NewWidth)
+		{
+			Width = NewWidth;
+		}
+
+		float GetWidth() const
+		{
+			return Width.Get();
+		}
+
+		/** Whether this mode is default or not. */
+		bool bIsDefault;
+		
+		/** Determines the visibility of the mode widget for this mode. */
+		bool bIsVisible;
+
+		/** Determines whether we need to generate icon widget for this mode or not. */
+		bool bShowIcon;
+
+		/** Holds the mode widget alignment by horizontal. */
+		EHorizontalAlignment CellHAlignment;
+
+		/** Holds the mode widget alignment by vertical. */
+		EVerticalAlignment CellVAlignment;
+
+		/** Default text to use if no widget is passed in. */
+		TAttribute<FText> DefaultText;
+
+		/** Default tooltip to use if no widget is passed in */
+		TAttribute<FText> DefaultTooltip;
+
+		/** A unique ID for this mode, so that it can be saved and restored. */
+		FName ModeId;
+
+		/** Optional icon to be used. */
+		TAttribute<FSlateIcon> OptionalIcon;
+
+		/** Custom tooltip to use */
+		TAttribute<TSharedPtr<IToolTip>> ToolTip;
+
+		/** Determines whether we need to generate mode widget for this mode or not. */
+		TAttribute<bool> ShouldGenerateWidget;
+
+		/** Determines the sizing of the mode widget for this mode. */
+		EModeWidgetSizing::Rule SizeRule;
+
+		/** Width of the mode widget in Slate Units */
+		TAttribute<float> Width;
+	};
+
+	/** Create a mode with the specified ModeId */
+	static FRCMode::FArguments Mode(const FName& InModeId)
+	{
+		FRCMode::FArguments NewArgs;
+		NewArgs.ModeId(InModeId);
+		return NewArgs;
+	}
+
+	/** Delegate to be invoked when the active mode is switched/changed. */
+	DECLARE_DELEGATE_OneParam(FOnModeSwitched, const FRCMode& /* NewMode */);
+	FOnModeSwitched& OnModeSwitched() { return ModeSwitched; }
+
+	/** Delegate to be invoked when the number of modes changed. */
+	DECLARE_EVENT_OneParam(SRCModeSwitcher, FModesChanged, const TSharedRef<SRCModeSwitcher>& /* ModeSwitcher */);
+	FModesChanged* OnModesChanged() { return &ModesChanged; }
+	
+	SLATE_BEGIN_ARGS(SRCModeSwitcher)
+	{}
+
+		SLATE_SUPPORTS_SLOT_WITH_ARGS(FRCMode)
+
+		SLATE_ARGUMENT(FName, DefaultMode)
+
+		SLATE_EVENT(FOnModeSwitched, OnModeSwitched)
+		SLATE_EVENT(FModesChanged::FDelegate, OnModesChanged)
+
+	SLATE_END_ARGS()
+
+	/** Constructs this widget with InArgs */
+	void Construct(const FArguments& InArgs);
+		
+	/** @return the Modes driven by the mode switcher. */
+	const TArray<FRCMode>& GetModes() const;
+
+	/** Adds a mode to the mode switcher. */
+	void AddMode(const FRCMode::FArguments& NewModeArgs);
+	void AddMode(FRCMode& NewMode);
+
+	/** Inserts a mode at the specified index in the mode switcher. */
+	void InsertMode(const FRCMode::FArguments& NewModeArgs, int32 InsertIdx);
+	void InsertMode(FRCMode& NewMode, int32 InsertIdx);
+
+	/** Removes a mode from the mode switcher. */
+	void RemoveMode(const FName& InModeId);
+
+	/** Force refreshing of the mode widgets*/
+	void RefreshModes();
+
+	/** Removes all modes from the mode switcher. */
+	void ClearModes();
+
+private:
+
+	/** Retrieves whether the given mode is enabled or not. */
+	ECheckBoxState IsModeEnabled(FRCMode* InMode) const;
+	
+	/** Regenerates all widgets in the mode switcher. */
+	void RegenerateWidgets();
+
+	/** Sets the given mode as active. */
+	void SetModeEnabled(ECheckBoxState NewState, FRCMode* NewMode);
+
+private:
+
+	/** Holdes the active mode. */
+	FName ActiveMode;
+	
+	/** Information about the various RC modes. */
+	TArray<FRCMode> RCModes;
+
+	/** Panel Style reference. */
+	const FRCPanelStyle* RCPanelStyle;
+
+	/** Holds the delegate to be invoked when the number of modes changed. */
+	FModesChanged ModesChanged;
+	
+	/** Holds the delegate to be invoked when the active mode is switched/changed. */
+	FOnModeSwitched ModeSwitched;
+};
diff --git a/Engine/Plugins/VirtualProduction/RemoteControl/Source/RemoteControlUI/Private/UI/SRCPanelExposedEntitiesList.cpp b/Engine/Plugins/VirtualProduction/RemoteControl/Source/RemoteControlUI/Private/UI/SRCPanelExposedEntitiesList.cpp
--- a/Engine/Plugins/VirtualProduction/RemoteControl/Source/RemoteControlUI/Private/UI/SRCPanelExposedEntitiesList.cpp
+++ b/Engine/Plugins/VirtualProduction/RemoteControl/Source/RemoteControlUI/Private/UI/SRCPanelExposedEntitiesList.cpp
@@ -3,12 +3,15 @@
 #include "SRCPanelExposedEntitiesList.h"
 
 #include "Algo/ForEach.h"
+#include "Commands/RemoteControlCommands.h"
 #include "Editor.h"
 #include "EditorFontGlyphs.h"
 #include "ISettingsModule.h"
 #include "RCPanelWidgetRegistry.h"
 #include "GameFramework/Actor.h"
 #include "Input/DragAndDrop.h"
+#include "IRemoteControlProtocolModule.h"
+#include "IRemoteControlProtocolWidgetsModule.h"
 #include "Engine/Selection.h"
 #include "Editor/EditorEngine.h"
 #include "Misc/Attribute.h"
@@ -18,6 +21,7 @@
 #include "RemoteControlPreset.h"
 #include "RemoteControlUIModule.h"
 #include "ScopedTransaction.h"
+#include "SRCModeSwitcher.h"
 #include "SRCPanelFieldGroup.h"
 #include "SRCPanelExposedField.h"
 #include "Styling/RemoteControlStyles.h"
@@ -32,6 +36,15 @@
 
 #define LOCTEXT_NAMESPACE "RemoteControlPanelEntitiesList"
 
+TSet<FName> SRCPanelExposedEntitiesList::DefaultProtocolColumns = {
+//	FRemoteControlPresetColumns::LinkIdentifier,
+	FRemoteControlPresetColumns::Mask,
+	FRemoteControlPresetColumns::Status,
+};
+
+/**
+ * A custom row to describe each entity in the list.
+ */
 class SEntityRow : public SMultiColumnTableRow<TSharedPtr<SRCPanelTreeNode>>
 {
 public:
@@ -39,6 +52,7 @@ public:
 	SLATE_BEGIN_ARGS(SEntityRow)
 	{}
 	
+		SLATE_ATTRIBUTE(FName, ActiveProtocol)
 		SLATE_ARGUMENT(TSharedPtr<SRCPanelTreeNode>, Entity)
 
 		SLATE_ATTRIBUTE(FMargin, Padding)
@@ -54,6 +68,7 @@ public:
 
 	void Construct(const FArguments& InArgs, const TSharedRef<STableViewBase>& OwnerTableView)
 	{
+		ActiveProtocol = InArgs._ActiveProtocol;
 		Entity = InArgs._Entity;
 
 		FSuperRowType::FArguments SuperArgs = FSuperRowType::FArguments();
@@ -73,6 +88,8 @@ public:
 
 	TSharedRef<SWidget> GenerateWidgetForColumn(const FName& InColumnName) override
 	{
+		const FName& ActiveProtocolName = ActiveProtocol.Get(NAME_None);
+
 		if (Entity.IsValid())
 		{
 			if (Entity->HasChildren() && InColumnName == FRemoteControlPresetColumns::Description)
@@ -98,7 +115,7 @@ public:
 					.Expose(InnerContentSlotNativePtr)
 					.Padding(8.f, 2.f)
 					[
-						Entity->GetWidget(InColumnName)
+						Entity->GetWidget(InColumnName, ActiveProtocolName)
 					];
 
 				InnerContentSlot = InnerContentSlotNativePtr;
@@ -108,11 +125,16 @@ public:
 			}
 			else
 			{
+				const bool bAlignCenter = InColumnName == FRemoteControlPresetColumns::Status || InColumnName == FRemoteControlDefaultProtocolColumns::BindingStatus;
+				
+				const bool bAlignLeft = InColumnName == FRemoteControlPresetColumns::Mask;
+
 				return SNew(SBox)
 					.Padding(FMargin(4.f, 2.f))
+					.HAlign(bAlignCenter ? HAlign_Center : bAlignLeft ? HAlign_Left : HAlign_Fill)
 					.VAlign(VAlign_Center)
 					[
-						Entity->GetWidget(InColumnName)
+						Entity->GetWidget(InColumnName, ActiveProtocolName)
 					];
 			}
 		}
@@ -135,6 +157,9 @@ private:
 
 private:
 
+	/** Holds the active protocol. */
+	TAttribute<FName> ActiveProtocol;
+
 	/** Cached reference of Entity */
 	TSharedPtr<SRCPanelTreeNode> Entity;
 };
@@ -142,6 +167,7 @@ private:
 void SRCPanelExposedEntitiesList::Construct(const FArguments& InArgs, URemoteControlPreset* InPreset, TWeakPtr<FRCPanelWidgetRegistry> InWidgetRegistry)
 {
 	bIsInEditMode = InArgs._EditMode;
+	bIsInProtocolsMode = InArgs._ProtocolsMode;
 	Preset = TStrongObjectPtr<URemoteControlPreset>(InPreset);
 	OnEntityListUpdatedDelegate = InArgs._OnEntityListUpdated;
 	WidgetRegistry = MoveTemp(InWidgetRegistry);
@@ -155,6 +181,7 @@ void SRCPanelExposedEntitiesList::Construct(const FArguments& InArgs, URemoteCon
 	SearchedText = MakeShared<FText>();
 
 	RCPanelStyle = &FRemoteControlPanelStyle::Get()->GetWidgetStyle<FRCPanelStyle>("RemoteControlPanel.MinorPanel");
+	ActiveListMode = EEntitiesListMode::Default;
 
 	// Major Panel
 	TSharedPtr<SRCMajorPanel> ExposePanel = SNew(SRCMajorPanel)
@@ -240,7 +267,7 @@ void SRCPanelExposedEntitiesList::Construct(const FArguments& InArgs, URemoteCon
 		.TreeViewStyle(&RCPanelStyle->TableViewStyle)
 		.OnGetChildren(this, &SRCPanelExposedEntitiesList::OnGetNodeChildren)
 		.HeaderRow(
-			SNew(SHeaderRow)
+			SAssignNew(FieldsHeaderRow, SHeaderRow)
 			.Style(&RCPanelStyle->HeaderRowStyle)
 
 			+ SHeaderRow::Column(FRemoteControlPresetColumns::DragDropHandle)
@@ -250,11 +277,13 @@ void SRCPanelExposedEntitiesList::Construct(const FArguments& InArgs, URemoteCon
 
 			+ SHeaderRow::Column(FRemoteControlPresetColumns::Description)
 			.DefaultLabel(LOCTEXT("RCPresetDescColumnHeader", "Description"))
+			.HAlignHeader(HAlign_Center)
 			.FillWidth(0.5f)
 			.HeaderContentPadding(RCPanelStyle->HeaderRowPadding)
 
 			+ SHeaderRow::Column(FRemoteControlPresetColumns::Value)
 			.DefaultLabel(LOCTEXT("RCPresetValueColumnHeader", "Value"))
+			.HAlignHeader(HAlign_Center)
 			.FillWidth(0.5f)
 			.HeaderContentPadding(RCPanelStyle->HeaderRowPadding)
 
@@ -266,7 +295,7 @@ void SRCPanelExposedEntitiesList::Construct(const FArguments& InArgs, URemoteCon
 
 	// Exposed Entities Dock Panel
 	TSharedPtr<SRCMinorPanel> ExposeDockPanel = SNew(SRCMinorPanel)
-		.HeaderLabel(LOCTEXT("PropertiesLabel", "Properties"))
+		.HeaderLabel(this, &SRCPanelExposedEntitiesList::HandleEntityListHeaderLabel)
 		.EnableFooter(true)
 		[
 			FieldsListView.ToSharedRef()
@@ -303,8 +332,63 @@ void SRCPanelExposedEntitiesList::Construct(const FArguments& InArgs, URemoteCon
 			]
 		];
 
+	// Mode Switcher
+	TSharedPtr<SRCModeSwitcher> ModeSwitcher = SNew(SRCModeSwitcher)
+		.OnModeSwitched_Lambda([this](const SRCModeSwitcher::FRCMode& NewMode)
+			{
+				if (ActiveProtocol != NewMode.ModeId)
+				{
+					ActiveListMode = EEntitiesListMode::Default; // Hack to trigger mode change via mode switcher.
+
+					ActiveProtocol = NewMode.ModeId;
+
+					if (bIsInProtocolsMode.Get())
+					{
+						RebuildListWithColumns(EEntitiesListMode::Protocols);
+					}
+				}
+			}
+		)
+		.Visibility_Lambda([this]()
+			{
+				return bIsInProtocolsMode.Get() && !ActiveProtocol.IsNone() ? EVisibility::Visible : EVisibility::Collapsed;
+			}
+		);
+
+	{ // Add dynamic entries based on the available protocols.
+		IRemoteControlProtocolModule& RCProtocolModule = IRemoteControlProtocolModule::Get();
+
+		TArray<FName> Protocols = RCProtocolModule.GetProtocolNames();
+
+		// To avoid random orientations, sort them in alphabetical order.
+		Algo::Sort(Protocols, [&](const FName& ThisProtocol, const FName& OtherProtocol) { return ThisProtocol.LexicalLess(OtherProtocol); });
+
+		for (int32 ProtocolIndex = 0; ProtocolIndex < Protocols.Num(); ProtocolIndex++)
+		{
+			const FName& ProtocolName = Protocols[ProtocolIndex];
+
+			const bool bIsDefault = ProtocolIndex == 0;
+
+			if (bIsDefault)
+			{
+				ActiveProtocol = ProtocolName;
+			}
+
+			SRCModeSwitcher::FRCMode::FArguments NewMode = SRCModeSwitcher::Mode(ProtocolName)
+				.DefaultLabel(FText::FromName(ProtocolName))
+				.DefaultTooltip(FText::FromName(ProtocolName))
+				.HAlignCell(HAlign_Fill)
+				.IsDefault(bIsDefault)
+				.VAlignCell(VAlign_Fill)
+				.FixedWidth(48.f);
+
+			ModeSwitcher->AddMode(NewMode);
+		}
+	}
+
 	// Expose Button
 	ExposeDockPanel->AddHeaderToolbarItem(EToolbar::Left, PlaceholderBox.ToSharedRef());
+	ExposeDockPanel->AddHeaderToolbarItem(EToolbar::Right, ModeSwitcher.ToSharedRef());
 	ExposeDockPanel->AddFooterToolbarItem(EToolbar::Right, InArgs._ExposeComboButton.Get().ToSharedRef());
 	ExposeDockPanel->AddFooterToolbarItem(EToolbar::Right, DeleteAllEntitiesButton.ToSharedRef());
 
@@ -405,6 +489,92 @@ void SRCPanelExposedEntitiesList::SetBackendFilter(const FRCFilter& InBackendFil
 	bFilterApplicationRequested = true;
 }
 
+void SRCPanelExposedEntitiesList::RebuildListWithColumns(EEntitiesListMode InListMode)
+{
+	// Only activate when the active mode changes to different one.
+	// Drop multiple calls to same mode change except for Protocols Mode.
+	if (ActiveListMode != InListMode && !ActiveProtocol.IsNone())
+	{
+		ActiveListMode = InListMode;
+
+		TSet<FName> ExisitingColumns;
+
+		if (FieldsListView.IsValid() && FieldsHeaderRow.IsValid())
+		{
+			const TIndirectArray<SHeaderRow::FColumn>& Columns = FieldsHeaderRow->GetColumns();
+
+			if (!Columns.IsEmpty())
+			{
+				Algo::TransformIf(Columns, ExisitingColumns
+					, [](const SHeaderRow::FColumn& InColumn) { return !InColumn.ColumnId.IsNone() && !SRCPanelTreeNode::DefaultColumns.Contains(InColumn.ColumnId); }
+					, [](const SHeaderRow::FColumn& InColumn) { return InColumn.ColumnId; }
+				);
+			}
+		}
+
+		const bool bShouldRemoveColumns = !ExisitingColumns.IsEmpty();
+
+		TSet<FName> ColumnsToBeRemoved = ExisitingColumns;
+
+		const bool bShouldAddColumns = InListMode == EEntitiesListMode::Protocols;
+
+		TSet<FName> ColumnsToBeAdded;
+
+		IRemoteControlProtocolModule& RCProtocolModule = IRemoteControlProtocolModule::Get();
+
+		TSet<FName> ActiveProtocolColumns;
+
+		TSet<FName> OtherProtocolColumns;
+
+		for (const FName& ProtocolName : RCProtocolModule.GetProtocolNames())
+		{
+			if (TSharedPtr<IRemoteControlProtocol> RCProtocol = RCProtocolModule.GetProtocolByName(ProtocolName))
+			{
+				if (ActiveProtocol == ProtocolName)
+				{
+					RCProtocol->GetRegisteredColumns(ActiveProtocolColumns);
+				}
+				else
+				{
+					RCProtocol->GetRegisteredColumns(OtherProtocolColumns);
+				}
+			}
+		}
+
+		if (bShouldRemoveColumns)
+		{
+			ColumnsToBeRemoved.Append(ActiveProtocolColumns.Intersect(OtherProtocolColumns));
+		
+			for (const FName& ColumnToBeRemoved : ColumnsToBeRemoved)
+			{
+				RemoveColumn(ColumnToBeRemoved);
+			}
+		}
+
+		if (bShouldAddColumns)
+		{
+			ColumnsToBeAdded = DefaultProtocolColumns.Union(ActiveProtocolColumns);
+		
+			for (const FName& ColumnToBeAdded : ColumnsToBeAdded)
+			{
+				InsertColumn(ColumnToBeAdded);
+			}
+		}
+	}
+}
+
+FText SRCPanelExposedEntitiesList::HandleEntityListHeaderLabel() const
+{
+	const FRemoteControlCommands& Commands = FRemoteControlCommands::Get();
+
+	if (bIsInProtocolsMode.Get() && Commands.ToggleProtocolMappings.IsValid())
+	{
+		return Commands.ToggleProtocolMappings->GetLabel();
+	}
+
+	return LOCTEXT("PropertiesLabel", "Properties");
+}
+
 void SRCPanelExposedEntitiesList::OnObjectPropertyChange(UObject* InObject, FPropertyChangedEvent& InChangeEvent)
 {
 	EPropertyChangeType::Type TypesNeedingRefresh = EPropertyChangeType::ArrayAdd | EPropertyChangeType::ArrayClear | EPropertyChangeType::ArrayRemove | EPropertyChangeType::ValueSet;
@@ -638,6 +808,7 @@ TSharedRef<ITableRow> SRCPanelExposedEntitiesList::OnGenerateRow(TSharedPtr<SRCP
 			.OnDrop_Lambda(OnDropLambda)
 			.Padding(Margin)
 			.Style(&RCPanelStyle->TableRowStyle)
+			.ActiveProtocol_Lambda([this]() { return ActiveProtocol; })
 			.Entity(Node);
 	}
 }
@@ -827,11 +998,15 @@ void SRCPanelExposedEntitiesList::SelectActorsInlevel(const TArray<UObject*>& Ob
 void SRCPanelExposedEntitiesList::RegisterEvents()
 {
 	OnPropertyChangedHandle = FCoreUObjectDelegates::OnObjectPropertyChanged.AddSP(this, &SRCPanelExposedEntitiesList::OnObjectPropertyChange);
+
+	OnProtocolBindingAddedOrRemovedHandle = IRemoteControlProtocolWidgetsModule::Get().OnProtocolBindingAddedOrRemoved().AddSP(this, &SRCPanelExposedEntitiesList::OnProtocolBindingAddedOrRemoved);
 }
 
 void SRCPanelExposedEntitiesList::UnregisterEvents()
 {
 	FCoreUObjectDelegates::OnObjectPropertyChanged.Remove(OnPropertyChangedHandle);
+
+	IRemoteControlProtocolWidgetsModule::Get().OnProtocolBindingAddedOrRemoved().Remove(OnProtocolBindingAddedOrRemovedHandle);
 }
 
 TSharedPtr<SRCPanelGroup> SRCPanelExposedEntitiesList::FindGroupById(const FGuid& Id)
@@ -1142,9 +1317,9 @@ FReply SRCPanelExposedEntitiesList::RequestDeleteAllEntities()
 		return FReply::Unhandled();
 	}
 
-	const FString WarningMessage = FString::Format(TEXT("You are about to delete '{0}' entities. This action might not be undone.\nAre you sure you want to proceed?"), { FString::FromInt(FieldWidgetMap.Num()) });
+	const FText WarningMessage = FText::Format(LOCTEXT("DeleteAllEntitiesWarning", "You are about to delete '{0}' entities. This action might not be undone.\nAre you sure you want to proceed?"), FieldWidgetMap.Num());
 
-	EAppReturnType::Type UserResponse = FMessageDialog::Open(EAppMsgType::YesNo, FText::FromString(WarningMessage));
+	EAppReturnType::Type UserResponse = FMessageDialog::Open(EAppMsgType::YesNo, WarningMessage);
 
 	if (UserResponse == EAppReturnType::Yes)
 	{
@@ -1167,9 +1342,9 @@ FReply SRCPanelExposedEntitiesList::RequestDeleteAllGroups()
 		return FReply::Unhandled();
 	}
 
-	const FString WarningMessage = FString::Format(TEXT("You are about to delete '{0}' groups. This action might not be undone.\nAre you sure you want to proceed?"), { FString::FromInt(FieldGroups.Num() - 1) });
+	const FText WarningMessage = FText::Format(LOCTEXT("DeleteAllGroupsWarning", "You are about to delete '{0}' groups. This action might not be undone.\nAre you sure you want to proceed?"), FieldGroups.Num() - 1);
 
-	EAppReturnType::Type UserResponse = FMessageDialog::Open(EAppMsgType::YesNo, FText::FromString(WarningMessage));
+	EAppReturnType::Type UserResponse = FMessageDialog::Open(EAppMsgType::YesNo, WarningMessage);
 
 	if (UserResponse == EAppReturnType::Yes)
 	{
@@ -1197,6 +1372,196 @@ FReply SRCPanelExposedEntitiesList::RequestDeleteAllGroups()
 	return FReply::Handled();
 }
 
+SHeaderRow::FColumn::FArguments SRCPanelExposedEntitiesList::CreateColumn(const FName ForColumnName)
+{
+	const FText DefaultColumnLabel = GetColumnLabel(ForColumnName);
+
+	if (ForColumnName == FRemoteControlPresetColumns::Status || ForColumnName == FRemoteControlDefaultProtocolColumns::BindingStatus) // Exception for Status & BindingStatus columns as they are fixed ones.
+	{
+		const SHeaderRow::FColumn::FArguments ColumnArgs = SHeaderRow::Column(ForColumnName)
+			.DefaultLabel(DefaultColumnLabel)
+			.FixedWidth(33.f)
+			.HAlignHeader(HAlign_Center)
+			.HeaderContentPadding(RCPanelStyle->HeaderRowPadding);
+
+		return ColumnArgs;
+	}
+
+	const SHeaderRow::FColumn::FArguments ColumnArgs = SHeaderRow::Column(ForColumnName)
+		.DefaultLabel(DefaultColumnLabel)
+		.FillWidth(this, &SRCPanelExposedEntitiesList::GetColumnSize, ForColumnName)
+		.HAlignHeader(HAlign_Center)
+		.HeaderContentPadding(RCPanelStyle->HeaderRowPadding);
+
+	return ColumnArgs;
+}
+
+int32 SRCPanelExposedEntitiesList::GetColumnIndex(const FName& ForColumn) const
+{
+	if (ForColumn == FRemoteControlPresetColumns::LinkIdentifier)
+	{
+		return GetColumnIndex_Internal(ForColumn, FRemoteControlPresetColumns::DragDropHandle, ERCColumn::ERC_After);
+	}
+	else if (ForColumn == FRemoteControlPresetColumns::Mask)
+	{
+		return GetColumnIndex_Internal(ForColumn, FRemoteControlPresetColumns::Description, ERCColumn::ERC_After);
+	}
+	else if (ForColumn == FRemoteControlPresetColumns::Status)
+	{
+		return GetColumnIndex_Internal(ForColumn, FRemoteControlPresetColumns::Value, ERCColumn::ERC_After);
+	}
+	else if (ForColumn == FRemoteControlPresetColumns::Type)
+	{
+		return GetColumnIndex_Internal(ForColumn, FRemoteControlPresetColumns::Value, ERCColumn::ERC_After);
+	}
+	else
+	{
+		return GetColumnIndex_Internal(ForColumn, FRemoteControlPresetColumns::Value, ERCColumn::ERC_After);
+	}
+
+	return INDEX_NONE;
+}
+
+int32 SRCPanelExposedEntitiesList::GetColumnIndex_Internal(const FName& ForColumn, const FName& ExistingColumnName, ERCColumn::Position InPosition) const
+{
+	int32 InsertIndex = INDEX_NONE;
+
+	if (FieldsListView.IsValid() && FieldsHeaderRow.IsValid())
+	{
+		const TIndirectArray<SHeaderRow::FColumn>& ExistingColumns = FieldsHeaderRow->GetColumns();
+		
+		if (ExistingColumns.IsEmpty())
+		{
+			return InsertIndex;
+		}
+
+		auto ColumnItr = ExistingColumns.CreateConstIterator();
+		
+		while (ColumnItr)
+		{
+			// If the given column already exists then we need not to insert it again.
+			if (ColumnItr->ColumnId == ForColumn)
+			{
+				InsertIndex = INDEX_NONE;
+
+				break;
+			}
+
+			switch (InPosition)
+			{
+				case ERCColumn::ERC_After:
+					if (ColumnItr->ColumnId == ExistingColumnName)
+					{
+						InsertIndex = ColumnItr.GetIndex() + 1;
+
+						break;
+					}
+					break;
+				case ERCColumn::ERC_Before:
+					if (ColumnItr->ColumnId == ExistingColumnName)
+					{
+						InsertIndex = ColumnItr.GetIndex() - 1;
+
+						break;
+					}
+					break;
+			}
+
+			++ColumnItr;
+		}
+	}
+
+	return InsertIndex;
+}
+
+FText SRCPanelExposedEntitiesList::GetColumnLabel(const FName& ForColumn) const
+{
+	if (ForColumn == FRemoteControlPresetColumns::LinkIdentifier)
+	{
+		return LOCTEXT("RCPresetLinkIDColumnHeader", "Link ID");
+	}
+	else if (ForColumn == FRemoteControlPresetColumns::Mask)
+	{
+		return LOCTEXT("RCPresetMaskColumnHeader", "Mask");
+	}
+	else if (ForColumn == FRemoteControlPresetColumns::Status)
+	{
+		return LOCTEXT("RCPresetStatusColumnHeader", "");
+	}
+	else
+	{
+		IRemoteControlProtocolModule& RCProtocolModule = IRemoteControlProtocolModule::Get();
+
+		if (TSharedPtr<IRemoteControlProtocol> RCProtocol = RCProtocolModule.GetProtocolByName(ActiveProtocol))
+		{
+			if (const FProtocolColumnPtr& ProtocolColumn = RCProtocol->GetRegisteredColumn(ForColumn))
+			{
+				return ProtocolColumn->DisplayText;
+			}
+		}
+	}
+
+	return FText::GetEmpty();
+}
+
+float SRCPanelExposedEntitiesList::GetColumnSize(const FName ForColumn) const
+{
+	float ColumnSize = ProtocolColumnConstants::ColumnSizeSmall;
+
+	if (ForColumn == FRemoteControlPresetColumns::LinkIdentifier ||
+		ForColumn == FRemoteControlPresetColumns::Type)
+	{
+		ColumnSize = ProtocolColumnConstants::ColumnSizeMini;
+	}
+	else
+	{
+		IRemoteControlProtocolModule& RCProtocolModule = IRemoteControlProtocolModule::Get();
+
+		if (TSharedPtr<IRemoteControlProtocol> RCProtocol = RCProtocolModule.GetProtocolByName(ActiveProtocol))
+		{
+			if (const FProtocolColumnPtr& ProtocolColumn = RCProtocol->GetRegisteredColumn(ForColumn))
+			{
+				ColumnSize = ProtocolColumn->ColumnSize;
+			}
+		}
+	}
+
+	return ColumnSize;
+}
+
+void SRCPanelExposedEntitiesList::InsertColumn(const FName& InColumnName)
+{
+	if (FieldsListView.IsValid() && FieldsHeaderRow.IsValid())
+	{
+		int32 InsertIndex = GetColumnIndex(InColumnName);
+
+		if (InsertIndex != INDEX_NONE)
+		{
+			const SHeaderRow::FColumn::FArguments NewColumn = CreateColumn(InColumnName);
+
+			FieldsHeaderRow->InsertColumn(NewColumn, InsertIndex);
+		}
+	}
+}
+
+void SRCPanelExposedEntitiesList::RemoveColumn(const FName& InColumnName)
+{
+	if (FieldsListView.IsValid() && FieldsHeaderRow.IsValid())
+	{
+		FieldsHeaderRow->RemoveColumn(InColumnName);
+	}
+}
+
+void SRCPanelExposedEntitiesList::OnProtocolBindingAddedOrRemoved(ERCProtocolBinding::Op BindingOperation)
+{
+	ActiveListMode = EEntitiesListMode::Default; // Hack to trigger mode change via mode switcher.
+
+	if (bIsInProtocolsMode.Get())
+	{
+		RebuildListWithColumns(EEntitiesListMode::Protocols);
+	}
+}
+
 bool FGroupDragEvent::IsDraggedFromSameGroup() const
 {
 	return DragOriginGroup->GetId() == DragTargetGroup->GetId();
diff --git a/Engine/Plugins/VirtualProduction/RemoteControl/Source/RemoteControlUI/Private/UI/SRCPanelExposedEntitiesList.h b/Engine/Plugins/VirtualProduction/RemoteControl/Source/RemoteControlUI/Private/UI/SRCPanelExposedEntitiesList.h
--- a/Engine/Plugins/VirtualProduction/RemoteControl/Source/RemoteControlUI/Private/UI/SRCPanelExposedEntitiesList.h
+++ b/Engine/Plugins/VirtualProduction/RemoteControl/Source/RemoteControlUI/Private/UI/SRCPanelExposedEntitiesList.h
@@ -20,6 +20,7 @@ struct FRemoteControlProperty;
 struct FRemoteControlPresetGroup;
 struct FRemoteControlFunction;
 class ITableRow;
+class SHeaderRow;
 class SRCPanelGroup;
 struct SRCPanelTreeNode;
 class SRemoteControlTarget;
@@ -28,6 +29,12 @@ class STableViewBase;
 class URemoteControlPreset;
 struct FRCPanelStyle;
 
+enum class EEntitiesListMode : uint8
+{
+	Default,
+	Protocols
+};
+
 /** Holds information about a group drag and drop event  */
 struct FGroupDragEvent
 {
@@ -54,6 +61,7 @@ public:
 		: _EditMode(true)
 	{}
 		SLATE_ATTRIBUTE(bool, EditMode)
+		SLATE_ATTRIBUTE(bool, ProtocolsMode)
 		SLATE_ATTRIBUTE(TSharedPtr<SWidget>, ExposeComboButton)
 		SLATE_EVENT(FSimpleDelegate, OnEntityListUpdated)
 	SLATE_END_ARGS()
@@ -78,6 +86,9 @@ public:
 	/** Notifies the entities list view that the filter-list filter has changed */
 	void SetBackendFilter(const FRCFilter& InBackendFilter);
 
+	/** Recreate the list with dynamic columns. */
+	void RebuildListWithColumns(EEntitiesListMode InListMode);
+	
 	/** Recreate everything in the panel. */
 	void Refresh();
 	
@@ -94,6 +105,8 @@ public:
 	FSimpleDelegate OnEntityListUpdated() { return OnEntityListUpdatedDelegate; }
 	
 private:
+	/** Handles label to be shown in the entity list header.  */
+	FText HandleEntityListHeaderLabel() const;
 	/** Handles object property changes, used to update arrays correctly.  */
 	void OnObjectPropertyChange(UObject* InObject, FPropertyChangedEvent& InChangeEvent);
 	/** Create exposed entity widgets. */
@@ -158,6 +171,16 @@ private:
 	/** Delete all operation for groups list. */
 	FReply RequestDeleteAllGroups();
 
+	SHeaderRow::FColumn::FArguments CreateColumn(const FName ForColumnName);
+	int32 GetColumnIndex(const FName& ForColumn) const;
+	int32 GetColumnIndex_Internal(const FName& ForColumn, const FName& ExistingColumnName, ERCColumn::Position InPosition) const;
+	FText GetColumnLabel(const FName& ForColumn) const;
+	float GetColumnSize(const FName ForColumn) const;
+	void InsertColumn(const FName& InColumnName);
+	void RemoveColumn(const FName& InColumnName);
+
+	void OnProtocolBindingAddedOrRemoved(ERCProtocolBinding::Op BindingOperation);
+
 private:
 	/** Holds the Groups list view. */
 	TSharedPtr<SListView<TSharedPtr<SRCPanelTreeNode>>> GroupsListView;
@@ -171,10 +194,14 @@ private:
 	TMap<FGuid, TSharedPtr<SRCPanelTreeNode>> FieldWidgetMap;
 	/** Whether the panel is in edit mode. */
 	TAttribute<bool> bIsInEditMode;
+	/** Whether the panel is in protocols mode. */
+	TAttribute<bool> bIsInProtocolsMode;
 	/** Holds the preset asset. */
 	TStrongObjectPtr<URemoteControlPreset> Preset;
 	/** Handle to the delegate called when an object property change is detected. */
 	FDelegateHandle OnPropertyChangedHandle;
+	/** Handle to the delegate called when a binding is added or removed. */
+	FDelegateHandle OnProtocolBindingAddedOrRemovedHandle;
 	/** Delegate called on selected group change. */
 	FOnSelectionChange OnSelectionChangeDelegate;
 	/** The column data shared between all tree nodes in order to share a splitter amongst all rows. */
@@ -195,4 +222,12 @@ private:
 	TSharedPtr<FText> SearchedText;
 	/** Panel Style reference. */
 	const FRCPanelStyle* RCPanelStyle;
+	/** Holds the header row of entities list. */
+	TSharedPtr<SHeaderRow> FieldsHeaderRow;
+	/** Holds the active list mode. */
+	EEntitiesListMode ActiveListMode;
+	/** Holds the active protocol enabled by the mode switcher. */
+	FName ActiveProtocol;
+	/** Columns to be present by default in protocols mode. */
+	static TSet<FName> DefaultProtocolColumns;
 };
\ No newline at end of file
diff --git a/Engine/Plugins/VirtualProduction/RemoteControl/Source/RemoteControlUI/Private/UI/SRCPanelExposedField.cpp b/Engine/Plugins/VirtualProduction/RemoteControl/Source/RemoteControlUI/Private/UI/SRCPanelExposedField.cpp
--- a/Engine/Plugins/VirtualProduction/RemoteControl/Source/RemoteControlUI/Private/UI/SRCPanelExposedField.cpp
+++ b/Engine/Plugins/VirtualProduction/RemoteControl/Source/RemoteControlUI/Private/UI/SRCPanelExposedField.cpp
@@ -5,6 +5,7 @@
 #include "Algo/Transform.h"
 #include "EditorFontGlyphs.h"
 #include "IDetailTreeNode.h"
+#include "IRemoteControlProtocolModule.h"
 #include "IRemoteControlUIModule.h"
 #include "Layout/Visibility.h"
 #include "Materials/MaterialInstanceDynamic.h"
@@ -21,6 +22,7 @@
 #include "Styling/RemoteControlStyles.h"
 #include "UObject/Object.h"
 #include "Widgets/Input/SButton.h"
+#include "Widgets/Masks/SRCProtocolMask.h"
 #include "Widgets/Layout/SBorder.h"
 #include "Widgets/Text/SInlineEditableTextBlock.h"
 #include "Widgets/Text/STextBlock.h"
@@ -129,6 +131,71 @@ void SRCPanelExposedField::SetIsHovered(bool bInIsHovered)
 	bIsHovered = bInIsHovered;
 }
 
+TSharedRef<SWidget> SRCPanelExposedField::GetProtocolWidget(const FName ForColumnName, const FName InProtocolName)
+{
+	if (TSharedPtr<FRemoteControlField> RCField = WeakField.Pin())
+	{
+		if (const TSharedPtr<FRemoteControlProperty> RCProperty = StaticCastSharedPtr<FRemoteControlProperty>(RCField))
+		{
+			for (const FRemoteControlProtocolBinding& RCProtocolIter : RCProperty->ProtocolBindings)
+			{
+				if (RCProtocolIter.GetProtocolName() == InProtocolName)
+				{
+					if (TSharedPtr<TStructOnScope<FRemoteControlProtocolEntity>> RCProtocolEntityPtr = RCProtocolIter.GetRemoteControlProtocolEntityPtr())
+					{
+						return RCProtocolEntityPtr->Get()->GetWidget(ForColumnName);
+					}
+				}
+			}
+		}
+	}
+
+	return SRCPanelTreeNode::GetProtocolWidget(ForColumnName, InProtocolName);
+}
+
+const bool SRCPanelExposedField::HasProtocolExtension() const
+{
+	return GetFieldType() == EExposedFieldType::Property;
+}
+
+const bool SRCPanelExposedField::GetProtocolBindingsNum() const
+{
+	IRemoteControlProtocolModule& RCProtocolModule = IRemoteControlProtocolModule::Get();
+
+	const TArray<FName>& AvailableProtocols = RCProtocolModule.GetProtocolNames();
+
+	if (AvailableProtocols.Num() < 2 || AvailableProtocols.IsEmpty())
+	{
+		return false;
+	}
+
+	const TArray<FName>& SupportedProtocols = AvailableProtocols.FilterByPredicate([&](const FName& ThisProtocol) { return SupportsProtocol(ThisProtocol); });
+
+	return SupportedProtocols.Num() >= 2;
+}
+
+const bool SRCPanelExposedField::SupportsProtocol(const FName& InProtocolName) const
+{
+	if (TSharedPtr<FRemoteControlField> RCField = WeakField.Pin())
+	{
+		if (const TSharedPtr<FRemoteControlProperty> RCProperty = StaticCastSharedPtr<FRemoteControlProperty>(RCField))
+		{
+			if (RCProperty->ProtocolBindings.Num())
+			{
+				for (TSet<FRemoteControlProtocolBinding>::TConstIterator RCProtocolIter(RCProperty->ProtocolBindings); RCProtocolIter; ++RCProtocolIter)
+				{
+					if (RCProtocolIter->GetProtocolName() == InProtocolName)
+					{
+						return true;
+					}
+				}
+			}
+		}
+	}
+
+	return false;
+}
+
 void SRCPanelExposedField::Refresh()
 {
 	if (TSharedPtr<FRemoteControlField> Field = WeakField.Pin())
@@ -146,6 +213,31 @@ void SRCPanelExposedField::Refresh()
 	}
 }
 
+TSharedRef<SWidget> SRCPanelExposedField::GetWidget(const FName ForColumnName, const FName InActiveProtocol)
+{
+	if (HasProtocolExtension())
+	{
+		if (ForColumnName == FRemoteControlPresetColumns::Mask)
+		{
+			return SNew(SRCProtocolMask, WeakField);
+		}
+		else if (ForColumnName == FRemoteControlPresetColumns::Status)
+		{
+			return SNew(STextBlock)
+				.Text(LOCTEXT("StatusText", "!"))
+				.TextStyle(&RCPanelStyle->HeaderTextStyle)
+				.ToolTipText_Lambda([this]() { return GetProtocolBindingsNum() ? LOCTEXT("StatusTooltip", "Entity is bound to two or more protocols.") : FText::GetEmpty(); })
+				.Visibility_Lambda([this]() { return GetProtocolBindingsNum() ? EVisibility::Visible : EVisibility::Collapsed; });
+		}
+		else if (!DefaultColumns.Contains(ForColumnName))
+		{
+			return GetProtocolWidget(ForColumnName, InActiveProtocol);
+		}
+	}
+
+	return SRCPanelTreeNode::GetWidget(ForColumnName, InActiveProtocol);
+}
+
 void SRCPanelExposedField::GetBoundObjects(TSet<UObject*>& OutBoundObjects) const
 {
 	if (TSharedPtr<FRemoteControlField> Field = WeakField.Pin())
diff --git a/Engine/Plugins/VirtualProduction/RemoteControl/Source/RemoteControlUI/Private/UI/SRCPanelExposedField.h b/Engine/Plugins/VirtualProduction/RemoteControl/Source/RemoteControlUI/Private/UI/SRCPanelExposedField.h
--- a/Engine/Plugins/VirtualProduction/RemoteControl/Source/RemoteControlUI/Private/UI/SRCPanelExposedField.h
+++ b/Engine/Plugins/VirtualProduction/RemoteControl/Source/RemoteControlUI/Private/UI/SRCPanelExposedField.h
@@ -47,12 +47,20 @@ struct SRCPanelExposedField : public SRCPanelExposedEntity
 	virtual ENodeType GetRCType() const override;
 	virtual bool HasChildren() const override;
 	virtual void Refresh() override;
+	virtual TSharedRef<SWidget> GetWidget(const FName ForColumnName, const FName InActiveProtocol) override;
 	//~ End SRCPanelTreeNode Interface
 
 	//~ SRCPanelExposedEntity Interface
 	virtual void SetIsHovered(bool bIsBeingHovered) override;
 	//~ End SRCPanelExposedEntity Interface
 
+	//~ BEGIN : IHasProtocolExtensibility Interface
+	virtual TSharedRef<SWidget> GetProtocolWidget(const FName ForColumnName, const FName InProtocolName = NAME_None) override;
+	virtual const bool HasProtocolExtension() const override;
+	virtual const bool GetProtocolBindingsNum() const override;
+	virtual const bool SupportsProtocol(const FName& InProtocolName) const override;
+	//~ END : IHasProtocolExtensibility Interface
+
 	/** Get a weak pointer to the underlying remote control field. */
 	TWeakPtr<FRemoteControlField> GetRemoteControlField() const { return WeakField; }
 
@@ -82,6 +90,7 @@ private:
 	TSharedRef<SWidget> ConstructCallFunctionButton(bool bIsEnabled = true);
 	/** Handles calling an exposed function.*/
 	FReply OnClickFunctionButton();
+
 private:
 	/** Weak pointer to the underlying RC Field. */
 	TWeakPtr<FRemoteControlField> WeakField;
diff --git a/Engine/Plugins/VirtualProduction/RemoteControl/Source/RemoteControlUI/Private/UI/SRCPanelTreeNode.cpp b/Engine/Plugins/VirtualProduction/RemoteControl/Source/RemoteControlUI/Private/UI/SRCPanelTreeNode.cpp
--- a/Engine/Plugins/VirtualProduction/RemoteControl/Source/RemoteControlUI/Private/UI/SRCPanelTreeNode.cpp
+++ b/Engine/Plugins/VirtualProduction/RemoteControl/Source/RemoteControlUI/Private/UI/SRCPanelTreeNode.cpp
@@ -10,7 +10,34 @@
 
 #define LOCTEXT_NAMESPACE "RemoteControlPanelNode"
 
-TSharedRef<SWidget> SRCPanelTreeNode::GetWidget(const FName& ForColumnName)
+TSet<FName> SRCPanelTreeNode::DefaultColumns = {
+	FRemoteControlPresetColumns::DragDropHandle,
+	FRemoteControlPresetColumns::Description,
+	FRemoteControlPresetColumns::Value,
+	FRemoteControlPresetColumns::Reset
+};
+
+TSharedRef<SWidget> SRCPanelTreeNode::GetProtocolWidget(const FName ForColumnName, const FName InProtocolName)
+{
+	return SNullWidget::NullWidget;
+}
+
+const bool SRCPanelTreeNode::HasProtocolExtension() const
+{
+	return false;
+}
+
+const bool SRCPanelTreeNode::GetProtocolBindingsNum() const
+{
+	return false;
+}
+
+const bool SRCPanelTreeNode::SupportsProtocol(const FName& InProtocolName) const
+{
+	return false;
+}
+
+TSharedRef<SWidget> SRCPanelTreeNode::GetWidget(const FName ForColumnName, const FName InActiveProtocol)
 {
 	if (ForColumnName == FRemoteControlPresetColumns::DragDropHandle)
 	{
diff --git a/Engine/Plugins/VirtualProduction/RemoteControl/Source/RemoteControlUI/Private/UI/SRCPanelTreeNode.h b/Engine/Plugins/VirtualProduction/RemoteControl/Source/RemoteControlUI/Private/UI/SRCPanelTreeNode.h
--- a/Engine/Plugins/VirtualProduction/RemoteControl/Source/RemoteControlUI/Private/UI/SRCPanelTreeNode.h
+++ b/Engine/Plugins/VirtualProduction/RemoteControl/Source/RemoteControlUI/Private/UI/SRCPanelTreeNode.h
@@ -3,6 +3,7 @@
 #pragma once
 
 #include "IRemoteControlUIModule.h"
+#include "IHasProtocolExtensibility.h"
 #include "Templates/SharedPointer.h"
 #include "Containers/Array.h"
 #include "Misc/Guid.h"
@@ -17,13 +18,29 @@ struct SRCPanelExposedMaterial;
 namespace FRemoteControlPresetColumns
 {
 	static FName DragDropHandle = TEXT("DragDropHandle");
+	static FName LinkIdentifier = TEXT("LinkID");
 	static FName Description = TEXT("Description");
+	static FName Mask = TEXT("Mask");
+	static FName Type = TEXT("Type");
 	static FName Value = TEXT("Value");
+	static FName Status = TEXT("Status");
 	static FName Reset = TEXT("Reset");
 }
 
+namespace ERCColumn
+{
+	enum Position
+	{
+		/** Places the given column after a specific column. */
+		ERC_After,
+
+		/** Places the given column before a specific column. */
+		ERC_Before
+	};
+}
+
 /** A node in the panel tree view. */
-struct SRCPanelTreeNode : public SCompoundWidget
+struct SRCPanelTreeNode : public SCompoundWidget, public IHasProtocolExtensibility
 {
 	enum ENodeType
 	{
@@ -37,6 +54,13 @@ struct SRCPanelTreeNode : public SCompoundWidget
 
 	virtual ~SRCPanelTreeNode() {}
 
+	//~ BEGIN : IHasProtocolExtensibility Interface
+	virtual TSharedRef<SWidget> GetProtocolWidget(const FName ForColumnName, const FName InProtocolName = NAME_None) override;
+	virtual const bool HasProtocolExtension() const override;
+	virtual const bool GetProtocolBindingsNum() const override;
+	virtual const bool SupportsProtocol(const FName& InProtocolName) const override;
+	//~ END : IHasProtocolExtensibility Interface
+
 	/** Get this tree node's childen. */
 	virtual void GetNodeChildren(TArray<TSharedPtr<SRCPanelTreeNode>>& OutChildren) const {}
 	/** Get this node's ID if any. */
@@ -56,7 +80,7 @@ struct SRCPanelTreeNode : public SCompoundWidget
 	/** Updates the highlight text to active search term. */
 	virtual void SetHighlightText(const FText& InHightlightText = FText::GetEmpty()) {};
 	/** Retrieves the referenced widget corresponding to the given column name. */
-	virtual TSharedRef<SWidget> GetWidget(const FName& ForColumnName);
+	virtual TSharedRef<SWidget> GetWidget(const FName ForColumnName, const FName InActiveProtocol);
 
 protected:
 	struct FMakeNodeWidgetArgs
@@ -83,6 +107,11 @@ private:
 	float GetRightColumnWidth() const;
 	void SetColumnWidth(float InWidth);
 
+public:
+
+	/** Columns to be present by default in protocols mode. */
+	static TSet<FName> DefaultColumns;
+
 protected:
 	/** Holds the row's columns' width. */
 	FRCColumnSizeData ColumnSizeData;
diff --git a/Engine/Plugins/VirtualProduction/RemoteControl/Source/RemoteControlUI/Private/UI/SRemoteControlPanel.cpp b/Engine/Plugins/VirtualProduction/RemoteControl/Source/RemoteControlUI/Private/UI/SRemoteControlPanel.cpp
--- a/Engine/Plugins/VirtualProduction/RemoteControl/Source/RemoteControlUI/Private/UI/SRemoteControlPanel.cpp
+++ b/Engine/Plugins/VirtualProduction/RemoteControl/Source/RemoteControlUI/Private/UI/SRemoteControlPanel.cpp
@@ -45,6 +45,7 @@
 #include "ScopedTransaction.h"
 #include "SClassViewer.h"
 #include "SRCLogger.h"
+#include "SRCModeSwitcher.h"
 #include "SRCPanelExposedEntitiesList.h"
 #include "SRCPanelFunctionPicker.h"
 #include "SRCPanelExposedActor.h"
@@ -156,8 +157,7 @@ public:
 		const TAttribute<FText> ActualToolTip = UICommand.IsValid() ? UICommand->GetDescription() : FText::GetEmpty();
 
 		// If we were supplied an image than go ahead and use that, otherwise we use a null widget
-		TSharedRef<SLayeredImage> IconWidget =
-			SNew(SLayeredImage)
+		TSharedRef<SLayeredImage> IconWidget = SNew(SLayeredImage)
 			.ColorAndOpacity(this, &SAutoResizeButton::GetIconForegroundColor)
 			.Visibility(EVisibility::HitTestInvisible)
 			.Image(this, &SAutoResizeButton::GetIconBrush);
@@ -292,38 +292,6 @@ void SRemoteControlPanel::Construct(const FArguments& InArgs, URemoteControlPres
 
 	BindRemoteControlCommands();
 
-	// Rebind All
-	AddToolbarWidget(SNew(SButton)
-		.ButtonStyle(&RCPanelStyle->FlatButtonStyle)
-		.Visibility_Lambda([this]() { return bShowRebindButton ? EVisibility::Visible : EVisibility::Collapsed; })
-		.OnClicked(this, &SRemoteControlPanel::OnClickRebindAllButton)
-		.ContentPadding(2.f)
-		[
-			SNew(SHorizontalBox)
-
-			+SHorizontalBox::Slot()
-			.HAlign(HAlign_Center)
-			.VAlign(VAlign_Center)
-			.AutoWidth()
-			.Padding(2.f)
-			[
-				SNew(SImage)
-				.Image(FAppStyle::Get().GetBrush("Icons.Link"))
-				.ColorAndOpacity(FSlateColor::UseForeground())
-			]
-
-			+SHorizontalBox::Slot()
-			.VAlign(VAlign_Center)
-			.Padding(4.f, 2.f)
-			.FillWidth(1.f)
-			[	
-				SNew(STextBlock)
-				.ToolTipText(LOCTEXT("RebindButtonToolTip", "Attempt to rebind all unbound entites of the preset."))
-				.Text(LOCTEXT("RebindButtonText", "Rebind All"))
-				.TextStyle(&RCPanelStyle->PanelTextStyle)
-			]
-		]);
-
 	// Show Log
 	AddToolbarWidget(SNew(SCheckBox)
 		.Style(&RCPanelStyle->ToggleButtonStyle)
@@ -391,7 +359,8 @@ void SRemoteControlPanel::Construct(const FArguments& InArgs, URemoteControlPres
 				CachedExposedPropertyArgs.Reset();
 			}
 		)
-		.EditMode_Lambda([this]() { return bIsInEditMode; });
+		.EditMode_Lambda([this]() { return bIsInEditMode; })
+		.ProtocolsMode_Lambda([this]() { return IsInProtocolsMode() && ActivePanel == ERCPanels::RCP_Protocols; });
 
 	EntityList->OnSelectionChange().AddSP(this, &SRemoteControlPanel::UpdateEntityDetailsView);
 
@@ -1829,7 +1798,7 @@ void SRemoteControlPanel::GenerateToolbar()
 				[
 					SNew(SSeparator)
 					.SeparatorImage(FAppStyle::Get().GetBrush("Separator"))
-					.Thickness(1.5f)
+					.Thickness(2.f)
 					.Orientation(EOrientation::Orient_Vertical)
 				];
 
@@ -1856,12 +1825,107 @@ void SRemoteControlPanel::GenerateToolbar()
 		.Padding(5.f, 0.f)
 		.HAlign(HAlign_Fill)
 		.VAlign(VAlign_Center)
-		.FillWidth(1.f)
+		.AutoWidth()
 		[
 			SNew(STextBlock)
 			.Text(this, &SRemoteControlPanel::HandlePresetName)
 		]
 		+ SHorizontalBox::Slot()
+		.Padding(5.f, 0.f)
+		.HAlign(HAlign_Fill)
+		.VAlign(VAlign_Center)
+		.FillWidth(1.f)
+		[
+			SNew(SSpacer)
+		]
+		+SHorizontalBox::Slot()
+		.Padding(5.f, 0.f)
+		.HAlign(HAlign_Fill)
+		.VAlign(VAlign_Center)
+		.AutoWidth()
+		[
+			SNew(SRCModeSwitcher)
+			.DefaultMode("Setup")
+			.OnModeSwitched_Lambda([this](const SRCModeSwitcher::FRCMode& NewMode)
+				{
+					UE_LOG(LogRemoteControl, Warning, TEXT("Active Mode : %s"), *NewMode.ModeId.ToString());
+				}
+			)
+		
+			+ SRCModeSwitcher::Mode("Setup")
+			.DefaultLabel(FText::Format(LOCTEXT("SetupModeLabel", "{0}"), { FText::FromString("Setup") }))
+			.DefaultTooltip(FText::Format(LOCTEXT("SetupModeTooltip", "Switch to {0} mode."), { FText::FromString("Setup") }))
+			.HAlignCell(HAlign_Fill)
+			.VAlignCell(VAlign_Fill)
+			.FixedWidth(96.f)
+			.OptionalIcon(FSlateIcon(FAppStyle::GetAppStyleSetName(), "LevelEditor.Tabs.Outliner"))
+
+			+ SRCModeSwitcher::Mode("Operation")
+			.DefaultLabel(FText::Format(LOCTEXT("OpModeLabel", "{0}"), { FText::FromString("Operation") }))
+			.DefaultTooltip(FText::Format(LOCTEXT("OpModeTooltip", "Switch to {0} mode."), { FText::FromString("Operation") }))
+			.HAlignCell(HAlign_Fill)
+			.VAlignCell(VAlign_Fill)
+			.FixedWidth(96.f)
+			.OptionalIcon(FSlateIcon(FAppStyle::GetAppStyleSetName(), "Icons.Blueprint"))
+		]
+		+ SHorizontalBox::Slot()
+		.Padding(5.f, 0.f)
+		.HAlign(HAlign_Fill)
+		.VAlign(VAlign_Center)
+		.FillWidth(1.f)
+		[
+			SNew(SSpacer)
+		]
+		// Companion Separator (Rebind All)
+		+SHorizontalBox::Slot()
+		.HAlign(HAlign_Center)
+		.VAlign(VAlign_Fill)
+		.AutoWidth()
+		[
+			SNew(SSeparator)
+			.SeparatorImage(FAppStyle::Get().GetBrush("Separator"))
+			.Thickness(1.5f)
+			.Orientation(EOrientation::Orient_Vertical)
+			.Visibility_Lambda([this]() { return bShowRebindButton ? EVisibility::Visible : EVisibility::Collapsed; })
+		]
+		// Rebind All (Statically added here as we cannot probagagte visibility to its companion separator when dynamically added.)
+		+ SHorizontalBox::Slot()
+		.Padding(5.f, 0.f)
+		.AutoWidth()
+		.VAlign(VAlign_Center)
+		[
+			SNew(SButton)
+			.ButtonStyle(&RCPanelStyle->FlatButtonStyle)
+			.Visibility_Lambda([this]() { return bShowRebindButton ? EVisibility::Visible : EVisibility::Collapsed; })
+			.OnClicked(this, &SRemoteControlPanel::OnClickRebindAllButton)
+			.ContentPadding(2.f)
+			[
+				SNew(SHorizontalBox)
+
+				+SHorizontalBox::Slot()
+				.HAlign(HAlign_Center)
+				.VAlign(VAlign_Center)
+				.AutoWidth()
+				.Padding(2.f)
+				[
+					SNew(SImage)
+					.Image(FAppStyle::Get().GetBrush("Icons.Link"))
+					.ColorAndOpacity(FSlateColor::UseForeground())
+				]
+
+				+SHorizontalBox::Slot()
+				.VAlign(VAlign_Center)
+				.Padding(4.f, 2.f)
+				.FillWidth(1.f)
+				[	
+					SNew(STextBlock)
+					.ToolTipText(LOCTEXT("RebindButtonToolTip", "Attempt to rebind all unbound entites of the preset."))
+					.Text(LOCTEXT("RebindButtonText", "Rebind All"))
+					.TextStyle(&RCPanelStyle->PanelTextStyle)
+				]
+			]
+		]
+		+ SHorizontalBox::Slot()
 		.HAlign(HAlign_Right)
 		.AutoWidth()
 		[
@@ -2037,6 +2101,13 @@ void SRemoteControlPanel::ToggleProtocolMappings_Execute()
 
 			PanelDrawer->TogglePanel(PropertiesPanel);
 		}
+
+		if (EntityList.IsValid())
+		{
+			const bool bToggleProtcolMode = IsInProtocolsMode() && ActivePanel == ERCPanels::RCP_Protocols;
+
+			EntityList->RebuildListWithColumns(bToggleProtcolMode ? EEntitiesListMode::Protocols : EEntitiesListMode::Default);
+		}
 	}
 }
 
@@ -2077,6 +2148,13 @@ void SRemoteControlPanel::OnRCPanelToggled(ERCPanels InPanelID)
 	if (InPanelID != ActivePanel)
 	{
 		ActivePanel = InPanelID;
+		
+		if (EntityList.IsValid())
+		{
+			const bool bToggleProtcolMode = IsInProtocolsMode() && ActivePanel == ERCPanels::RCP_Protocols;
+
+			EntityList->RebuildListWithColumns(bToggleProtcolMode ? EEntitiesListMode::Protocols : EEntitiesListMode::Default);
+		}
 	}
 }
 
diff --git a/Engine/Plugins/VirtualProduction/RemoteControl/Source/RemoteControlUI/Public/IHasProtocolExtensibility.h b/Engine/Plugins/VirtualProduction/RemoteControl/Source/RemoteControlUI/Public/IHasProtocolExtensibility.h
new file mode 100644
--- /dev/null
+++ b/Engine/Plugins/VirtualProduction/RemoteControl/Source/RemoteControlUI/Public/IHasProtocolExtensibility.h
@@ -0,0 +1,33 @@
+// Copyright Epic Games, Inc. All Rights Reserved.
+
+#pragma once
+
+#include "CoreMinimal.h"
+
+/**
+ * A simple interface provides protocol specific utilities.
+ */
+class IHasProtocolExtensibility
+{
+public:
+
+	/**
+	 * Retrieves the protocol specific widget if protocol is specified, column specific widget otherwise.
+	 */
+	virtual TSharedRef<SWidget> GetProtocolWidget(const FName ForColumnName, const FName InProtocolName = NAME_None) = 0;
+	
+	/**
+	 * Ensures the derived classes supports protocols.
+	 */
+	virtual const bool HasProtocolExtension() const = 0;
+	
+	/**
+	 * Returns true when the derived classes has two or more protocol bindings.
+	 */
+	virtual const bool GetProtocolBindingsNum() const = 0;
+
+	/**
+	 * Determines whether the derived class supports the given protocol.
+	 */
+	virtual const bool SupportsProtocol(const FName& InProtocolName) const = 0;
+};
diff --git a/Engine/Plugins/VirtualProduction/RemoteControl/Source/RemoteControlUI/RemoteControlUI.Build.cs b/Engine/Plugins/VirtualProduction/RemoteControl/Source/RemoteControlUI/RemoteControlUI.Build.cs
--- a/Engine/Plugins/VirtualProduction/RemoteControl/Source/RemoteControlUI/RemoteControlUI.Build.cs
+++ b/Engine/Plugins/VirtualProduction/RemoteControl/Source/RemoteControlUI/RemoteControlUI.Build.cs
@@ -34,6 +34,7 @@ public class RemoteControlUI : ModuleRules
 				"RemoteControl",
 				"RemoteControlCommon",
 				"RemoteControlLogic",				
+				"RemoteControlProtocol",
 				"RemoteControlProtocolWidgets",
 				"SceneOutliner",
 				"Slate",
