Fixed issue with loading in Remote Control Module in other module's load time

diff --git a/Engine/Plugins/VirtualProduction/RemoteControl/Source/RemoteControl/Private/RemoteControlModule.cpp b/Engine/Plugins/VirtualProduction/RemoteControl/Source/RemoteControl/Private/RemoteControlModule.cpp
--- a/Engine/Plugins/VirtualProduction/RemoteControl/Source/RemoteControl/Private/RemoteControlModule.cpp
+++ b/Engine/Plugins/VirtualProduction/RemoteControl/Source/RemoteControl/Private/RemoteControlModule.cpp
@@ -38,15 +38,6 @@
 DEFINE_LOG_CATEGORY(LogRemoteControl);
 #define LOCTEXT_NAMESPACE "RemoteControl"
 
-TSet<UScriptStruct*> FRemoteControlModule::MaskingSupportedStructs = { TBaseStructure<FColor>::Get()
-	, TBaseStructure<FLinearColor>::Get()
-	, TBaseStructure<FRotator>::Get()
-	, TBaseStructure<FVector>::Get()
-	, TBaseStructure<FVector4>::Get()
-	, TBaseStructure<FIntVector>::Get()
-	, TBaseStructure<FIntVector4>::Get()
-};
-
 struct FRCInterceptionPayload
 {
 	TArray<uint8> Payload;
@@ -761,6 +752,16 @@ void FRemoteControlModule::PerformMasking(const TSharedRef<FRCMaskingOperation>&
 
 bool FRemoteControlModule::SupportsMasking(const FProperty* InProperty) const
 {
+	/** Holds the set of structs supported by masking. */
+	static const TSet<UScriptStruct*> MaskingSupportedStructs = { TBaseStructure<FColor>::Get()
+		, TBaseStructure<FLinearColor>::Get()
+		, TBaseStructure<FRotator>::Get()
+		, TBaseStructure<FVector>::Get()
+		, TBaseStructure<FVector4>::Get()
+		, TBaseStructure<FIntVector>::Get()
+		, TBaseStructure<FIntVector4>::Get()
+	};
+	
 	if (const FStructProperty* StructProperty = CastField<FStructProperty>(InProperty))
 	{
 		return MaskingSupportedStructs.Contains(StructProperty->Struct);
diff --git a/Engine/Plugins/VirtualProduction/RemoteControl/Source/RemoteControl/Private/RemoteControlModule.h b/Engine/Plugins/VirtualProduction/RemoteControl/Source/RemoteControl/Private/RemoteControlModule.h
--- a/Engine/Plugins/VirtualProduction/RemoteControl/Source/RemoteControl/Private/RemoteControlModule.h
+++ b/Engine/Plugins/VirtualProduction/RemoteControl/Source/RemoteControl/Private/RemoteControlModule.h
@@ -192,9 +192,6 @@ private:
 
 	/** Holds the set of active masking operations. */
 	TSet<TSharedPtr<FRCMaskingOperation>> ActiveMaskingOperations;
-
-	/** Holds the set of structs supported by masking. */
-	static TSet<UScriptStruct*> MaskingSupportedStructs;
 };
 
 PRAGMA_ENABLE_DEPRECATION_WARNINGS
