Remote Control Logic - Change File Location

Making it fit with the other Built-in Behaviour File Locations

diff --git a/Engine/Plugins/VirtualProduction/RemoteControl/Source/RemoteControlLogic/Private/Behaviour/Builtin/Path/RCBehaviourSetAssetByPathNode.cpp b/Engine/Plugins/VirtualProduction/RemoteControl/Source/RemoteControlLogic/Private/Behaviour/Builtin/Path/RCBehaviourSetAssetByPathNode.cpp
new file mode 100644
--- /dev/null
+++ b/Engine/Plugins/VirtualProduction/RemoteControl/Source/RemoteControlLogic/Private/Behaviour/Builtin/Path/RCBehaviourSetAssetByPathNode.cpp
@@ -0,0 +1,57 @@
+// Copyright Epic Games, Inc. All Rights Reserved.
+
+#include "Behaviour/Builtin/Path/RCBehaviourSetAssetByPathNode.h"
+
+#include "Behaviour/Builtin/Path/RCSetAssetByPathBehaviour.h"
+#include "Behaviour/RCBehaviour.h"
+#include "Controller/RCController.h"
+#include "RCVirtualProperty.h"
+#include "RCVirtualPropertyContainer.h"
+
+URCBehaviourSetAssetByPathNode::URCBehaviourSetAssetByPathNode()
+{
+	DisplayName = NSLOCTEXT("Remote Control Behaviour", "Behaviour Name - Set Asset By Path", "Set Asset By Path");
+	BehaviorDescription = NSLOCTEXT("Remote Control Behaviour", "Behaviour Desc - Set Asset By Path", "Triggers an event which sets an object based on the selected Exposed Entity.");
+}
+
+bool URCBehaviourSetAssetByPathNode::Execute_Implementation(URCBehaviour* InBehaviour) const
+{
+	if (!ensure(InBehaviour))
+	{
+		return false;
+	}
+
+	URCController* RCController = InBehaviour->ControllerWeakPtr.Get();
+	if (!RCController)
+	{
+		return false;
+	}
+
+	URCSetAssetByPathBehaviour* SetAssetByPathBehaviour = Cast<URCSetAssetByPathBehaviour>(InBehaviour);
+	if (!SetAssetByPathBehaviour)
+	{
+		return false;
+	}
+
+	FString DefaultString;
+	SetAssetByPathBehaviour->PropertyInContainer->GetVirtualProperty(SetAssetByPathBehaviourHelpers::DefaultInput)->GetValueString(DefaultString);
+
+	const FString ConcatPath = SetAssetByPathBehaviour->GetCurrentPath();
+	
+	return SetAssetByPathBehaviour->SetAssetByPath(ConcatPath, DefaultString);
+}
+
+bool URCBehaviourSetAssetByPathNode::IsSupported_Implementation(URCBehaviour* InBehaviour) const
+{
+	static TArray<TPair<EPropertyBagPropertyType, UObject*>> SupportedPropertyBagTypes =
+	{
+		{ EPropertyBagPropertyType::String, nullptr }
+	};
+	
+	return SupportedPropertyBagTypes.ContainsByPredicate(GetIsSupportedCallback(InBehaviour));
+}
+
+UClass* URCBehaviourSetAssetByPathNode::GetBehaviourClass() const
+{
+	return URCSetAssetByPathBehaviour::StaticClass();
+}
diff --git a/Engine/Plugins/VirtualProduction/RemoteControl/Source/RemoteControlLogic/Private/Behaviour/Builtin/Path/RCSetAssetByPathBehaviour.cpp b/Engine/Plugins/VirtualProduction/RemoteControl/Source/RemoteControlLogic/Private/Behaviour/Builtin/Path/RCSetAssetByPathBehaviour.cpp
new file mode 100644
--- /dev/null
+++ b/Engine/Plugins/VirtualProduction/RemoteControl/Source/RemoteControlLogic/Private/Behaviour/Builtin/Path/RCSetAssetByPathBehaviour.cpp
@@ -0,0 +1,420 @@
+// Copyright Epic Games, Inc. All Rights Reserved.
+
+#include "Behaviour/Builtin/Path/RCSetAssetByPathBehaviour.h"
+
+#include "Backends/CborStructDeserializerBackend.h"
+#include "Backends/CborStructSerializerBackend.h"
+#include "Components/MeshComponent.h"
+#include "Components/StaticMeshComponent.h"
+#include "Controller/RCController.h"
+#include "Engine/StaticMesh.h"
+#include "Engine/Texture2DDynamic.h"
+#include "Engine/World.h"
+#include "IRemoteControlModule.h"
+#include "IRemoteControlPropertyHandle.h"
+#include "Kismet/KismetRenderingLibrary.h"
+#include "Materials/MaterialInstance.h"
+#include "RCVirtualProperty.h"
+#include "RCVirtualPropertyContainer.h"
+#include "RemoteControlLogger.h"
+#include "RemoteControlPreset.h"
+
+URCSetAssetByPathBehaviour::URCSetAssetByPathBehaviour()
+{
+	PropertyInContainer = CreateDefaultSubobject<URCVirtualPropertyContainerBase>(FName("VirtualPropertyInContainer"));
+}
+
+void URCSetAssetByPathBehaviour::Initialize()
+{
+	PropertyInContainer->AddProperty(SetAssetByPathBehaviourHelpers::DefaultInput, URCController::StaticClass(), EPropertyBagPropertyType::String);
+	UpdateTargetEntity();
+	
+	Super::Initialize();
+}
+
+bool URCSetAssetByPathBehaviour::SetAssetByPath(const FString& AssetPath, const FString& DefaultString)
+{
+	const URCController* Controller = ControllerWeakPtr.Get();
+	
+	if (AssetPath.IsEmpty() || !Controller)
+	{
+		return false;
+	}
+	
+	FString ControllerString;
+	if (!Controller->GetValueString(ControllerString))
+	{
+		return false;
+	}
+
+	if (!bInternal)
+	{
+		FRemoteControlLogger::Get().Log(SetAssetByPathBehaviourHelpers::SetAssetByPathBehaviour, [AssetPath, ControllerString]
+		{
+			return FText::FromString(FString("Path Behaviour attempts to set Asset to %s") + *AssetPath + *ControllerString);
+		});
+		if (SetExternalAsset(AssetPath + ControllerString))
+		{
+			return true;
+		}
+
+		FRemoteControlLogger::Get().Log(SetAssetByPathBehaviourHelpers::SetAssetByPathBehaviour, [AssetPath, DefaultString]
+		{
+			return FText::FromString(FString("Path Behaviour attempts to set Asset to %s") + *AssetPath + *DefaultString);
+		});
+		if (!DefaultString.IsEmpty() && SetExternalAsset(AssetPath + DefaultString))
+		{
+			return true;
+		}
+
+		return false;
+	}
+	
+	FSoftObjectPath MainObjectRef(AssetPath + ControllerString);
+	
+	if (UObject* Object = MainObjectRef.TryLoad())
+	{
+		FRemoteControlLogger::Get().Log(SetAssetByPathBehaviourHelpers::SetAssetByPathBehaviour, [MainObjectRef]
+		{
+			return FText::FromString(FString("Path Behaviour attempts to set Asset to %s") + *MainObjectRef.GetAssetPathString());
+		});
+		return SetInternalAsset(Object);
+	}
+	
+	if (!DefaultString.IsEmpty())
+	{
+		FSoftObjectPath DefaultObjectRef(AssetPath + DefaultString);
+		
+		FRemoteControlLogger::Get().Log(SetAssetByPathBehaviourHelpers::SetAssetByPathBehaviour, [DefaultObjectRef]
+		{
+			return FText::FromString(FString("Path Behaviour attempts to set Asset to %s") + *DefaultObjectRef.GetAssetPathString());
+		});
+		return SetInternalAsset(DefaultObjectRef.TryLoad());
+	}
+	
+	return false;
+}
+
+bool URCSetAssetByPathBehaviour::SetInternalAsset(UObject* SetterObject)
+{
+	const URCController* Controller = ControllerWeakPtr.Get();
+	
+	if (!SetterObject || !Controller || !TargetEntity)
+	{
+		FRemoteControlLogger::Get().Log(SetAssetByPathBehaviourHelpers::SetAssetByPathBehaviour, [SetterObject]
+		{
+			return FText::FromString(FString("Path Behaviour fails to set Asset: %s") + *SetterObject->GetName());
+		}, EVerbosityLevel::Error);
+		return false;
+	}
+	
+	const TWeakObjectPtr<URemoteControlPreset> PresetPtr = Controller->PresetWeakPtr;
+	
+	TArray<TWeakPtr<FRemoteControlProperty>> ExposedProperties = PresetPtr->GetExposedEntities<FRemoteControlProperty>();
+
+	if (const TWeakPtr<FRemoteControlProperty>* ExposedPropertyPtr = ExposedProperties.FindByPredicate([this](const TWeakPtr<FRemoteControlProperty>& Property){ return Property.Pin()->GetLabel() == TargetEntity->GetLabel(); }))
+	{
+		const TSharedPtr<FRemoteControlProperty> ExposedProperty = (*ExposedPropertyPtr).Pin();
+		if (!ExposedProperty)
+		{
+			ensureMsgf(false, TEXT("Exposed Property is not valid."));
+			return false;
+		}
+		
+		if (AssetClass == UStaticMesh::StaticClass() && SetterObject->IsA(UStaticMesh::StaticClass()))
+		{
+			UStaticMeshComponent* StaticMeshComponent = Cast<UStaticMeshComponent>(ExposedProperty->GetBoundObject());
+			StaticMeshComponent->SetStaticMesh(Cast<UStaticMesh>(SetterObject));
+			
+			UE_LOG(LogRemoteControl, Display, TEXT("Path Behaviour sets Static Mesh Asset %s"), *SetterObject->GetName());
+			return true;
+		}
+		
+		if (AssetClass == UMaterial::StaticClass() && SetterObject->IsA(UMaterial::StaticClass()))
+		{
+			if (UMeshComponent* MaterialOwnerComponent = Cast<UMeshComponent>(ExposedProperty->GetBoundObject()))
+			{
+				MaterialOwnerComponent->SetMaterial(0, Cast<UMaterial>(SetterObject));
+			}
+			
+			UE_LOG(LogRemoteControl, Display, TEXT("Path Behaviour sets Material Asset %s"), *SetterObject->GetName());
+			return true;
+		}
+
+		// Only Texture runs through here
+		if (GetSupportedClasses().Contains(AssetClass) && SetterObject->IsA(UTexture::StaticClass()))
+		{
+			if (const TSharedPtr<FRemoteControlProperty> RemoteControlProperty = ExposedProperty)
+			{
+				SetTextureAsset(RemoteControlProperty, Cast<UTexture>(SetterObject));
+			}
+		}
+		
+		return false;
+	}
+	
+	if (TargetEntity)
+	{
+		const TWeakPtr<const FRemoteControlEntity> ExposedEntity = TargetEntity;
+		if (AssetClass == UBlueprint::StaticClass() && SetterObject->IsA(UBlueprint::StaticClass()))
+		{
+			AActor* OldActor = Cast<AActor>(ExposedEntity.Pin()->GetBoundObject());
+			UWorld* World = OldActor->GetWorld();
+			
+			OldActor->UnregisterAllComponents();
+			
+			const FVector OldLocation = OldActor->GetActorLocation();
+			const FRotator OldRotation = OldActor->GetActorRotation();
+			const FName OldActorName = OldActor->GetFName();
+			
+			FActorSpawnParameters SpawnParams;
+			SpawnParams.Name = OldActorName;
+			
+			const FName OldActorReplacedNamed = MakeUniqueObjectName(OldActor->GetOuter(), OldActor->GetClass(), *FString::Printf(TEXT("%s_REPLACED"), *OldActor->GetFName().ToString()));
+			OldActor->Rename(*OldActorReplacedNamed.ToString(), OldActor->GetOuter(), REN_DoNotDirty | REN_DontCreateRedirectors | REN_ForceNoResetLoaders);
+
+			if (AActor* NewActor = World->SpawnActor(Cast<UBlueprint>(SetterObject)->GeneratedClass, &OldLocation, &OldRotation, SpawnParams))
+			{
+				World->DestroyActor(OldActor);
+				
+				UE_LOG(LogRemoteControl, Display, TEXT("Path Behaviour sets Blueprint Asset %s"), *SetterObject->GetName());
+				return true;
+			}
+
+			UE_LOG(LogRemoteControl, Error, TEXT("Path Behaviour unable to delete old Actor %s"), *OldActor->GetName());
+			OldActor->UObject::Rename(*OldActorName.ToString(), OldActor->GetOuter(), REN_DoNotDirty | REN_DontCreateRedirectors | REN_ForceNoResetLoaders);
+			OldActor->RegisterAllComponents();
+		}
+	}
+	
+	UE_LOG(LogRemoteControl, Error, TEXT("Path Behaviour fails to set Asset %s"), *SetterObject->GetName());
+	return false;
+}
+
+TArray<UClass*> URCSetAssetByPathBehaviour::GetSupportedClasses() const
+{
+	static TArray<UClass*> SupportedClasses =
+	{
+		UStaticMesh::StaticClass(),
+		UMaterial::StaticClass(),
+		UTexture::StaticClass(),
+		UBlueprint::StaticClass(),
+	};
+
+	return SupportedClasses;
+}
+
+FString URCSetAssetByPathBehaviour::GetCurrentPath()
+{
+	FString CurrentPath;
+	URCController* RCController = ControllerWeakPtr.Get();
+	if (!RCController)
+	{
+		ensureMsgf(false, TEXT("Controller is invalid."));
+		return TEXT("Path is invalid!");
+	}
+
+	CurrentPath.Empty();
+	FString ControllerString;
+	RCController->GetValueString(ControllerString);
+
+	// Add Path String Concat
+	CurrentPath = bInternal ? SetAssetByPathBehaviourHelpers::ContentFolder : FString("");
+	for (FString PathPart : PathStruct.PathArray)
+	{
+		if (PathPart.IsEmpty())
+		{
+			continue;
+		}
+		
+		// Change *INPUT String Token to the correct Controller
+		if (PathPart.Contains(SetAssetByPathBehaviourHelpers::InputToken, ESearchCase::CaseSensitive)
+			&& PathPart.Find(SetAssetByPathBehaviourHelpers::InputToken) == 0)
+		{
+			const URemoteControlPreset* RemoteControlPreset = RCController->PresetWeakPtr.Get();
+			if (!RemoteControlPreset)
+			{
+				ensureMsgf(false, TEXT("Remote Control Preset is invalid"));
+				continue;
+			}
+			
+			PathPart.RemoveFromStart(SetAssetByPathBehaviourHelpers::InputToken);
+			// TODO: Test and rewrite to something better
+			int32 CharIndexStart;
+			int32 CharIndexEnd;
+			PathPart.FindChar('(', CharIndexStart);
+			PathPart.FindChar(')', CharIndexEnd);
+			if (CharIndexStart == INDEX_NONE || CharIndexEnd == INDEX_NONE)
+			{
+				// TODO: Add Message that Input token has been done wrongly.
+				ensureMsgf(false, TEXT("Token has been input wrongly"));
+				continue;
+			}
+
+			PathPart.RemoveAt(CharIndexStart);
+			PathPart.RemoveAt(CharIndexEnd-1);
+			const URCVirtualPropertyBase* TokenController = RemoteControlPreset->GetVirtualPropertyByDisplayName(FName(PathPart));
+			if (!TokenController)
+			{
+				ensureMsgf(false, TEXT("No Controller with given name found."));
+				continue;
+			}
+
+			PathPart = TokenController->GetDisplayValueAsString();
+		}
+
+		// Check for certain chars.
+		int32 IndexFound;
+		if (PathPart.FindLastChar('_', IndexFound) && IndexFound == PathPart.Len() - 1)
+		{
+			// In the case there's underscore char in the at the end of one of the Path Strings, do nothing to facilitate more complex pathing behaviours.
+		}
+		else if (!PathPart.FindLastChar('/', IndexFound) || IndexFound < PathPart.Len() - 1)
+		{
+			PathPart = PathPart.AppendChar('/');
+		}
+
+		PathPart.ReplaceCharInline(TCHAR('\\'), TCHAR('/'), ESearchCase::IgnoreCase);
+		
+		CurrentPath += PathPart;
+	}
+
+	return CurrentPath;
+}
+
+void URCSetAssetByPathBehaviour::SetTargetEntity(const TSharedPtr<const FRemoteControlEntity>& InEntity)
+{
+	TargetEntityId = InEntity->GetId();
+	TargetEntity = InEntity;
+}
+
+TWeakPtr<const FRemoteControlEntity> URCSetAssetByPathBehaviour::GetTargetEntity() const
+{
+	return TargetEntity;
+}
+
+bool URCSetAssetByPathBehaviour::SetExternalAsset(FString InExternalPath)
+{
+	URCController* Controller = ControllerWeakPtr.Get();
+	if (!Controller || !TargetEntity)
+	{
+		return false;
+	}
+	
+	URemoteControlPreset* PresetPtr = Controller->PresetWeakPtr.Get();
+	if (!PresetPtr)
+	{
+		return false;
+	}
+	
+	TArray<TWeakPtr<FRemoteControlProperty>> ExposedProperties = PresetPtr->GetExposedEntities<FRemoteControlProperty>();
+	const TWeakPtr<FRemoteControlProperty>* ExposedPropertyPtr = ExposedProperties.FindByPredicate([this](const TWeakPtr<FRemoteControlProperty>& Property)
+	{
+		if (const TSharedPtr<FRemoteControlProperty> PinnedProperty = Property.Pin())
+		{
+			return PinnedProperty->GetLabel() == TargetEntity->GetLabel();
+		}
+		
+		return false;
+	});
+	const TSharedPtr<FRemoteControlProperty> RemoteControlProperty = ExposedPropertyPtr->Pin();
+
+	if (!RemoteControlProperty)
+	{
+		FRemoteControlLogger::Get().Log(SetAssetByPathBehaviourHelpers::SetAssetByPathBehaviour, []
+		{
+			return FText::FromString(TEXT("Path Behaviour Exposed Property not found"));
+		});
+		return false;
+	}
+
+	// TODO: Find better solution.
+	UTexture2D* ImportedTexture = UKismetRenderingLibrary::ImportFileAsTexture2D(this, InExternalPath);
+	if (!ImportedTexture)
+	{
+		FRemoteControlLogger::Get().Log(SetAssetByPathBehaviourHelpers::SetAssetByPathBehaviour, [InExternalPath]
+		{
+			return FText::FromString(FString("Path Behaviour File not found: ") + *InExternalPath);
+		});
+		return false;
+	}
+
+	// Apply Texture
+	if (SetTextureAsset(RemoteControlProperty, ImportedTexture))
+	{
+		FRemoteControlLogger::Get().Log(SetAssetByPathBehaviourHelpers::SetAssetByPathBehaviour, [InExternalPath]
+		{
+			return FText::FromString(FString("Path Behaviour Set external Asset: ") + *InExternalPath);
+		});
+		return true;
+	}
+
+	FRemoteControlLogger::Get().Log(SetAssetByPathBehaviourHelpers::SetAssetByPathBehaviour, []
+	{
+		return FText::FromString(TEXT("Path Behaviour Set external Asset failed"));
+	});
+	return false;
+}
+
+bool URCSetAssetByPathBehaviour::SetTextureAsset(TSharedPtr<FRemoteControlProperty> InRemoteControlPropertyPtr, UTexture* InObject)
+{
+	if (!InObject)
+	{
+		return false;
+	}
+	
+	FRCObjectReference ObjectRef;
+	ObjectRef.Property = InRemoteControlPropertyPtr->GetProperty();
+	ObjectRef.Access = InRemoteControlPropertyPtr->GetPropertyHandle()->ShouldGenerateTransaction() ? ERCAccess::WRITE_TRANSACTION_ACCESS : ERCAccess::WRITE_ACCESS;
+	ObjectRef.PropertyPathInfo = InRemoteControlPropertyPtr->FieldPathInfo.ToString();
+
+	for (UObject* Object : InRemoteControlPropertyPtr->GetBoundObjects())
+	{
+		if (IRemoteControlModule::Get().ResolveObjectProperty(ObjectRef.Access, Object, ObjectRef.PropertyPathInfo, ObjectRef))
+		{
+			/**
+			 * Setting with Pointer and using Serialization of itself afterwards as a workaround with the Texture Asset not updating in the world.
+			 */
+			FObjectPropertyBase* ObjectProperty = CastField<FObjectPropertyBase>(InRemoteControlPropertyPtr->GetProperty());
+			if (!ObjectProperty)
+			{
+				return false;
+			}
+
+			FProperty* Property = InRemoteControlPropertyPtr->GetProperty();
+			uint8* ValueAddress = Property->ContainerPtrToValuePtr<uint8>(ObjectRef.ContainerAdress);
+			ObjectProperty->SetObjectPropertyValue(ValueAddress, InObject);
+
+			TArray<uint8> Buffer;
+			FMemoryReader Reader(Buffer);
+			FCborStructDeserializerBackend DeserializerBackend(Reader);
+
+			IRemoteControlModule::Get().SetObjectProperties(ObjectRef, DeserializerBackend, ERCPayloadType::Cbor, Buffer);
+			return true;
+		}
+	}
+	return false;	
+}
+
+void URCSetAssetByPathBehaviour::UpdateTargetEntity()
+{
+	const URCController* Controller = ControllerWeakPtr.Get();
+	if (!Controller)
+	{
+		return;
+	}
+	
+	const URemoteControlPreset* PresetPtr = Controller->PresetWeakPtr.Get();
+	if (!PresetPtr)
+	{
+		return;
+	}
+
+	TWeakPtr<const FRemoteControlEntity> ExposedEntity = PresetPtr->GetExposedEntity<FRemoteControlEntity>(TargetEntityId);
+	if (!ExposedEntity.Pin())
+	{
+		return;
+	}
+	
+	TargetEntity = ExposedEntity.Pin();
+}
diff --git a/Engine/Plugins/VirtualProduction/RemoteControl/Source/RemoteControlLogic/Private/Behaviour/Builtin/RCBehaviourSetAssetByPathNode.cpp b/Engine/Plugins/VirtualProduction/RemoteControl/Source/RemoteControlLogic/Private/Behaviour/Builtin/RCBehaviourSetAssetByPathNode.cpp
deleted file mode 100644
diff --git a/Engine/Plugins/VirtualProduction/RemoteControl/Source/RemoteControlLogic/Private/Behaviour/RCSetAssetByPathBehaviour.cpp b/Engine/Plugins/VirtualProduction/RemoteControl/Source/RemoteControlLogic/Private/Behaviour/RCSetAssetByPathBehaviour.cpp
deleted file mode 100644
diff --git a/Engine/Plugins/VirtualProduction/RemoteControl/Source/RemoteControlLogic/Public/Behaviour/Builtin/Path/RCBehaviourSetAssetByPathNode.h b/Engine/Plugins/VirtualProduction/RemoteControl/Source/RemoteControlLogic/Public/Behaviour/Builtin/Path/RCBehaviourSetAssetByPathNode.h
new file mode 100644
--- /dev/null
+++ b/Engine/Plugins/VirtualProduction/RemoteControl/Source/RemoteControlLogic/Public/Behaviour/Builtin/Path/RCBehaviourSetAssetByPathNode.h
@@ -0,0 +1,28 @@
+// Copyright Epic Games, Inc. All Rights Reserved.
+
+#pragma once
+
+#include "Behaviour/RCBehaviourNode.h"
+#include "RCBehaviourSetAssetByPathNode.generated.h"
+
+/**
+ * Takes the string as a path and goes on to search for the Asset it is connected to, setting it to a Target Exposed Property.
+ */
+UCLASS()
+class REMOTECONTROLLOGIC_API URCBehaviourSetAssetByPathNode : public URCBehaviourNode
+{
+	GENERATED_BODY()
+	
+public:
+	URCBehaviourSetAssetByPathNode();
+	
+	//~ Begin URCBehaviourNode interface
+	UFUNCTION(BlueprintNativeEvent, BlueprintCallable, Category = "Remote Control Behaviour")
+	bool Execute(URCBehaviour* InBehaviour) const override;
+
+	UFUNCTION(BlueprintNativeEvent, BlueprintCallable, Category = "Remote Control Behaviour")
+	bool IsSupported(URCBehaviour* InBehaviour) const override;
+
+	UClass* GetBehaviourClass() const override;
+	//~ End URCBehaviourNode interface
+};
diff --git a/Engine/Plugins/VirtualProduction/RemoteControl/Source/RemoteControlLogic/Public/Behaviour/Builtin/Path/RCSetAssetByPathBehaviour.h b/Engine/Plugins/VirtualProduction/RemoteControl/Source/RemoteControlLogic/Public/Behaviour/Builtin/Path/RCSetAssetByPathBehaviour.h
new file mode 100644
--- /dev/null
+++ b/Engine/Plugins/VirtualProduction/RemoteControl/Source/RemoteControlLogic/Public/Behaviour/Builtin/Path/RCSetAssetByPathBehaviour.h
@@ -0,0 +1,104 @@
+// Copyright Epic Games, Inc. All Rights Reserved.
+
+#pragma once
+
+#include "Behaviour/RCBehaviour.h"
+#include "RCVirtualProperty.h"
+#include "RemoteControlEntity.h"
+
+#include "RCSetAssetByPathBehaviour.generated.h"
+
+class URCVirtualPropertyContainerBase;
+class URCUserDefinedStruct;
+class UTexture;
+
+namespace SetAssetByPathBehaviourHelpers
+{
+	const FString ContentFolder = FString(TEXT("/Game/"));
+	const FString InputToken = FString(TEXT("{INPUT}"));
+	const FName TargetProperty = FName(TEXT("Target Property"));
+	const FName DefaultInput = FName(TEXT("Default Input"));
+	const FName SetAssetByPathBehaviour = FName(TEXT("Set Asset By Path"));
+}
+
+/** Struct to help generate Widgts for the DetailsPanel of the Bahviour */
+USTRUCT()
+struct FRCSetAssetPath
+{
+	GENERATED_BODY()
+
+	FRCSetAssetPath() = default;
+	
+	/** An Array of Strings holding the Path of an Asset, seperated in several String. Will concated back together later. */
+	UPROPERTY()
+	TArray<FString> PathArray;
+};
+
+/**
+ * Custom behaviour for Set Asset By Path
+ */
+UCLASS()
+class REMOTECONTROLLOGIC_API URCSetAssetByPathBehaviour : public URCBehaviour
+{
+	GENERATED_BODY()
+public:
+	URCSetAssetByPathBehaviour();
+
+	//~ Begin URCBehaviour interface
+	virtual void Initialize() override;
+	//~ End URCBehaviour interface
+
+	/** Given an Input Path, sets the Target Exposed Propterty to the Asset. */
+	bool SetAssetByPath(const FString& AssetPath, const FString& DefaultString);
+
+	/** List of Supported Assets SetAssetByPath Logic Behaviour can set and change. */
+	TArray<UClass*> GetSupportedClasses() const;
+
+	/** Returns the Path as the completed version, including Token Inputs. */
+	FString GetCurrentPath();
+
+	/** Sets the target entity upon which the asset will be set upon. */
+	void SetTargetEntity(const TSharedPtr<const FRemoteControlEntity>& InEntity);
+
+	/** Returns a Pointer to the current Target Entityy */
+	TWeakPtr<const FRemoteControlEntity> GetTargetEntity() const;
+
+	/** Auxialiary Function to apply to update the Target Texture */
+	void UpdateTargetEntity();
+	
+public:
+	/** Pointer to property container */
+	UPROPERTY()
+	TObjectPtr<URCVirtualPropertyContainerBase> PropertyInContainer;
+
+	/** Pointer to the current Class the Asset will use to set the Targeted Exposed Object. */
+	UPROPERTY()
+	TObjectPtr<UClass> AssetClass = GetSupportedClasses()[0];
+
+	/** Struct holding the Path Information to help load Assets. */
+	UPROPERTY()
+	FRCSetAssetPath PathStruct;
+
+	/** Bool used to help tell if the path given is an internal or external one. */
+	UPROPERTY()
+	bool bInternal = true;
+
+private:
+	/** Targeted Property Id */
+	UPROPERTY()
+	FGuid TargetEntityId;
+
+	/** Internal Targeted Property, used for any of the setter operations */
+	TSharedPtr<const FRemoteControlEntity> TargetEntity;
+	
+private:
+	/** Auxiliary Function which sets the given SetterObject Asset onto the Exposed Asset with the given PropertyString name. */
+	bool SetInternalAsset(UObject* SetterObject);
+
+	/** Given an External Path towards a external location, loads the asset associated and places it onto an object */
+	bool SetExternalAsset(FString InExternalPath);
+
+	/** Auxiliary Function to apply a Texture onto a given Property */
+	bool SetTextureAsset(TSharedPtr<FRemoteControlProperty> InRemoteControlPropertyPtr, UTexture* InObject);
+
+};
diff --git a/Engine/Plugins/VirtualProduction/RemoteControl/Source/RemoteControlLogic/Public/Behaviour/Builtin/RCBehaviourSetAssetByPathNode.h b/Engine/Plugins/VirtualProduction/RemoteControl/Source/RemoteControlLogic/Public/Behaviour/Builtin/RCBehaviourSetAssetByPathNode.h
deleted file mode 100644
diff --git a/Engine/Plugins/VirtualProduction/RemoteControl/Source/RemoteControlLogic/Public/Behaviour/RCSetAssetByPathBehaviour.h b/Engine/Plugins/VirtualProduction/RemoteControl/Source/RemoteControlLogic/Public/Behaviour/RCSetAssetByPathBehaviour.h
deleted file mode 100644
diff --git a/Engine/Plugins/VirtualProduction/RemoteControl/Source/RemoteControlUI/Private/UI/Behaviour/Builtin/Path/RCBehaviourSetAssetByPathModel.cpp b/Engine/Plugins/VirtualProduction/RemoteControl/Source/RemoteControlUI/Private/UI/Behaviour/Builtin/Path/RCBehaviourSetAssetByPathModel.cpp
new file mode 100644
--- /dev/null
+++ b/Engine/Plugins/VirtualProduction/RemoteControl/Source/RemoteControlUI/Private/UI/Behaviour/Builtin/Path/RCBehaviourSetAssetByPathModel.cpp
@@ -0,0 +1,417 @@
+// Copyright Epic Games, Inc. All Rights Reserved.
+
+#include "RCBehaviourSetAssetByPathModel.h"
+
+#include "Controller/RCController.h"
+#include "Editor.h"
+#include "IDetailTreeNode.h"
+#include "IPropertyRowGenerator.h"
+#include "Modules/ModuleManager.h"
+#include "RCVirtualProperty.h"
+#include "RCVirtualPropertyContainer.h"
+#include "RemoteControlPreset.h"
+#include "SRCBehaviourSetAssetByPath.h"
+#include "Selection.h"
+#include "Widgets/Images/SImage.h"
+#include "Widgets/Input/SButton.h"
+#include "Widgets/Input/SComboBox.h"
+#include "Widgets/Layout/SBox.h"
+#include "Widgets/Text/STextBlock.h"
+
+#define LOCTEXT_NAMESPACE "FRCSetAssetByPathBehaviourModel"
+
+FRCSetAssetByPathBehaviourModel::FRCSetAssetByPathBehaviourModel(URCSetAssetByPathBehaviour* SetAssetByPathBehaviour)
+	: FRCBehaviourModel(SetAssetByPathBehaviour)
+	, SetAssetByPathBehaviourWeakPtr(SetAssetByPathBehaviour)
+{
+	FPropertyRowGeneratorArgs Args;
+	Args.bShouldShowHiddenProperties = true;
+	PropertyRowGenerator = FModuleManager::GetModuleChecked<FPropertyEditorModule>("PropertyEditor").CreatePropertyRowGenerator(Args);
+	PropertyRowGeneratorArray = FModuleManager::GetModuleChecked<FPropertyEditorModule>("PropertyEditor").CreatePropertyRowGenerator(Args);
+
+	if (SetAssetByPathBehaviour)
+	{
+		PropertyRowGenerator->SetStructure(SetAssetByPathBehaviour->PropertyInContainer->CreateStructOnScope());
+		DetailTreeNodeWeakPtr.Empty();
+		for (const TSharedRef<IDetailTreeNode>& CategoryNode : PropertyRowGenerator->GetRootTreeNodes())
+		{
+			TArray<TSharedRef<IDetailTreeNode>> Children;
+			CategoryNode->GetChildren(Children);
+			for (const TSharedRef<IDetailTreeNode>& Child : Children)
+			{
+				DetailTreeNodeWeakPtr.Add(Child);
+			}
+		}
+
+		// Secondary TArray Struct
+		PropertyRowGeneratorArray->SetStructure(MakeShareable(new FStructOnScope(FRCSetAssetPath::StaticStruct(), (uint8*) &SetAssetByPathBehaviour->PathStruct)));
+		PropertyRowGeneratorArray->OnFinishedChangingProperties().AddLambda([this](const FPropertyChangedEvent& InEvent)
+		{
+			if (InEvent.ChangeType != EPropertyChangeType::ValueSet)
+			{
+				return;
+			}
+
+			RefreshPathAndPreview();
+		});
+		PropertyRowGeneratorArray->OnRowsRefreshed().AddLambda([this]()
+		{
+			RefreshPathAndPreview();
+		});
+		DetailTreeNodeWeakPtrArray.Empty();
+		for (const TSharedRef<IDetailTreeNode>& CategoryNode : PropertyRowGeneratorArray->GetRootTreeNodes())
+		{
+			TArray<TSharedRef<IDetailTreeNode>> Children;
+			CategoryNode->GetChildren(Children);
+			for (const TSharedRef<IDetailTreeNode>& Child : Children)
+			{
+				DetailTreeNodeWeakPtrArray.Add(Child);
+			}
+		}
+	}
+
+	PreviewPathWidget = SNew(STextBlock)
+		.AutoWrapText(true);
+	PathArrayWidget = SNew(SBox);
+
+	Cast<URCSetAssetByPathBehaviour>(GetBehaviour())->UpdateTargetEntity();
+	TWeakPtr<const FRemoteControlEntity> InitialSelected = Cast<URCSetAssetByPathBehaviour>(GetBehaviour())->GetTargetEntity();
+	
+	RefreshPathAndPreview();
+	SelectorBox = GetSelectorWidget(InitialSelected);
+}
+
+TSharedRef<SWidget> FRCSetAssetByPathBehaviourModel::GetBehaviourDetailsWidget()
+{
+	return SNew(SRCBehaviourSetAssetByPath, SharedThis(this));
+}
+
+TSharedRef<SWidget> FRCSetAssetByPathBehaviourModel::GetPropertyWidget()
+{
+	TSharedRef<SVerticalBox> FieldWidget = SNew(SVerticalBox);
+	
+	const TSharedPtr<IDetailTreeNode> PinnedNodeDefault = DetailTreeNodeWeakPtr[0];
+
+	const FNodeWidgets NodeWidgetsDefault = PinnedNodeDefault->CreateNodeWidgets();
+	if (NodeWidgetsDefault.ValueWidget)
+	{
+		FieldWidget->AddSlot()
+			.Padding(FMargin(3.0f, 2.0f))
+			.AutoHeight()
+			[
+				SNew(SHorizontalBox)
+				+ SHorizontalBox::Slot()
+				[
+					SNew(STextBlock)
+						.Text(FText::FromString("Target Exposed Property"))
+				]
+				+ SHorizontalBox::Slot()
+				[
+					NodeWidgetsDefault.NameWidget.ToSharedRef()
+				]
+			];
+		
+		FieldWidget->AddSlot()
+			.Padding(FMargin(3.0f, 2.0f))
+			.AutoHeight()
+			[
+				SNew(SHorizontalBox)
+				+ SHorizontalBox::Slot()
+				[
+					SelectorBox.ToSharedRef()
+				]
+				+ SHorizontalBox::Slot()
+				[
+					NodeWidgetsDefault.ValueWidget.ToSharedRef()
+				]
+			];
+	}
+	else if (NodeWidgetsDefault.WholeRowWidget)
+	{
+		FieldWidget->AddSlot()
+			.Padding(FMargin(3.0f, 2.0f))
+			.AutoHeight()
+			[
+				NodeWidgetsDefault.WholeRowWidget.ToSharedRef()
+			];
+	}
+
+	FieldWidget->AddSlot()
+		.Padding(FMargin(3.f, 3.f))
+		.AutoHeight()
+		[
+			PathArrayWidget->AsShared()
+		];
+
+	FieldWidget->AddSlot()
+		.Padding(FMargin(3.f, 3.f))
+		.AutoHeight()
+		[
+			SNew(SHorizontalBox)
+			+ SHorizontalBox::Slot()
+				.AutoWidth()
+				[
+					SNew(STextBlock)
+					.Text(FText::FromString("Preview Path: "))
+				]
+			+ SHorizontalBox::Slot()
+				[
+					PreviewPathWidget->AsShared()
+				]
+		];
+
+	
+	return 	SNew(SHorizontalBox)
+		.Clipping(EWidgetClipping::OnDemand)
+		+ SHorizontalBox::Slot()
+		.VAlign(VAlign_Center)
+		.AutoWidth()
+		.FillWidth(1.0f)
+		[
+			FieldWidget
+		];
+}
+
+void FRCSetAssetByPathBehaviourModel::RegenerateWeakPtrInternal()
+{
+	DetailTreeNodeWeakPtrArray.Empty();
+	for (const TSharedRef<IDetailTreeNode>& CategoryNode : PropertyRowGeneratorArray->GetRootTreeNodes())
+	{
+		TArray<TSharedRef<IDetailTreeNode>> Children;
+		CategoryNode->GetChildren(Children);
+		for (const TSharedRef<IDetailTreeNode>& Child : Children)
+		{
+			DetailTreeNodeWeakPtrArray.Add(Child);
+		}
+	}
+
+	URCSetAssetByPathBehaviour* SetAssetByPathBehaviour = SetAssetByPathBehaviourWeakPtr.Get();
+	if (!SetAssetByPathBehaviour)
+	{
+		return;
+	}
+
+	URCController* RCController = SetAssetByPathBehaviour->ControllerWeakPtr.Get();
+	if (!RCController)
+	{
+		return;
+	}
+
+	URemoteControlPreset* RemoteControlPreset =  RCController->PresetWeakPtr.Get();
+	if (!RemoteControlPreset)
+	{
+		return;
+	}
+
+	ExposedEntities = RemoteControlPreset->GetExposedEntities<const FRemoteControlEntity>();
+}
+
+void FRCSetAssetByPathBehaviourModel::RefreshPathAndPreview()
+{
+	TSharedPtr<SVerticalBox> FieldArrayWidget = SNew(SVerticalBox);
+
+	RegenerateWeakPtrInternal();
+	
+	for (const TSharedPtr<IDetailTreeNode>& DetailTreeNodeArray : DetailTreeNodeWeakPtrArray)
+	{
+		const TSharedPtr<IDetailTreeNode> PinnedNode = DetailTreeNodeArray;
+
+		TArray<TSharedRef<IDetailTreeNode>> Children;
+		PinnedNode->GetChildren(Children);
+		Children.Insert(PinnedNode.ToSharedRef(), 0);
+		
+		for (uint8 Counter = 0; Counter < Children.Num(); Counter++)
+		{
+			TSharedPtr<SWidget> Button = SNullWidget::NullWidget;
+			if (Counter == 0)
+			{
+				Button = SNew(SBox)
+					.HAlign(HAlign_Center)
+					.VAlign(VAlign_Center)
+					.WidthOverride(22)
+					.HeightOverride(22)
+					[
+						SNew(SButton)
+						.ButtonStyle( FAppStyle::Get(), "SimpleButton" )
+						.OnClicked_Lambda([this]()
+						{
+							if (URCSetAssetByPathBehaviour* Behaviour = Cast<URCSetAssetByPathBehaviour>(GetBehaviour()))
+							{
+								Behaviour->Modify();
+								CreateAssetPathFromSelection();
+								RefreshPathAndPreview();
+							}
+
+							return FReply::Handled();
+						})
+						.IsFocusable(false)
+						.ContentPadding(0)
+						.Visibility_Lambda([this]()
+						{
+							URCSetAssetByPathBehaviour* Behaviour = Cast<URCSetAssetByPathBehaviour>(GetBehaviour());
+							if (!Behaviour)
+							{
+								return EVisibility::Collapsed;
+							}
+
+							return Behaviour->bInternal ? EVisibility::Visible : EVisibility::Collapsed;
+						})
+						.Content()
+						[
+							SNew(SImage)
+							.Image( FAppStyle::GetBrush("Icons.Use") )
+							.ColorAndOpacity( FSlateColor::UseForeground() )
+						]
+					];
+			}
+			
+			const FNodeWidgets NodeWidgets = Children[Counter]->CreateNodeWidgets();
+			if (NodeWidgets.ValueWidget)
+			{
+				FieldArrayWidget->AddSlot()
+					.Padding(FMargin(3.0f, 2.0f))
+					.AutoHeight()
+					[
+						SNew(SHorizontalBox)
+						+ SHorizontalBox::Slot()
+						[
+							NodeWidgets.NameWidget.ToSharedRef()
+						]
+						+ SHorizontalBox::Slot()
+						[
+							NodeWidgets.ValueWidget.ToSharedRef()
+						]
+						+ SHorizontalBox::Slot()
+						.AutoWidth()
+						[
+							Button.ToSharedRef()
+						]
+					];
+			}
+			else if (NodeWidgets.WholeRowWidget)
+			{
+				FieldArrayWidget->AddSlot()
+					.Padding(FMargin(3.0f, 2.0f))
+					.AutoHeight()
+					[
+						NodeWidgets.WholeRowWidget.ToSharedRef()
+					];
+			}
+		}
+	}
+
+	FString PreviewText = TEXT("Path is invalid!");
+	if (URCSetAssetByPathBehaviour* SetAssetByPathBehaviour = SetAssetByPathBehaviourWeakPtr.Get())
+	{
+		PreviewText = SetAssetByPathBehaviour->GetCurrentPath();
+	}
+	
+	PathArrayWidget->SetContent(FieldArrayWidget.ToSharedRef());
+	PreviewPathWidget->SetText(FText::FromString(PreviewText));
+}
+
+TSharedRef<SWidget> FRCSetAssetByPathBehaviourModel::GetSelectorWidget(TWeakPtr<const FRemoteControlEntity> InInitialSelected)
+{
+	return SNew(SComboBox<TWeakPtr<const FRemoteControlEntity>>)
+		.OptionsSource(&ExposedEntities)
+		.OnGenerateWidget_Lambda([this](TWeakPtr<const FRemoteControlEntity> InItem) -> TSharedRef<SWidget>
+		{
+			TSharedPtr<const FRemoteControlEntity> PinnedItem = InItem.Pin();
+			if (!PinnedItem)
+			{
+				return SNullWidget::NullWidget;
+			}
+
+			return SNew(STextBlock).Text(FText::FromString(*PinnedItem->GetLabel().ToString()));
+		})
+		.OnComboBoxOpening_Lambda([this]()
+		{
+			const URCSetAssetByPathBehaviour* PathBehaviour = Cast<URCSetAssetByPathBehaviour>(GetBehaviour());
+			if (!PathBehaviour)
+			{
+				return;
+			}
+			
+			const URCController* Controller = PathBehaviour->ControllerWeakPtr.Get();
+			if (!Controller)
+			{
+				return;
+			}
+
+			const URemoteControlPreset* Preset = Controller->PresetWeakPtr.Get();
+			if (!Preset)
+			{
+				return;
+			}
+			
+			ExposedEntities = Preset->GetExposedEntities<const FRemoteControlEntity>();
+		})
+		.OnSelectionChanged_Lambda([this](TWeakPtr<const FRemoteControlEntity> InItem, ESelectInfo::Type SelectType)
+		{
+			if (TSharedPtr<const FRemoteControlEntity> Entity = InItem.Pin())
+			{
+				if (URCSetAssetByPathBehaviour* PathBehaviour = Cast<URCSetAssetByPathBehaviour>(GetBehaviour()))
+				{
+					PathBehaviour->SetTargetEntity(Entity);
+				}
+			}
+		})
+		.InitiallySelectedItem(InInitialSelected)
+		.Content()
+		[
+			SNew(STextBlock)
+				.Text_Raw(this, &FRCSetAssetByPathBehaviourModel::GetSelectedEntityText)
+		];
+}
+
+FText FRCSetAssetByPathBehaviourModel::GetSelectedEntityText() const
+{
+	FText ReturnText = FText::FromString(*FString("Nothing selected yet"));
+
+	if (URCSetAssetByPathBehaviour* Behaviour = SetAssetByPathBehaviourWeakPtr.Get())
+	{
+		if (Behaviour->GetTargetEntity().Pin())
+		{
+			ReturnText = FText::FromString(*Behaviour->GetTargetEntity().Pin()->GetLabel().ToString());
+		}
+	}
+	return ReturnText;
+}
+
+
+void FRCSetAssetByPathBehaviourModel::CreateAssetPathFromSelection()
+{
+	URCSetAssetByPathBehaviour* SetAssetByPathBehaviour = Cast<URCSetAssetByPathBehaviour>(GetBehaviour());
+	if (!SetAssetByPathBehaviour)
+	{
+		return;
+	}
+	
+	// Clear it, in case it is an already used one.
+	TArray<FString>& PathArray = SetAssetByPathBehaviour->PathStruct.PathArray;
+	PathArray.Empty();
+	
+	TArray<FAssetData> AssetData;
+	GEditor->GetContentBrowserSelections(AssetData);
+	if (AssetData.Num() != 1)
+	{
+		// Only works with one selected Content
+		return;
+	}
+
+	int32 IndexOfLast;
+	const UObject* SelectedAsset = AssetData[0].GetAsset();
+	FString PathString = SelectedAsset->GetPathName();
+
+	// Remove the initial Game
+	PathString.RemoveFromStart("/Game/");
+
+	// Remove anything after the last /
+	PathString.FindLastChar('/', IndexOfLast);
+	PathString.RemoveAt(IndexOfLast, PathString.Len() - IndexOfLast);
+
+	PathArray.Add(PathString);
+}
+
+
+#undef LOCTEXT_NAMESPACE
diff --git a/Engine/Plugins/VirtualProduction/RemoteControl/Source/RemoteControlUI/Private/UI/Behaviour/Builtin/Path/RCBehaviourSetAssetByPathModel.h b/Engine/Plugins/VirtualProduction/RemoteControl/Source/RemoteControlUI/Private/UI/Behaviour/Builtin/Path/RCBehaviourSetAssetByPathModel.h
new file mode 100644
--- /dev/null
+++ b/Engine/Plugins/VirtualProduction/RemoteControl/Source/RemoteControlUI/Private/UI/Behaviour/Builtin/Path/RCBehaviourSetAssetByPathModel.h
@@ -0,0 +1,82 @@
+// Copyright Epic Games, Inc. All Rights Reserved.
+
+#pragma once
+#include "Behaviour/Builtin/Path/RCSetAssetByPathBehaviour.h"
+#include "RemoteControlEntity.h"
+#include "UI/Behaviour/RCBehaviourModel.h"
+
+class IDetailTreeNode;
+class SBox;
+
+namespace SetAssetByPathModelHelper
+{
+	const FString InputToken = FString(TEXT("*INPUT "));
+}
+
+/*
+ * ~ FRCSetAssetByPathModel ~
+ *
+ * Child Behaviour class representing the "Set Asset By Path" Behaviour's UI model.
+ *
+ * Generates several Widgets where users can enter the RootPath, Target Property, and Default Property as FStrings.
+ * The values are then put together, getting a path towards a possible object.
+ */
+class FRCSetAssetByPathBehaviourModel : public FRCBehaviourModel
+{
+public:
+	FRCSetAssetByPathBehaviourModel(URCSetAssetByPathBehaviour* SetAssetByPathBehaviour);
+
+	/** Builds a Behaviour specific widget as required for the Set Asset By Path Behaviour */
+	virtual TSharedRef<SWidget> GetBehaviourDetailsWidget() override;
+
+	/*
+	 * Builds the Property Details Widget, including a generic expendable Array Widget and
+	 * further two Text Widgets representing elements needed for the Path Behaviour.
+	 * All of them store the user input and use them to perform the SetAssetByPath Behaviour.
+	 */
+	TSharedRef<SWidget> GetPropertyWidget();
+
+	/** Builds the Combobox which the Behaviour will use to select and and apply it's Asset Setting on. */
+	TSharedRef<SWidget> GetSelectorWidget(TWeakPtr<const FRemoteControlEntity> InInitialSelected);
+
+protected:
+	/** Gets the currently selected Entities Name, or a failure Text in case it fails. */
+	FText GetSelectedEntityText() const;
+	
+private:
+	/** The SetAssetByPath Behaviour associated with this Model */
+	TWeakObjectPtr<URCSetAssetByPathBehaviour> SetAssetByPathBehaviourWeakPtr;
+
+	/** Pointer to the Widget holding the PathArray created for this behaviour */
+	TSharedPtr<SBox> PathArrayWidget;
+
+	/** Pointer to the Preview Text Widget for the Path */
+	TSharedPtr<STextBlock> PreviewPathWidget;
+
+	/** POinter to the ComboBox Widget */
+	TSharedPtr<SWidget> SelectorBox;
+
+	/** The row generator used to build the generic value widgets */
+	TSharedPtr<IPropertyRowGenerator> PropertyRowGenerator;
+	
+	/** The row generator used to build the generic array value widget for the Path Array */
+	TSharedPtr<IPropertyRowGenerator> PropertyRowGeneratorArray;
+
+	/** Used to create a generic Value Widget based on the active Controller's type*/
+	TArray<TSharedPtr<IDetailTreeNode>> DetailTreeNodeWeakPtr;
+	
+	/** Used to create a generic Value Widget based on the Paths Available*/
+	TArray<TSharedPtr<IDetailTreeNode>> DetailTreeNodeWeakPtrArray;
+	
+	/** Array of Exposed Entities from the Preset to help create a Widget to choose the Target */
+	TArray<TWeakPtr<const FRemoteControlEntity>> ExposedEntities;
+	
+private:
+	void RegenerateWeakPtrInternal();
+	
+	/** Regenerates and creates a new PathArray Widget if changed */
+	void RefreshPathAndPreview();
+
+	/** Creates a asset path, based on the selected Asset in the content browser. */
+	void CreateAssetPathFromSelection();
+};
diff --git a/Engine/Plugins/VirtualProduction/RemoteControl/Source/RemoteControlUI/Private/UI/Behaviour/Builtin/Path/SRCBehaviourSetAssetByPath.cpp b/Engine/Plugins/VirtualProduction/RemoteControl/Source/RemoteControlUI/Private/UI/Behaviour/Builtin/Path/SRCBehaviourSetAssetByPath.cpp
new file mode 100644
--- /dev/null
+++ b/Engine/Plugins/VirtualProduction/RemoteControl/Source/RemoteControlUI/Private/UI/Behaviour/Builtin/Path/SRCBehaviourSetAssetByPath.cpp
@@ -0,0 +1,124 @@
+// Copyright Epic Games, Inc. All Rights Reserved.
+
+#include "SRCBehaviourSetAssetByPath.h"
+
+#include "IDetailTreeNode.h"
+#include "RemoteControlPreset.h"
+#include "SPositiveActionButton.h"
+#include "SlateOptMacros.h"
+#include "Styling/StyleColors.h"
+#include "UI/Behaviour/Builtin/Path/RCBehaviourSetAssetByPathModel.h"
+#include "UI/RemoteControlPanelStyle.h"
+#include "Widgets/Input/SButton.h"
+
+#define LOCTEXT_NAMESPACE "SRCBehaviourSetAssetByPath"
+
+class SPositiveActionButton;
+
+BEGIN_SLATE_FUNCTION_BUILD_OPTIMIZATION
+
+void SRCBehaviourSetAssetByPath::Construct(const FArguments& InArgs, TSharedRef<FRCSetAssetByPathBehaviourModel> InBehaviourItem)
+{
+	URCSetAssetByPathBehaviour* PathBehaviour = Cast<URCSetAssetByPathBehaviour>(InBehaviourItem->GetBehaviour());
+	TSharedRef<SHorizontalBox> InternalExternalSwitchWidget = SNew(SHorizontalBox);
+	TSharedRef<SHorizontalBox> SelectedClassWidget = SNew(SHorizontalBox)
+		.Visibility_Lambda([PathBehaviour]()
+		{
+			return PathBehaviour->bInternal ? EVisibility::Visible : EVisibility::Collapsed;
+		});
+
+	if (!PathBehaviour)
+	{
+		ChildSlot
+		[
+			SNullWidget::NullWidget
+		];
+		return;
+	}
+		
+	InternalExternalSwitchWidget->AddSlot()
+	.HAlign(HAlign_Right)
+	.VAlign(VAlign_Center)
+	.Padding(FMargin(0.0f,0.0f))
+	.FillWidth(1.0f)
+	.AutoWidth()
+	[
+		SNew(SButton)
+		.Text(FText::FromString(FString("Internal")))
+		.ButtonColorAndOpacity_Lambda([PathBehaviour]()
+		{
+			FSlateColor Color = PathBehaviour->bInternal ? FAppStyle::Get().GetSlateColor("Colors.Highlight") : FAppStyle::Get().GetSlateColor("Colors.AccentWhite");
+			return Color;
+		})
+		.OnPressed_Lambda([PathBehaviour]()
+		{
+			PathBehaviour->bInternal = true;
+		})
+	];
+	
+	InternalExternalSwitchWidget->AddSlot()
+	.HAlign(HAlign_Right)
+	.VAlign(VAlign_Center)
+	.Padding(FMargin(0.0f,0.0f))
+	.FillWidth(1.0f)
+	.AutoWidth()
+	[
+		SNew(SButton)
+		.Text(FText::FromString(FString("External")))
+		.ButtonColorAndOpacity_Lambda([PathBehaviour]()
+		{
+			FSlateColor Color = PathBehaviour->bInternal ? FAppStyle::Get().GetSlateColor("Colors.AccentWhite") : FAppStyle::Get().GetSlateColor("Colors.Highlight");
+			return Color;
+		})
+		.OnPressed_Lambda([PathBehaviour]()
+		{
+			PathBehaviour->bInternal = false;
+		})
+	];
+	
+	for (UClass* SupportedClass : PathBehaviour->GetSupportedClasses())
+	{
+		SelectedClassWidget->AddSlot()
+		.HAlign(HAlign_Right)
+		.VAlign(VAlign_Center)
+		.Padding(FMargin(0.0f,0.0f))
+		.FillWidth(1.0f)
+		.AutoWidth()
+		[
+			SNew(SButton)
+			.Text(FText::FromString(SupportedClass->GetName()))
+			.ButtonColorAndOpacity_Lambda([PathBehaviour, SupportedClass]()
+			{
+				FSlateColor Color = PathBehaviour->AssetClass == SupportedClass ? FAppStyle::Get().GetSlateColor("Colors.Select") : FAppStyle::Get().GetSlateColor("Colors.White");
+				return Color;
+			})
+			.OnPressed_Lambda([PathBehaviour, SupportedClass]()
+			{
+				PathBehaviour->AssetClass = SupportedClass;
+			})
+		];
+	}
+	
+	ChildSlot
+	[
+		SNew(SVerticalBox)
+		+ SVerticalBox::Slot()
+		[
+			InternalExternalSwitchWidget
+		]
+		+ SVerticalBox::Slot()
+		.Padding(0.f,12.f)
+		[
+			SelectedClassWidget
+		]
+		+ SVerticalBox::Slot()
+		.AutoHeight()
+		[
+			InBehaviourItem->GetPropertyWidget()
+		]
+	];
+}
+
+END_SLATE_FUNCTION_BUILD_OPTIMIZATION
+
+#undef LOCTEXT_NAMESPACE
diff --git a/Engine/Plugins/VirtualProduction/RemoteControl/Source/RemoteControlUI/Private/UI/Behaviour/Builtin/Path/SRCBehaviourSetAssetByPath.h b/Engine/Plugins/VirtualProduction/RemoteControl/Source/RemoteControlUI/Private/UI/Behaviour/Builtin/Path/SRCBehaviourSetAssetByPath.h
new file mode 100644
--- /dev/null
+++ b/Engine/Plugins/VirtualProduction/RemoteControl/Source/RemoteControlUI/Private/UI/Behaviour/Builtin/Path/SRCBehaviourSetAssetByPath.h
@@ -0,0 +1,25 @@
+// Copyright Epic Games, Inc. All Rights Reserved.
+
+#pragma once
+
+#include "CoreMinimal.h"
+#include "Widgets/SCompoundWidget.h"
+
+class FRCSetAssetByPathBehaviourModel;
+
+class REMOTECONTROLUI_API SRCBehaviourSetAssetByPath : public  SCompoundWidget
+{
+public:
+	SLATE_BEGIN_ARGS(SRCBehaviourSetAssetByPath)
+		{
+		}
+
+	SLATE_END_ARGS()
+
+	/** Constructs this widget with InArgs */
+	void Construct(const FArguments& InArgs, TSharedRef<FRCSetAssetByPathBehaviourModel> InBehaviourItem);
+
+private:
+	/** The Behaviour (UI model) associated with us*/
+	TWeakPtr<const FRCSetAssetByPathBehaviourModel> SetAssetByPathWeakPtr;
+};
diff --git a/Engine/Plugins/VirtualProduction/RemoteControl/Source/RemoteControlUI/Private/UI/Behaviour/Builtin/RCBehaviourSetAssetByPathModel.cpp b/Engine/Plugins/VirtualProduction/RemoteControl/Source/RemoteControlUI/Private/UI/Behaviour/Builtin/RCBehaviourSetAssetByPathModel.cpp
deleted file mode 100644
diff --git a/Engine/Plugins/VirtualProduction/RemoteControl/Source/RemoteControlUI/Private/UI/Behaviour/Builtin/RCBehaviourSetAssetByPathModel.h b/Engine/Plugins/VirtualProduction/RemoteControl/Source/RemoteControlUI/Private/UI/Behaviour/Builtin/RCBehaviourSetAssetByPathModel.h
deleted file mode 100644
diff --git a/Engine/Plugins/VirtualProduction/RemoteControl/Source/RemoteControlUI/Private/UI/Behaviour/Builtin/SRCBehaviourSetAssetByPath.cpp b/Engine/Plugins/VirtualProduction/RemoteControl/Source/RemoteControlUI/Private/UI/Behaviour/Builtin/SRCBehaviourSetAssetByPath.cpp
deleted file mode 100644
diff --git a/Engine/Plugins/VirtualProduction/RemoteControl/Source/RemoteControlUI/Private/UI/Behaviour/Builtin/SRCBehaviourSetAssetByPath.h b/Engine/Plugins/VirtualProduction/RemoteControl/Source/RemoteControlUI/Private/UI/Behaviour/Builtin/SRCBehaviourSetAssetByPath.h
deleted file mode 100644
diff --git a/Engine/Plugins/VirtualProduction/RemoteControl/Source/RemoteControlUI/Private/UI/Behaviour/SRCBehaviourPanelList.cpp b/Engine/Plugins/VirtualProduction/RemoteControl/Source/RemoteControlUI/Private/UI/Behaviour/SRCBehaviourPanelList.cpp
--- a/Engine/Plugins/VirtualProduction/RemoteControl/Source/RemoteControlUI/Private/UI/Behaviour/SRCBehaviourPanelList.cpp
+++ b/Engine/Plugins/VirtualProduction/RemoteControl/Source/RemoteControlUI/Private/UI/Behaviour/SRCBehaviourPanelList.cpp
@@ -4,7 +4,7 @@
 
 #include "Behaviour/Builtin/Bind/RCBehaviourBind.h"
 #include "Behaviour/Builtin/Conditional/RCBehaviourConditional.h"
-#include "Behaviour/RCSetAssetByPathBehaviour.h"
+#include "Behaviour/Builtin/Path/RCSetAssetByPathBehaviour.h"
 #include "Controller/RCController.h"
 #include "RCBehaviourModel.h"
 #include "RemoteControlPreset.h"
@@ -14,7 +14,7 @@
 #include "UI/Behaviour/Builtin/Bind/RCBehaviourBindModel.h"
 #include "UI/Behaviour/Builtin/Conditional/RCBehaviourConditionalModel.h"
 #include "UI/Behaviour/Builtin/RCBehaviourIsEqualModel.h"
-#include "UI/Behaviour/Builtin/RCBehaviourSetAssetByPathModel.h"
+#include "UI/Behaviour/Builtin/Path/RCBehaviourSetAssetByPathModel.h"
 #include "UI/Controller/RCControllerModel.h"
 #include "UI/Behaviour/Builtin/RangeMap/RCBehaviourRangeMapModel.h"
 #include "UI/RemoteControlPanelStyle.h"
