Remote Control - Range Map Behaviour Fix it once again

Change overseen Conditional.

diff --git a/Engine/Plugins/VirtualProduction/RemoteControl/Source/RemoteControlLogic/Public/Behaviour/Builtin/RangeMap/RCRangeMapBehaviour.h b/Engine/Plugins/VirtualProduction/RemoteControl/Source/RemoteControlLogic/Public/Behaviour/Builtin/RangeMap/RCRangeMapBehaviour.h
--- a/Engine/Plugins/VirtualProduction/RemoteControl/Source/RemoteControlLogic/Public/Behaviour/Builtin/RangeMap/RCRangeMapBehaviour.h
+++ b/Engine/Plugins/VirtualProduction/RemoteControl/Source/RemoteControlLogic/Public/Behaviour/Builtin/RangeMap/RCRangeMapBehaviour.h
@@ -18,7 +18,7 @@ struct FRCRangeMapStep
 	FRCRangeMapStep(const double InStepValue, URCVirtualPropertySelfContainer* InPropertyValue)
 		: StepValue(InStepValue)
 	{
-		if (!PropertyValue)
+		if (!InPropertyValue)
 		{
 			return;
 		}
