[EpicStageApp] Fix stage app failing to control nDisplay nodes when connected directly

diff --git a/Engine/Plugins/Runtime/nDisplay/nDisplay.uplugin b/Engine/Plugins/Runtime/nDisplay/nDisplay.uplugin
--- a/Engine/Plugins/Runtime/nDisplay/nDisplay.uplugin
+++ b/Engine/Plugins/Runtime/nDisplay/nDisplay.uplugin
@@ -110,7 +110,7 @@
 		},
 		{
 			"Name": "DisplayClusterScenePreview",
-			"Type": "Editor",
+			"Type": "Runtime",
 			"LoadingPhase": "Default",
 			"PlatformAllowList": [
 				"Win64",
diff --git a/Engine/Plugins/VirtualProduction/RemoteControl/Source/WebRemoteControl/Private/WebSocketMessageHandler.cpp b/Engine/Plugins/VirtualProduction/RemoteControl/Source/WebRemoteControl/Private/WebSocketMessageHandler.cpp
--- a/Engine/Plugins/VirtualProduction/RemoteControl/Source/WebRemoteControl/Private/WebSocketMessageHandler.cpp
+++ b/Engine/Plugins/VirtualProduction/RemoteControl/Source/WebRemoteControl/Private/WebSocketMessageHandler.cpp
@@ -585,8 +585,8 @@ void FWebSocketMessageHandler::HandleWebSocketActorRegister(const FRemoteControl
 		UWorld* World = WorldContext.World();
 		
 		// To support non-level editor worlds, instead of just checking the level editor's world,
-		// now every Editor-type world is checked.
-		if (!World || WorldContext.WorldType != EWorldType::Editor)
+		// now every Editor-type or Game-type world is checked.
+		if (!World || (WorldContext.WorldType != EWorldType::Editor && WorldContext.WorldType != EWorldType::Game))
 		{
 			continue;
 		}
@@ -642,11 +642,13 @@ void FWebSocketMessageHandler::HandleWebSocketPresetModifyProperty(const FRemote
 	case ERCTransactionMode::MANUAL:
 		Access = ERCAccess::WRITE_MANUAL_TRANSACTION_ACCESS;
 
+#if WITH_EDITOR
 		// Indicate that we want to contribute to this transaction if it's active
-		if (!ContributeToTransaction(WebSocketMessage.ClientId, Body.TransactionId))
+		if (GEditor && !ContributeToTransaction(WebSocketMessage.ClientId, Body.TransactionId))
 		{
 			return;
 		}
+#endif
 
 		break;
 
@@ -724,14 +726,16 @@ void FWebSocketMessageHandler::HandleWebSocketFunctionCall(const FRemoteControlW
 		return;
 	}
 
+#if WITH_EDITOR
 	if (Body.TransactionMode == ERCTransactionMode::MANUAL)
 	{
 		// Indicate that we want to contribute to this transaction
-		if (!ContributeToTransaction(WebSocketMessage.ClientId, Body.TransactionId))
+		if (GEditor && !ContributeToTransaction(WebSocketMessage.ClientId, Body.TransactionId))
 		{
 			return;
 		}
 	}
+#endif
 
 	UpdateSequenceNumber(WebSocketMessage.ClientId, Body.SequenceNumber);
 
@@ -741,6 +745,11 @@ void FWebSocketMessageHandler::HandleWebSocketFunctionCall(const FRemoteControlW
 void FWebSocketMessageHandler::HandleWebSocketBeginEditorTransaction(const FRemoteControlWebSocketMessage& WebSocketMessage)
 {
 #if WITH_EDITOR
+	if (!GEditor)
+	{
+		return;
+	}
+
 	FRCWebSocketTransactionStartBody Body;
 	if (!WebRemoteControlInternalUtils::DeserializeRequestPayload(WebSocketMessage.RequestPayload, nullptr, Body))
 	{
@@ -774,6 +783,11 @@ void FWebSocketMessageHandler::HandleWebSocketBeginEditorTransaction(const FRemo
 void FWebSocketMessageHandler::HandleWebSocketEndEditorTransaction(const FRemoteControlWebSocketMessage& WebSocketMessage)
 {
 #if WITH_EDITOR
+	if (!GEditor)
+	{
+		return;
+	}
+
 	FRCWebSocketTransactionEndBody Body;
 	if (!WebRemoteControlInternalUtils::DeserializeRequestPayload(WebSocketMessage.RequestPayload, nullptr, Body))
 	{
@@ -1111,7 +1125,9 @@ void FWebSocketMessageHandler::OnEndFrame()
 
 	if (PropertyNotificationFrameCounter >= CVarWebRemoteControlFramesBetweenPropertyNotifications.GetValueOnGameThread())
 	{
+#if WITH_EDITOR
 		TimeOutTransactions();
+#endif
 
 		//Early exit if no clients are requesting notifications
 		if (PresetNotificationMap.Num() <= 0 && ActorNotificationMap.Num() <= 0)
@@ -1131,6 +1147,7 @@ void FWebSocketMessageHandler::OnEndFrame()
 	}
 }
 
+#if WITH_EDITOR
 void FWebSocketMessageHandler::TimeOutTransactions()
 {
 	TArray<TPair<FGuid, int32>> TransactionsToEnd;
@@ -1154,6 +1171,7 @@ void FWebSocketMessageHandler::TimeOutTransactions()
 		EndClientTransaction(TransactionIterator.Key, TransactionIterator.Value);
 	}
 }
+#endif
 
 void FWebSocketMessageHandler::ProcessAddedProperties()
 {
@@ -1816,7 +1834,6 @@ void FWebSocketMessageHandler::HandleTransactionStateChanged(const FTransactionC
 		HandleTransactionEnded(InTransactionContext.TransactionId);
 	}
 }
-#endif
 
 void FWebSocketMessageHandler::HandleTransactionEnded(const FGuid& TransactionGuid)
 {
@@ -1846,6 +1863,7 @@ void FWebSocketMessageHandler::HandleTransactionEnded(const FGuid& TransactionGu
 		}
 	}
 }
+#endif
 
 void FWebSocketMessageHandler::StartWatchingActor(AActor* Actor, UClass* WatchedClass)
 {
@@ -1933,8 +1951,21 @@ void FWebSocketMessageHandler::UpdateSequenceNumber(const FGuid& ClientId, int64
 	}
 }
 
+int64 FWebSocketMessageHandler::GetSequenceNumber(const FGuid& ClientId) const
+{
+	const int64* ClientSequenceNumber = ClientSequenceNumbers.Find(ClientId);
+	return ClientSequenceNumber ? *ClientSequenceNumber : DefaultSequenceNumber;
+}
+
+#if WITH_EDITOR
+
 bool FWebSocketMessageHandler::ContributeToTransaction(const FGuid& ClientId, int32 TransactionId)
 {
+	if (!GEditor)
+	{
+		return false;
+	}
+
 	const FGuid InternalId = GetTransactionGuid(ClientId, TransactionId);
 	if (ClientId.IsValid())
 	{
@@ -1950,6 +1981,11 @@ bool FWebSocketMessageHandler::ContributeToTransaction(const FGuid& ClientId, in
 
 void FWebSocketMessageHandler::EndClientTransaction(const FGuid& ClientId, int32 TransactionId)
 {
+	if (!GEditor)
+	{
+		return;
+	}
+
 	if (TMap<FGuid, int32>* TransactionIds = TransactionIdsByClientId.Find(ClientId))
 	{
 		const FGuid TransactionGuid = GetTransactionGuid(ClientId, TransactionId);
@@ -1980,14 +2016,13 @@ void FWebSocketMessageHandler::EndClientTransaction(const FGuid& ClientId, int32
 	}
 }
 
-int64 FWebSocketMessageHandler::GetSequenceNumber(const FGuid& ClientId) const
-{
-	const int64* ClientSequenceNumber = ClientSequenceNumbers.Find(ClientId);
-	return ClientSequenceNumber ? *ClientSequenceNumber : DefaultSequenceNumber;
-}
-
 FGuid FWebSocketMessageHandler::GetTransactionGuid(const FGuid& ClientId, int32 TransactionId) const
 {
+	if (!GEditor)
+	{
+		return FGuid();
+	}
+
 	if (const TMap<FGuid, int32>* TransactionIdPairs = TransactionIdsByClientId.Find(ClientId))
 	{
 		for (const TPair<FGuid, int32>& Pair : *TransactionIdPairs)
@@ -2004,6 +2039,11 @@ FGuid FWebSocketMessageHandler::GetTransactionGuid(const FGuid& ClientId, int32
 
 int32 FWebSocketMessageHandler::GetClientTransactionId(const FGuid& ClientId, const FGuid& TransactionGuid) const
 {
+	if (!GEditor)
+	{
+		return InvalidTransactionId;
+	}
+
 	if (const TMap<FGuid, int32>* TransactionIds = TransactionIdsByClientId.Find(ClientId))
 	{
 		if (const int32* ClientTransactionId = TransactionIds->Find(TransactionGuid))
@@ -2015,4 +2055,6 @@ int32 FWebSocketMessageHandler::GetClientTransactionId(const FGuid& ClientId, co
 	return InvalidTransactionId;
 }
 
+#endif
+
 #undef LOCTEXT_NAMESPACE
diff --git a/Engine/Plugins/VirtualProduction/RemoteControl/Source/WebRemoteControl/Private/WebSocketMessageHandler.h b/Engine/Plugins/VirtualProduction/RemoteControl/Source/WebRemoteControl/Private/WebSocketMessageHandler.h
--- a/Engine/Plugins/VirtualProduction/RemoteControl/Source/WebRemoteControl/Private/WebSocketMessageHandler.h
+++ b/Engine/Plugins/VirtualProduction/RemoteControl/Source/WebRemoteControl/Private/WebSocketMessageHandler.h
@@ -267,6 +267,7 @@ private:
 	 */
 	int64 GetSequenceNumber(const FGuid& ClientId) const;
 
+#if WITH_EDITOR
 	/**
 	 * Indicate that a client is going to contribute to the transaction with the given ID.
 	 * This should be called whenever a change is made that will be part of the transaction, not just at the start.
@@ -288,6 +289,7 @@ private:
 	 * Converts an internal transaction GUID to the ID used by the given client to refer to it.
 	 */
 	int32 GetClientTransactionId(const FGuid& ClientId, const FGuid& TransactionGuid) const;
+#endif
 
 private:
 
diff --git a/Engine/Plugins/VirtualProduction/RemoteControlInterception/Source/RemoteControlInterception/Public/IRemoteControlInterceptionCommands.h b/Engine/Plugins/VirtualProduction/RemoteControlInterception/Source/RemoteControlInterception/Public/IRemoteControlInterceptionCommands.h
--- a/Engine/Plugins/VirtualProduction/RemoteControlInterception/Source/RemoteControlInterception/Public/IRemoteControlInterceptionCommands.h
+++ b/Engine/Plugins/VirtualProduction/RemoteControlInterception/Source/RemoteControlInterception/Public/IRemoteControlInterceptionCommands.h
@@ -214,6 +214,7 @@ public:
 		Ar << Super;
 		
 		Ar << Instance.PayloadType;
+		Ar << Instance.Operation;
 		// const cast is needed for keeping the const& on a constructor and as a class parameter
 		Ar << const_cast<TArray<uint8>&>(Instance.Payload);
 		return Ar;
