Undo/Redo support for RC Logic

Currently covers the following usecases:
* Add Controller, Delete Controller, Delete All Controllers
* Add Behaviour Delete Behaviour, Delete All Behaviours
* Add Action, Delete Action, Delete All Actions

diff --git a/Engine/Plugins/VirtualProduction/RemoteControl/Source/RemoteControl/Private/RCVirtualProperty.cpp b/Engine/Plugins/VirtualProduction/RemoteControl/Source/RemoteControl/Private/RCVirtualProperty.cpp
--- a/Engine/Plugins/VirtualProduction/RemoteControl/Source/RemoteControl/Private/RCVirtualProperty.cpp
+++ b/Engine/Plugins/VirtualProduction/RemoteControl/Source/RemoteControl/Private/RCVirtualProperty.cpp
@@ -648,6 +648,11 @@ FName URCVirtualPropertyBase::GetVirtualPropertyTypeDisplayName(const EPropertyB
 	return *PropertyTypeString;
 }
 
+void URCVirtualPropertyBase::SetDisplayIndex(const int32 InDisplayIndex)
+{
+	DisplayIndex = InDisplayIndex;
+}
+
 const FInstancedPropertyBag* URCVirtualPropertyInContainer::GetPropertyBagInstance() const
 {
 	if (ContainerWeakPtr.Get())
diff --git a/Engine/Plugins/VirtualProduction/RemoteControl/Source/RemoteControl/Private/RCVirtualPropertyContainer.cpp b/Engine/Plugins/VirtualProduction/RemoteControl/Source/RemoteControl/Private/RCVirtualPropertyContainer.cpp
--- a/Engine/Plugins/VirtualProduction/RemoteControl/Source/RemoteControl/Private/RCVirtualPropertyContainer.cpp
+++ b/Engine/Plugins/VirtualProduction/RemoteControl/Source/RemoteControl/Private/RCVirtualPropertyContainer.cpp
@@ -15,14 +15,6 @@
 
 URCVirtualPropertyInContainer* URCVirtualPropertyContainerBase::AddProperty(const FName& InPropertyName, TSubclassOf<URCVirtualPropertyInContainer> InPropertyClass, const EPropertyBagPropertyType InValueType, UObject* InValueTypeObject)
 {
-#if WITH_EDITOR
-	const FScopedTransaction Transaction(LOCTEXT("AddProperty", "AddP roperty"));
-
-	Modify();
-
-	MarkPackageDirty();
-#endif
-
 	const FName PropertyName = GenerateUniquePropertyName(InPropertyName, InValueType, InValueTypeObject, this);
 	Bag.AddProperty(PropertyName, InValueType, InValueTypeObject);
 
@@ -33,7 +25,7 @@ URCVirtualPropertyInContainer* URCVirtualPropertyContainerBase::AddProperty(cons
 	}
 	
 	// Create Property in Container
-	URCVirtualPropertyInContainer* VirtualPropertyInContainer = NewObject<URCVirtualPropertyInContainer>(this, InPropertyClass.Get());
+	URCVirtualPropertyInContainer* VirtualPropertyInContainer = NewObject<URCVirtualPropertyInContainer>(this, InPropertyClass.Get(), NAME_None, RF_Transactional);
 	VirtualPropertyInContainer->PropertyName = PropertyName;
 	VirtualPropertyInContainer->PresetWeakPtr = PresetWeakPtr;
 	VirtualPropertyInContainer->ContainerWeakPtr = this;
@@ -47,15 +39,6 @@ URCVirtualPropertyInContainer* URCVirtualPropertyContainerBase::AddProperty(cons
 
 URCVirtualPropertyInContainer* URCVirtualPropertyContainerBase::DuplicateProperty(const FName& InPropertyName, const FProperty* InSourceProperty, TSubclassOf<URCVirtualPropertyInContainer> InPropertyClass)
 {
-#if WITH_EDITOR
-	const FScopedTransaction Transaction(LOCTEXT("DuplicateProperty", "Duplicate Property"));
-
-	Modify();
-
-	MarkPackageDirty();
-#endif
-
-
 	const FPropertyBagPropertyDesc* BagPropertyDesc = Bag.FindPropertyDescByName(InPropertyName);
 	if (ensure(!BagPropertyDesc))
 	{
@@ -99,14 +82,6 @@ URCVirtualPropertyInContainer* URCVirtualPropertyContainerBase::DuplicatePropert
 
 bool URCVirtualPropertyContainerBase::RemoveProperty(const FName& InPropertyName)
 {
-#if WITH_EDITOR
-	const FScopedTransaction Transaction(LOCTEXT("RemoveProperty", "Remove Property"));
-
-	Modify();
-
-	MarkPackageDirty();
-#endif
-
 	Bag.RemovePropertyByName(InPropertyName);
 
 	for (auto PropertiesIt = VirtualProperties.CreateIterator(); PropertiesIt; ++PropertiesIt)
@@ -126,14 +101,6 @@ bool URCVirtualPropertyContainerBase::RemoveProperty(const FName& InPropertyName
 
 void URCVirtualPropertyContainerBase::Reset()
 {
-#if WITH_EDITOR
-	const FScopedTransaction Transaction(LOCTEXT("EmptyProperties", "Empty Properties"));
-
-	Modify();
-
-	MarkPackageDirty();
-#endif
-	
 	VirtualProperties.Empty();
 
 	Bag.Reset();
@@ -227,6 +194,13 @@ FName URCVirtualPropertyContainerBase::GenerateUniquePropertyName(const FName& I
 }
 
 #if WITH_EDITOR
+void URCVirtualPropertyContainerBase::PostEditUndo()
+{
+	Super::PostEditUndo();
+
+	OnVirtualPropertyContainerModifiedDelegate.Broadcast();
+}
+
 void URCVirtualPropertyContainerBase::OnModifyPropertyValue(const FPropertyChangedEvent& PropertyChangedEvent)
 {
 	const FScopedTransaction Transaction(LOCTEXT("OnModifyPropertyValue", "On Modify Property Value"));
diff --git a/Engine/Plugins/VirtualProduction/RemoteControl/Source/RemoteControl/Private/RemoteControlPreset.cpp b/Engine/Plugins/VirtualProduction/RemoteControl/Source/RemoteControl/Private/RemoteControlPreset.cpp
--- a/Engine/Plugins/VirtualProduction/RemoteControl/Source/RemoteControl/Private/RemoteControlPreset.cpp
+++ b/Engine/Plugins/VirtualProduction/RemoteControl/Source/RemoteControl/Private/RemoteControlPreset.cpp
@@ -551,7 +551,11 @@ URCVirtualPropertyInContainer* URemoteControlPreset::AddVirtualProperty(TSubclas
 	if (NewPropertyName.IsNone())
 	{
 		NewPropertyName = URCVirtualPropertyContainerBase::GenerateUniquePropertyName(TEXT(""), InValueType, InValueTypeObject, ControllerContainer);
-	}	
+	}
+
+#if WITH_EDITOR
+	ControllerContainer->Modify();
+#endif
 
 	return ControllerContainer->AddProperty(NewPropertyName, InPropertyClass, InValueType, InValueTypeObject);
 }
@@ -560,17 +564,24 @@ bool URemoteControlPreset::RemoveVirtualProperty(const FName& InPropertyName)
 {
 	if (ensure(ControllerContainer))
 	{
+#if WITH_EDITOR
+		ControllerContainer->Modify();
+#endif
+
 		return ControllerContainer->RemoveProperty(InPropertyName);
 	}
 
 	return false;
-	
 }
 
 void URemoteControlPreset::ResetVirtualProperties()
 {
 	if (ensure(ControllerContainer))
 	{
+#if WITH_EDITOR
+		ControllerContainer->Modify();
+#endif
+
 		ControllerContainer->Reset();
 	}
 }
@@ -607,9 +618,16 @@ void URemoteControlPreset::OnModifyVirtualProperty(const FPropertyChangedEvent&
 	{
 		ControllerContainer->OnModifyPropertyValue(PropertyChangedEvent);
 	}	
-}
+} 
 #endif
 
+FOnVirtualPropertyContainerModified& URemoteControlPreset::OnVirtualPropertyContainerModified() const
+{
+	check(ControllerContainer);
+
+	return ControllerContainer->OnVirtualPropertyContainerModified();
+}
+
 TWeakPtr<FRemoteControlProperty> URemoteControlPreset::ExposeProperty(UObject* Object, FRCFieldPathInfo FieldPath, FRemoteControlPresetExposeArgs Args)
 {
 	if (!Object)
diff --git a/Engine/Plugins/VirtualProduction/RemoteControl/Source/RemoteControl/Public/RCVirtualProperty.h b/Engine/Plugins/VirtualProduction/RemoteControl/Source/RemoteControl/Public/RCVirtualProperty.h
--- a/Engine/Plugins/VirtualProduction/RemoteControl/Source/RemoteControl/Public/RCVirtualProperty.h
+++ b/Engine/Plugins/VirtualProduction/RemoteControl/Source/RemoteControl/Public/RCVirtualProperty.h
@@ -264,6 +264,9 @@ public:
 	/** Fetches a user-friendly representation of the Virtual Property's type name (i.e. Controller type), given a Property Bag type and value object (for structures)*/
 	static FName GetVirtualPropertyTypeDisplayName(const EPropertyBagPropertyType InValueType, UObject* InValueTypeObject);
 
+	/** Updates the display index of the virtual property (transaction aware)*/
+	void SetDisplayIndex(const int32 InDisplayIndex);
+
 public:
 	/** Unique property name */
 	UPROPERTY()
diff --git a/Engine/Plugins/VirtualProduction/RemoteControl/Source/RemoteControl/Public/RCVirtualPropertyContainer.h b/Engine/Plugins/VirtualProduction/RemoteControl/Source/RemoteControl/Public/RCVirtualPropertyContainer.h
--- a/Engine/Plugins/VirtualProduction/RemoteControl/Source/RemoteControl/Public/RCVirtualPropertyContainer.h
+++ b/Engine/Plugins/VirtualProduction/RemoteControl/Source/RemoteControl/Public/RCVirtualPropertyContainer.h
@@ -101,10 +101,17 @@ public:
 	static FName GenerateUniquePropertyName(const FName& InPropertyName, const EPropertyBagPropertyType InValueType, UObject* InValueTypeObject, const URCVirtualPropertyContainerBase* InContainer);
 
 #if WITH_EDITOR
+	/** Called after applying a transaction to the object. Used to broadcast Undo related container changes to UI. */
+	virtual void PostEditUndo();
+
 	/** Delegate when object changed */
 	virtual void OnModifyPropertyValue(const FPropertyChangedEvent& PropertyChangedEvent);
 #endif
 
+	/** Returns the delegate that notifies changes to the virtual property container */
+	DECLARE_MULTICAST_DELEGATE(FOnVirtualPropertyContainerModified);
+	FOnVirtualPropertyContainerModified& OnVirtualPropertyContainerModified() { return OnVirtualPropertyContainerModifiedDelegate; }
+
 protected:
 	/** Holds bag of properties. */
 	UPROPERTY()
@@ -118,6 +125,10 @@ public:
 	/** Pointer to Remote Control Preset */
 	UPROPERTY()
 	TWeakObjectPtr<URemoteControlPreset> PresetWeakPtr;
+
+	/** Delegate that notifies changes to the virtual property container*/
+	FOnVirtualPropertyContainerModified OnVirtualPropertyContainerModifiedDelegate;
+
 };
 
 
diff --git a/Engine/Plugins/VirtualProduction/RemoteControl/Source/RemoteControl/Public/RemoteControlPreset.h b/Engine/Plugins/VirtualProduction/RemoteControl/Source/RemoteControl/Public/RemoteControlPreset.h
--- a/Engine/Plugins/VirtualProduction/RemoteControl/Source/RemoteControl/Public/RemoteControlPreset.h
+++ b/Engine/Plugins/VirtualProduction/RemoteControl/Source/RemoteControl/Public/RemoteControlPreset.h
@@ -21,8 +21,8 @@ struct FPropertyChangedEvent;
 struct FRCFieldPathInfo;
 struct FRemoteControlActor;
 struct FRemoteControlPresetLayout;
-class FTransactionObjectEvent;
 class FRemoteControlPresetRebindingManager;
+class FTransactionObjectEvent;
 class UBlueprint;
 class URCVirtualPropertyBase;
 class URCVirtualPropertyContainerBase;
@@ -31,6 +31,8 @@ class URemoteControlExposeRegistry;
 class URemoteControlBinding;
 class URemoteControlPreset;
 
+DECLARE_MULTICAST_DELEGATE(FOnVirtualPropertyContainerModified);
+
 /** Arguments used to expose an entity (Actor, property, function, etc.) */
 struct REMOTECONTROL_API FRemoteControlPresetExposeArgs
 {
@@ -657,6 +659,8 @@ public:
 	void OnModifyVirtualProperty(const FPropertyChangedEvent& PropertyChangedEvent);
 #endif
 
+	FOnVirtualPropertyContainerModified& OnVirtualPropertyContainerModified() const;
+
 private:
 	/** Holds virtual controllers properties, behaviours and actions */
 	UPROPERTY(Instanced)
diff --git a/Engine/Plugins/VirtualProduction/RemoteControl/Source/RemoteControlLogic/Private/Action/RCActionContainer.cpp b/Engine/Plugins/VirtualProduction/RemoteControl/Source/RemoteControlLogic/Private/Action/RCActionContainer.cpp
--- a/Engine/Plugins/VirtualProduction/RemoteControl/Source/RemoteControlLogic/Private/Action/RCActionContainer.cpp
+++ b/Engine/Plugins/VirtualProduction/RemoteControl/Source/RemoteControlLogic/Private/Action/RCActionContainer.cpp
@@ -84,8 +84,7 @@ URCPropertyAction* URCActionContainer::AddPropertyAction(const TSharedRef<const
 		NewPropertyAction->PropertySelfContainer->DuplicatePropertyWithCopy(PropertyName, InRemoteControlProperty->GetProperty(), (uint8*)ObjectRef.ContainerAdress);
 	}
 
-	// Add action to array
-	Actions.Add(NewPropertyAction);
+	AddAction(NewPropertyAction);
 	
 	return NewPropertyAction;
 }
@@ -98,11 +97,25 @@ URCFunctionAction* URCActionContainer::AddFunctionAction(const TSharedRef<const
 	NewFunctionAction->ExposedFieldId = InRemoteControlFunction->GetId();
 	NewFunctionAction->Id = FGuid::NewGuid();
 
-	Actions.Add(NewFunctionAction);
+	AddAction(NewFunctionAction);
 	
 	return NewFunctionAction;
 }
 
+void URCActionContainer::AddAction(URCAction* NewAction)
+{
+	Actions.Add(NewAction);
+}
+
+#if WITH_EDITOR
+void URCActionContainer::PostEditUndo()
+{
+	Super::PostEditUndo();
+
+	OnActionsListModified.Broadcast();
+}
+#endif
+
 URCAction* URCActionContainer::FindActionByFieldId(const FGuid InId) const
 {
 	for (URCAction* Action :  Actions)
@@ -145,4 +158,4 @@ int32 URCActionContainer::RemoveAction(URCAction* InAction)
 void URCActionContainer::EmptyActions()
 {
 	Actions.Empty();
-}
+}
\ No newline at end of file
diff --git a/Engine/Plugins/VirtualProduction/RemoteControl/Source/RemoteControlLogic/Private/Behaviour/Builtin/Bind/RCBehaviourBind.cpp b/Engine/Plugins/VirtualProduction/RemoteControl/Source/RemoteControlLogic/Private/Behaviour/Builtin/Bind/RCBehaviourBind.cpp
--- a/Engine/Plugins/VirtualProduction/RemoteControl/Source/RemoteControlLogic/Private/Behaviour/Builtin/Bind/RCBehaviourBind.cpp
+++ b/Engine/Plugins/VirtualProduction/RemoteControl/Source/RemoteControlLogic/Private/Behaviour/Builtin/Bind/RCBehaviourBind.cpp
@@ -33,7 +33,7 @@ URCPropertyBindAction* URCBehaviourBind::AddPropertyBindAction(const TSharedRef<
 	BindAction->Id = FGuid::NewGuid();
 
 	// Add action to array
-	ActionContainer->Actions.Add(BindAction);
+	ActionContainer->AddAction(BindAction);
 
 	return BindAction;
 }
diff --git a/Engine/Plugins/VirtualProduction/RemoteControl/Source/RemoteControlLogic/Private/Behaviour/Builtin/Conditional/RCBehaviourConditional.cpp b/Engine/Plugins/VirtualProduction/RemoteControl/Source/RemoteControlLogic/Private/Behaviour/Builtin/Conditional/RCBehaviourConditional.cpp
--- a/Engine/Plugins/VirtualProduction/RemoteControl/Source/RemoteControlLogic/Private/Behaviour/Builtin/Conditional/RCBehaviourConditional.cpp
+++ b/Engine/Plugins/VirtualProduction/RemoteControl/Source/RemoteControlLogic/Private/Behaviour/Builtin/Conditional/RCBehaviourConditional.cpp
@@ -56,7 +56,7 @@ void URCBehaviourConditional::Execute()
 
 	ERCBehaviourConditionType PreviousConditionType = ERCBehaviourConditionType::None;
 
-	for (TObjectPtr<URCAction> Action : ActionContainer->Actions)
+	for (TObjectPtr<URCAction> Action : ActionContainer->GetActions())
 	{
 		FRCBehaviourCondition* Condition = Conditions.Find(Action);
 		if (!Condition)
diff --git a/Engine/Plugins/VirtualProduction/RemoteControl/Source/RemoteControlLogic/Private/Behaviour/Builtin/RangeMap/RCRangeMapBehaviour.cpp b/Engine/Plugins/VirtualProduction/RemoteControl/Source/RemoteControlLogic/Private/Behaviour/Builtin/RangeMap/RCRangeMapBehaviour.cpp
--- a/Engine/Plugins/VirtualProduction/RemoteControl/Source/RemoteControlLogic/Private/Behaviour/Builtin/RangeMap/RCRangeMapBehaviour.cpp
+++ b/Engine/Plugins/VirtualProduction/RemoteControl/Source/RemoteControlLogic/Private/Behaviour/Builtin/RangeMap/RCRangeMapBehaviour.cpp
@@ -126,7 +126,7 @@ bool URCRangeMapBehaviour::GetNearestActionByThreshold(TTuple<URCAction*, bool>&
 
 void URCRangeMapBehaviour::GetLerpActions(TMap<FGuid, TArray<URCAction*>>& OutNumericActionsByField)
 {
-	for (URCAction* Action : ActionContainer->Actions)
+	for (URCAction* Action : ActionContainer->GetActions())
 	{
 		TArray<URCAction*>& LerpActionArray = OutNumericActionsByField.FindOrAdd(Action->ExposedFieldId);
 		
@@ -364,7 +364,7 @@ TMap<double, URCAction*> URCRangeMapBehaviour::GetNonLerpActions()
 {
 	TMap<double, URCAction*> ActionMap;
 	
-	for (URCAction* Action : ActionContainer->Actions)
+	for (URCAction* Action : ActionContainer->GetActions())
 	{
 		
 		if (!IsSupportedActionLerpType(Action))
diff --git a/Engine/Plugins/VirtualProduction/RemoteControl/Source/RemoteControlLogic/Private/Behaviour/RCBehaviour.cpp b/Engine/Plugins/VirtualProduction/RemoteControl/Source/RemoteControlLogic/Private/Behaviour/RCBehaviour.cpp
--- a/Engine/Plugins/VirtualProduction/RemoteControl/Source/RemoteControlLogic/Private/Behaviour/RCBehaviour.cpp
+++ b/Engine/Plugins/VirtualProduction/RemoteControl/Source/RemoteControlLogic/Private/Behaviour/RCBehaviour.cpp
@@ -30,12 +30,16 @@ void URCBehaviour::Execute()
 
 URCAction* URCBehaviour::AddAction(const TSharedRef<const FRemoteControlField> InRemoteControlField)
 {
+#if WITH_EDITOR
+	ActionContainer->Modify();
+#endif
+
 	return ActionContainer->AddAction(InRemoteControlField);
 }
 
 int32 URCBehaviour::GetNumActions() const
 {
-	return ActionContainer->Actions.Num();
+	return ActionContainer->GetActions().Num();
 }
 
 bool URCBehaviour::CanHaveActionForField(const TSharedPtr<FRemoteControlField> InRemoteControlField) const
diff --git a/Engine/Plugins/VirtualProduction/RemoteControl/Source/RemoteControlLogic/Private/Controller/RCController.cpp b/Engine/Plugins/VirtualProduction/RemoteControl/Source/RemoteControlLogic/Private/Controller/RCController.cpp
--- a/Engine/Plugins/VirtualProduction/RemoteControl/Source/RemoteControlLogic/Private/Controller/RCController.cpp
+++ b/Engine/Plugins/VirtualProduction/RemoteControl/Source/RemoteControlLogic/Private/Controller/RCController.cpp
@@ -8,6 +8,17 @@
 #include "Behaviour/RCBehaviour.h"
 #include "Behaviour/RCBehaviourNode.h"
 
+#define LOCTEXT_NAMESPACE "RCController"
+
+#if WITH_EDITOR
+void URCController::PostEditUndo()
+{
+	Super::PostEditUndo();
+
+	OnBehaviourListModified.Broadcast();
+}
+#endif
+
 URCBehaviour* URCController::AddBehaviour(TSubclassOf<URCBehaviourNode> InBehaviourNodeClass)
 {
 	URCBehaviour* NewBehaviour = CreateBehaviour(InBehaviourNodeClass);
@@ -27,7 +38,7 @@ URCBehaviour* URCController::CreateBehaviour(TSubclassOf<URCBehaviourNode> InBeh
 {
 	const URCBehaviourNode* DefaultBehaviourNode = Cast<URCBehaviourNode>(InBehaviourNodeClass->GetDefaultObject());
 	
-	URCBehaviour* NewBehaviour = NewObject<URCBehaviour>(this, DefaultBehaviourNode->GetBehaviourClass());
+	URCBehaviour* NewBehaviour = NewObject<URCBehaviour>(this, DefaultBehaviourNode->GetBehaviourClass(), NAME_None, RF_Transactional);
 	NewBehaviour->BehaviourNodeClass = InBehaviourNodeClass;
 	NewBehaviour->Id = FGuid::NewGuid();
 	NewBehaviour->ActionContainer->PresetWeakPtr = PresetWeakPtr;
@@ -77,3 +88,5 @@ void URCController::ExecuteBehaviours()
 		}
 	}
 }
+
+#undef LOCTEXT_NAMESPACE /* RCController */ 
\ No newline at end of file
diff --git a/Engine/Plugins/VirtualProduction/RemoteControl/Source/RemoteControlLogic/Private/RemoteControlLogicModule.cpp b/Engine/Plugins/VirtualProduction/RemoteControl/Source/RemoteControlLogic/Private/RemoteControlLogicModule.cpp
--- a/Engine/Plugins/VirtualProduction/RemoteControl/Source/RemoteControlLogic/Private/RemoteControlLogicModule.cpp
+++ b/Engine/Plugins/VirtualProduction/RemoteControl/Source/RemoteControlLogic/Private/RemoteControlLogicModule.cpp
@@ -29,7 +29,7 @@ private:
 		// Create a new controller for preset if that does not exists
 		if (InPreset && !InPreset->IsControllerContainerValid())
 		{
-			URCControllerContainer* ControllerContainer = NewObject<URCControllerContainer>(InPreset);
+			URCControllerContainer* ControllerContainer = NewObject<URCControllerContainer>(InPreset, NAME_None, RF_Transactional);
 			ControllerContainer->PresetWeakPtr = InPreset;
 
 			InPreset->SetControllerContainer(ControllerContainer);
diff --git a/Engine/Plugins/VirtualProduction/RemoteControl/Source/RemoteControlLogic/Public/Action/RCActionContainer.h b/Engine/Plugins/VirtualProduction/RemoteControl/Source/RemoteControlLogic/Public/Action/RCActionContainer.h
--- a/Engine/Plugins/VirtualProduction/RemoteControl/Source/RemoteControlLogic/Public/Action/RCActionContainer.h
+++ b/Engine/Plugins/VirtualProduction/RemoteControl/Source/RemoteControlLogic/Public/Action/RCActionContainer.h
@@ -20,6 +20,8 @@ struct FRemoteControlProperty;
 
 typedef TFunction<bool(const TSet<TObjectPtr<URCAction>>& Actions)> TRCActionUniquenessTest;
 
+DECLARE_MULTICAST_DELEGATE(FOnActionsListModified);
+
 /**
  * Container for created actions
  */
@@ -55,10 +57,17 @@ public:
 	/** Empty action set */
 	void EmptyActions();
 
-	/** Set of actions */
-	UPROPERTY()
-	TSet<TObjectPtr<URCAction>> Actions;
-	
+	/** Register a newly created Action with the Action container, provides Undo/Redo support. */
+	void AddAction(URCAction* NewAction);
+
+	/** Retrieves the actions present in this container */
+	const TSet<TObjectPtr<URCAction>>& GetActions() const { return Actions; }
+
+#if WITH_EDITOR
+	/** Called after applying a transaction to the object. Used to broadcast Undo related container changes to UI */
+	virtual void PostEditUndo();
+#endif
+
 	/** Set of child action container */
 	UPROPERTY()
 	TSet<TObjectPtr<URCActionContainer>> ActionContainers;
@@ -67,10 +76,17 @@ public:
 	UPROPERTY()
 	TWeakObjectPtr<URemoteControlPreset> PresetWeakPtr;
 
+	/** Delegate that notifies changes to the list of actions to a single listener*/
+	FOnActionsListModified OnActionsListModified;
+
 private:
 	/** Add remote control property action  */
 	URCPropertyAction* AddPropertyAction(const TSharedRef<const FRemoteControlProperty> InRemoteControlProperty);
 
 	/** Add remote control property function */
 	URCFunctionAction* AddFunctionAction(const TSharedRef<const FRemoteControlFunction> InRemoteControlFunction);
+
+	/** The list of Actions present in this container */
+	UPROPERTY()
+	TSet<TObjectPtr<URCAction>> Actions;
 };
diff --git a/Engine/Plugins/VirtualProduction/RemoteControl/Source/RemoteControlLogic/Public/Controller/RCController.h b/Engine/Plugins/VirtualProduction/RemoteControl/Source/RemoteControlLogic/Public/Controller/RCController.h
--- a/Engine/Plugins/VirtualProduction/RemoteControl/Source/RemoteControlLogic/Public/Controller/RCController.h
+++ b/Engine/Plugins/VirtualProduction/RemoteControl/Source/RemoteControlLogic/Public/Controller/RCController.h
@@ -8,6 +8,8 @@
 class URCBehaviour;
 class URCBehaviourNode;
 
+DECLARE_MULTICAST_DELEGATE(FOnBehaviourListModified);
+
 /**
  * Remote Control Controller. Container for Behaviours and Actions
  */
@@ -17,6 +19,11 @@ class REMOTECONTROLLOGIC_API URCController : public URCVirtualPropertyInContaine
 	GENERATED_BODY()
 
 public:
+#if WITH_EDITOR
+	/** Called after applying a transaction to the object. Used to broadcast Undo related container changes to UI Used to broadcast Undo related container changes to UI */
+	virtual void PostEditUndo();
+#endif
+
 	/** Create and add behaviour to behaviour set */
 	virtual URCBehaviour* AddBehaviour(TSubclassOf<URCBehaviourNode> InBehaviourNodeClass);
 
@@ -41,6 +48,9 @@ public:
 		ExecuteBehaviours();
 	}
 
+	/** Delegate that notifies changes to the list of behaviours*/
+	FOnBehaviourListModified OnBehaviourListModified;
+
 public:
 	/** Set of the behaviours */
 	UPROPERTY()
diff --git a/Engine/Plugins/VirtualProduction/RemoteControl/Source/RemoteControlUI/Private/UI/Action/SRCActionPanel.cpp b/Engine/Plugins/VirtualProduction/RemoteControl/Source/RemoteControlUI/Private/UI/Action/SRCActionPanel.cpp
--- a/Engine/Plugins/VirtualProduction/RemoteControl/Source/RemoteControlUI/Private/UI/Action/SRCActionPanel.cpp
+++ b/Engine/Plugins/VirtualProduction/RemoteControl/Source/RemoteControlUI/Private/UI/Action/SRCActionPanel.cpp
@@ -33,6 +33,10 @@
 #include "Widgets/Input/SCheckBox.h"
 #include "Widgets/Layout/SBox.h"
 
+#if WITH_EDITOR
+#include "ScopedTransaction.h"
+#endif
+
 #define LOCTEXT_NAMESPACE "SRCActionPanel"
 
 TSharedPtr<SBox> SRCActionPanel::NoneSelectedWidget = SNew(SBox)
@@ -299,15 +303,20 @@ URCAction* SRCActionPanel::AddAction(const TSharedRef<const FRemoteControlField>
 {
 	if (const TSharedPtr<FRCBehaviourModel> BehaviourItem = SelectedBehaviourItemWeakPtr.Pin())
 	{
-		URCAction* NewAction = BehaviourItem->AddAction(InRemoteControlField);
-
-		// Broadcast new Action to other panels
-		if (const TSharedPtr<SRemoteControlPanel> RemoteControlPanel = GetRemoteControlPanel())
+		if (const URCBehaviour* Behaviour = BehaviourItem->GetBehaviour())
 		{
-			RemoteControlPanel->OnActionAdded.Broadcast(NewAction);
-		}
+			Behaviour->ActionContainer->Modify();
+
+			URCAction* NewAction = BehaviourItem->AddAction(InRemoteControlField);
+
+			// Broadcast new Action to other panels
+			if (const TSharedPtr<SRemoteControlPanel> RemoteControlPanel = GetRemoteControlPanel())
+			{
+				RemoteControlPanel->OnActionAdded.Broadcast(NewAction);
+			}
 
-		return NewAction;
+			return NewAction;
+		}
 	}
 
 	return nullptr;
@@ -339,6 +348,8 @@ void SRCActionPanel::OnAddActionClicked(TSharedPtr<FRemoteControlField> InRemote
 		return;
 	}
 
+	FScopedTransaction Transaction(LOCTEXT("AddAction", "Add Action"));
+
 	AddAction(InRemoteControlField.ToSharedRef());
 }
 
@@ -348,6 +359,9 @@ FReply SRCActionPanel::OnClickEmptyButton()
 	{
 		if (const URCBehaviour* Behaviour = BehaviourItem->GetBehaviour())
 		{
+			FScopedTransaction Transaction(LOCTEXT("EmptyActions", "Empty Actions"));
+			Behaviour->ActionContainer->Modify();
+
 			Behaviour->ActionContainer->EmptyActions();
 		}
 	}
@@ -373,6 +387,8 @@ FReply SRCActionPanel::OnAddAllFields()
 		{
 			const TArray<TWeakPtr<FRemoteControlField>>& RemoteControlFields = Preset->GetExposedEntities<FRemoteControlField>();
 
+			FScopedTransaction Transaction(LOCTEXT("AddAllActions", "Add All Actions"));
+
 			// Enumerate the list of Exposed Entities and Functions available in this Preset for our use as Actions
 			for (const TWeakPtr<FRemoteControlField>& RemoteControlFieldWeakPtr : RemoteControlFields)
 			{
diff --git a/Engine/Plugins/VirtualProduction/RemoteControl/Source/RemoteControlUI/Private/UI/Action/SRCActionPanelList.h b/Engine/Plugins/VirtualProduction/RemoteControl/Source/RemoteControlUI/Private/UI/Action/SRCActionPanelList.h
--- a/Engine/Plugins/VirtualProduction/RemoteControl/Source/RemoteControlUI/Private/UI/Action/SRCActionPanelList.h
+++ b/Engine/Plugins/VirtualProduction/RemoteControl/Source/RemoteControlUI/Private/UI/Action/SRCActionPanelList.h
@@ -25,6 +25,12 @@
 #include "Widgets/Views/STableRow.h"
 #include "Widgets/Views/STableViewBase.h"
 
+#if WITH_EDITOR
+#include "ScopedTransaction.h"
+#endif
+
+#define LOCTEXT_NAMESPACE "SRCActionPanelList"
+
 /*
 * ~ SRCActionPanelList ~
 *
@@ -80,6 +86,17 @@ public:
 			RemoteControlPanel->OnEmptyActions.AddSP(this, &SRCActionPanelList::OnEmptyActions);
 		}
 
+		if (URCBehaviour* Behaviour = InBehaviourItem->GetBehaviour())
+		{
+			if (Behaviour->ActionContainer)
+			{
+				Behaviour->ActionContainer->OnActionsListModified.AddSP(this, &SRCActionPanelList::OnActionsListModified);
+			}
+		}
+	}
+
+	void OnActionsListModified()
+	{
 		Reset();
 	}
 
@@ -177,7 +194,7 @@ private:
 				{
 					if (URCBehaviour* Behaviour = Cast<URCBehaviour>(BehaviourItem->GetBehaviour()))
 					{
-						for (URCAction* Action : Behaviour->ActionContainer->Actions)
+						for (URCAction* Action : Behaviour->ActionContainer->GetActions())
 						{
 							TSharedPtr<ActionType> ActionItem = ActionType::GetModelByActionType(Action, BehaviourItem, RemoteControlPanel);
 
@@ -215,6 +232,9 @@ private:
 				if (const TSharedPtr<ActionType> SelectedAction = StaticCastSharedPtr<ActionType>(InModel))
 				{
 					// Remove Model from Data Container
+					FScopedTransaction Transaction(LOCTEXT("RemoveAction", "Remove Action"));
+					Behaviour->ActionContainer->Modify();
+
 					const int32 RemoveCount = Behaviour->ActionContainer->RemoveAction(SelectedAction->GetAction());
 
 					return RemoveCount;
@@ -364,4 +384,6 @@ private:
 
 	/** Panel Style reference. */
 	const FRCPanelStyle* RCPanelStyle;
-};
\ No newline at end of file
+};
+
+#undef LOCTEXT_NAMESPACE
\ No newline at end of file
diff --git a/Engine/Plugins/VirtualProduction/RemoteControl/Source/RemoteControlUI/Private/UI/Behaviour/SRCBehaviourPanel.cpp b/Engine/Plugins/VirtualProduction/RemoteControl/Source/RemoteControlUI/Private/UI/Behaviour/SRCBehaviourPanel.cpp
--- a/Engine/Plugins/VirtualProduction/RemoteControl/Source/RemoteControlUI/Private/UI/Behaviour/SRCBehaviourPanel.cpp
+++ b/Engine/Plugins/VirtualProduction/RemoteControl/Source/RemoteControlUI/Private/UI/Behaviour/SRCBehaviourPanel.cpp
@@ -21,6 +21,10 @@
 #include "UI/Controller/RCControllerModel.h"
 #include "Widgets/Input/SComboButton.h"
 
+#if WITH_EDITOR
+#include "ScopedTransaction.h"
+#endif
+
 #define LOCTEXT_NAMESPACE "SRCBehaviourPanel"
 
 TSharedPtr<SBox> SRCBehaviourPanel::NoneSelectedWidget = SNew(SBox)
@@ -192,6 +196,9 @@ void SRCBehaviourPanel::OnAddBehaviourClicked(UClass* InClass)
 	{
 		if (URCController* Controller = Cast<URCController>(ControllerItem->GetVirtualProperty()))
 		{
+			FScopedTransaction Transaction(LOCTEXT("AddBehaviour", "Add Behaviour"));
+			Controller->Modify();
+
 			URCBehaviour* NewBehaviour = Controller->AddBehaviour(InClass);
 
 			if (const TSharedPtr<SRemoteControlPanel> RemoteControlPanel = GetRemoteControlPanel())
@@ -208,6 +215,8 @@ FReply SRCBehaviourPanel::OnClickEmptyButton()
 	{
 		if (URCController* Controller = Cast<URCController>(ControllerItem->GetVirtualProperty()))
 		{
+			FScopedTransaction Transaction(LOCTEXT("EmptyBehaviours", "Empty Behaviours"));
+
 			Controller->EmptyBehaviours();
 		}
 	}
diff --git a/Engine/Plugins/VirtualProduction/RemoteControl/Source/RemoteControlUI/Private/UI/Behaviour/SRCBehaviourPanelList.cpp b/Engine/Plugins/VirtualProduction/RemoteControl/Source/RemoteControlUI/Private/UI/Behaviour/SRCBehaviourPanelList.cpp
--- a/Engine/Plugins/VirtualProduction/RemoteControl/Source/RemoteControlUI/Private/UI/Behaviour/SRCBehaviourPanelList.cpp
+++ b/Engine/Plugins/VirtualProduction/RemoteControl/Source/RemoteControlUI/Private/UI/Behaviour/SRCBehaviourPanelList.cpp
@@ -21,6 +21,10 @@
 #include "UI/SRemoteControlPanel.h"
 #include "Widgets/Views/SHeaderRow.h"
 
+#if WITH_EDITOR
+#include "ScopedTransaction.h"
+#endif
+
 #define LOCTEXT_NAMESPACE "RemoteControlPanelBehavioursList"
 
 namespace FRemoteControlBehaviourColumns
@@ -65,6 +69,11 @@ void SRCBehaviourPanelList::Construct(const FArguments& InArgs, TSharedRef<SRCBe
 	RemoteControlPanel->OnBehaviourAdded.AddSP(this, &SRCBehaviourPanelList::OnBehaviourAdded);
 	RemoteControlPanel->OnEmptyBehaviours.AddSP(this, &SRCBehaviourPanelList::OnEmptyBehaviours);
 
+	if (URCController* Controller = Cast< URCController>(InControllerItem->GetVirtualProperty()))
+	{
+		Controller->OnBehaviourListModified.AddSP(this, &SRCBehaviourPanelList::OnBehaviourListModified);
+	}
+
 	Reset();
 }
 
@@ -192,6 +201,9 @@ int32 SRCBehaviourPanelList::RemoveModel(const TSharedPtr<FRCLogicModeBase> InMo
 		{
 			if(const TSharedPtr<FRCBehaviourModel> SelectedBehaviour = StaticCastSharedPtr<FRCBehaviourModel>(InModel))
 			{
+				FScopedTransaction Transaction(LOCTEXT("RemoveBehaviour", "Remove Behaviour"));
+				Controller->Modify();
+
 				// Remove Model from Data Container
 				const int32 RemoveCount = Controller->RemoveBehaviour(SelectedBehaviour->GetBehaviour());
 
@@ -203,6 +215,11 @@ int32 SRCBehaviourPanelList::RemoveModel(const TSharedPtr<FRCLogicModeBase> InMo
 	return 0;
 }
 
+void SRCBehaviourPanelList::OnBehaviourListModified()
+{
+	Reset();
+}
+
 bool SRCBehaviourPanelList::IsListFocused() const
 {
 	return ListView->HasAnyUserFocus().IsSet();
diff --git a/Engine/Plugins/VirtualProduction/RemoteControl/Source/RemoteControlUI/Private/UI/Behaviour/SRCBehaviourPanelList.h b/Engine/Plugins/VirtualProduction/RemoteControl/Source/RemoteControlUI/Private/UI/Behaviour/SRCBehaviourPanelList.h
--- a/Engine/Plugins/VirtualProduction/RemoteControl/Source/RemoteControlUI/Private/UI/Behaviour/SRCBehaviourPanelList.h
+++ b/Engine/Plugins/VirtualProduction/RemoteControl/Source/RemoteControlUI/Private/UI/Behaviour/SRCBehaviourPanelList.h
@@ -71,6 +71,8 @@ private:
 	/** Removes the given Behaviour UI model item from the list of UI models*/
 	virtual int32 RemoveModel(const TSharedPtr<FRCLogicModeBase> InModel) override;
 
+	void OnBehaviourListModified();
+
 private:
 
 	/** The parent Behaviour Panel widget */
diff --git a/Engine/Plugins/VirtualProduction/RemoteControl/Source/RemoteControlUI/Private/UI/Controller/SRCControllerPanel.cpp b/Engine/Plugins/VirtualProduction/RemoteControl/Source/RemoteControlUI/Private/UI/Controller/SRCControllerPanel.cpp
--- a/Engine/Plugins/VirtualProduction/RemoteControl/Source/RemoteControlUI/Private/UI/Controller/SRCControllerPanel.cpp
+++ b/Engine/Plugins/VirtualProduction/RemoteControl/Source/RemoteControlUI/Private/UI/Controller/SRCControllerPanel.cpp
@@ -21,6 +21,10 @@
 #include "Widgets/Input/SComboButton.h"
 #include "Widgets/Views/SListView.h"
 
+#if WITH_EDITOR
+#include "ScopedTransaction.h"
+#endif
+
 #define LOCTEXT_NAMESPACE "SRCControllerPanel"
 
 BEGIN_SLATE_FUNCTION_BUILD_OPTIMIZATION
@@ -235,6 +239,9 @@ void SRCControllerPanel::OnAddControllerClicked(const EPropertyBagPropertyType I
 	// Add to the asset
 	if (URemoteControlPreset* Preset = GetPreset())
 	{
+		FScopedTransaction Transaction(LOCTEXT("AddController", "Add Controller"));
+		Preset->Modify();
+
 		URCVirtualPropertyInContainer* NewVirtualProperty = Preset->AddVirtualProperty(URCController::StaticClass(), InValueType, InValueTypeObject);
 
 		if (ControllerPanelList.IsValid())
@@ -261,6 +268,9 @@ FReply SRCControllerPanel::OnClickEmptyButton()
 {
 	if (URemoteControlPreset* Preset = GetPreset())
 	{
+		FScopedTransaction Transaction(LOCTEXT("Empty Controllers", "EmptyControllers"));
+		Preset->Modify();
+
 		Preset->ResetVirtualProperties();
 	}
 
diff --git a/Engine/Plugins/VirtualProduction/RemoteControl/Source/RemoteControlUI/Private/UI/Controller/SRCControllerPanelList.cpp b/Engine/Plugins/VirtualProduction/RemoteControl/Source/RemoteControlUI/Private/UI/Controller/SRCControllerPanelList.cpp
--- a/Engine/Plugins/VirtualProduction/RemoteControl/Source/RemoteControlUI/Private/UI/Controller/SRCControllerPanelList.cpp
+++ b/Engine/Plugins/VirtualProduction/RemoteControl/Source/RemoteControlUI/Private/UI/Controller/SRCControllerPanelList.cpp
@@ -21,6 +21,10 @@
 #include "UI/SRemoteControlPanel.h"
 #include "Widgets/Views/SHeaderRow.h"
 
+#if WITH_EDITOR
+#include "ScopedTransaction.h"
+#endif
+
 #define LOCTEXT_NAMESPACE "SRCControllerPanelList"
 
 namespace UE::RCControllerPanelList
@@ -201,6 +205,9 @@ void SRCControllerPanelList::Construct(const FArguments& InArgs, const TSharedRe
 		check(RemoteControlPanel)
 		RemoteControlPanel->OnControllerAdded.AddSP(this, &SRCControllerPanelList::OnControllerAdded);
 		RemoteControlPanel->OnEmptyControllers.AddSP(this, &SRCControllerPanelList::OnEmptyControllers);
+
+		Preset->OnVirtualPropertyContainerModified().AddSP(this, &SRCControllerPanelList::OnControllerContainerModified);
+
 	}
 
 	FPropertyRowGeneratorArgs Args;
@@ -318,6 +325,11 @@ void SRCControllerPanelList::OnEmptyControllers()
 	}
 }
 
+void SRCControllerPanelList::OnControllerContainerModified()
+{
+	Reset();
+}
+
 void SRCControllerPanelList::BroadcastOnItemRemoved()
 {
 	if (const TSharedPtr<SRemoteControlPanel> RemoteControlPanel = ControllerPanelWeakPtr.Pin()->GetRemoteControlPanel())
@@ -345,11 +357,36 @@ int32 SRCControllerPanelList::RemoveModel(const TSharedPtr<FRCLogicModeBase> InM
 		{
 			if(const TSharedPtr<FRCControllerModel> SelectedController = StaticCastSharedPtr<FRCControllerModel>(InModel))
 			{
-				// Remove Model from Data Container
-				const bool bRemoved = Preset->RemoveVirtualProperty(SelectedController->GetPropertyName());
-				if (bRemoved)
+				if(URCVirtualPropertyBase* ControllerToRemove = SelectedController->GetVirtualProperty())
 				{
-					return 1; // Remove Count
+					const int32 DisplayIndexToRemove = ControllerToRemove->DisplayIndex;
+
+					FScopedTransaction Transaction(LOCTEXT("RemoveController", "Remove Controller"));
+					Preset->Modify();
+
+					// Remove Model from Data Container
+					const bool bRemoved = Preset->RemoveVirtualProperty(SelectedController->GetPropertyName());
+					if (bRemoved)
+					{
+						// Shift all display indices higher than the deleted index down by 1
+						for (int32 ControllerIndex = 0; ControllerIndex < ControllerItems.Num(); ControllerIndex++)
+						{
+							if (ensure(ControllerItems[ControllerIndex]))
+							{
+								if (URCVirtualPropertyBase* Controller = ControllerItems[ControllerIndex]->GetVirtualProperty())
+								{
+									if (Controller->DisplayIndex > DisplayIndexToRemove)
+									{
+										Controller->Modify();
+
+										Controller->SetDisplayIndex(Controller->DisplayIndex - 1);
+									}
+								}
+							}
+						}
+
+						return 1; // Remove Count
+					}
 				}
 			}
 		}
@@ -400,9 +437,12 @@ void SRCControllerPanelList::ReorderControllerItem(TSharedRef<FRCControllerModel
 	// Update display indices
 	for (int32 i = Index; i < ControllerItems.Num(); i++)
 	{
-		if(URCVirtualPropertyBase * Controller = ControllerItems[i]->GetVirtualProperty())
+		if (ensure(ControllerItems[i]))
 		{
-			Controller->DisplayIndex = i;
+			if (URCVirtualPropertyBase* Controller = ControllerItems[i]->GetVirtualProperty())
+			{
+				Controller->SetDisplayIndex(i);
+			}
 		}
 	}
 
diff --git a/Engine/Plugins/VirtualProduction/RemoteControl/Source/RemoteControlUI/Private/UI/Controller/SRCControllerPanelList.h b/Engine/Plugins/VirtualProduction/RemoteControl/Source/RemoteControlUI/Private/UI/Controller/SRCControllerPanelList.h
--- a/Engine/Plugins/VirtualProduction/RemoteControl/Source/RemoteControlUI/Private/UI/Controller/SRCControllerPanelList.h
+++ b/Engine/Plugins/VirtualProduction/RemoteControl/Source/RemoteControlUI/Private/UI/Controller/SRCControllerPanelList.h
@@ -107,6 +107,8 @@ private:
 	/** Responds to the removal of all Controllers. Resets UI state */
 	void OnEmptyControllers();
 
+	void OnControllerContainerModified();
+
 	/** Change listener for Controllers. Bound to the PropertyRowGenerator's delegate
 	* This is propagated to the corresponding Controller model (Virtual Property) for evaluating all associated Behaviours.
 	*/
