Actions not loading. Fix bug introduced for Undo/Redo related submission

diff --git a/Engine/Plugins/VirtualProduction/RemoteControl/Source/RemoteControlUI/Private/UI/Action/SRCActionPanelList.h b/Engine/Plugins/VirtualProduction/RemoteControl/Source/RemoteControlUI/Private/UI/Action/SRCActionPanelList.h
--- a/Engine/Plugins/VirtualProduction/RemoteControl/Source/RemoteControlUI/Private/UI/Action/SRCActionPanelList.h
+++ b/Engine/Plugins/VirtualProduction/RemoteControl/Source/RemoteControlUI/Private/UI/Action/SRCActionPanelList.h
@@ -93,6 +93,8 @@ public:
 				Behaviour->ActionContainer->OnActionsListModified.AddSP(this, &SRCActionPanelList::OnActionsListModified);
 			}
 		}
+
+		Reset();
 	}
 
 	void OnActionsListModified()
