Localization conflict fix

diff --git a/Engine/Plugins/VirtualProduction/RemoteControl/Source/RemoteControlUI/Private/UI/Action/SRCActionPanel.cpp b/Engine/Plugins/VirtualProduction/RemoteControl/Source/RemoteControlUI/Private/UI/Action/SRCActionPanel.cpp
--- a/Engine/Plugins/VirtualProduction/RemoteControl/Source/RemoteControlUI/Private/UI/Action/SRCActionPanel.cpp
+++ b/Engine/Plugins/VirtualProduction/RemoteControl/Source/RemoteControlUI/Private/UI/Action/SRCActionPanel.cpp
@@ -348,7 +348,7 @@ void SRCActionPanel::OnAddActionClicked(TSharedPtr<FRemoteControlField> InRemote
 		return;
 	}
 
-	FScopedTransaction Transaction(LOCTEXT("AddAction", "Add Action"));
+	FScopedTransaction Transaction(LOCTEXT("AddActionTransaction", "Add Action"));
 
 	AddAction(InRemoteControlField.ToSharedRef());
 }
@@ -359,7 +359,7 @@ FReply SRCActionPanel::OnClickEmptyButton()
 	{
 		if (const URCBehaviour* Behaviour = BehaviourItem->GetBehaviour())
 		{
-			FScopedTransaction Transaction(LOCTEXT("EmptyActions", "Empty Actions"));
+			FScopedTransaction Transaction(LOCTEXT("EmptyActionsTransaction", "Empty Actions"));
 			Behaviour->ActionContainer->Modify();
 
 			Behaviour->ActionContainer->EmptyActions();
@@ -387,7 +387,7 @@ FReply SRCActionPanel::OnAddAllFields()
 		{
 			const TArray<TWeakPtr<FRemoteControlField>>& RemoteControlFields = Preset->GetExposedEntities<FRemoteControlField>();
 
-			FScopedTransaction Transaction(LOCTEXT("AddAllActions", "Add All Actions"));
+			FScopedTransaction Transaction(LOCTEXT("AddAllActionsTransaction", "Add All Actions"));
 
 			// Enumerate the list of Exposed Entities and Functions available in this Preset for our use as Actions
 			for (const TWeakPtr<FRemoteControlField>& RemoteControlFieldWeakPtr : RemoteControlFields)
