[RemoteControl] Fix delta operations failing to apply to nested properties

diff --git a/Engine/Plugins/VirtualProduction/RemoteControl/Source/RemoteControl/Private/RemoteControlFunctionLibrary.cpp b/Engine/Plugins/VirtualProduction/RemoteControl/Source/RemoteControl/Private/RemoteControlFunctionLibrary.cpp
--- a/Engine/Plugins/VirtualProduction/RemoteControl/Source/RemoteControl/Private/RemoteControlFunctionLibrary.cpp
+++ b/Engine/Plugins/VirtualProduction/RemoteControl/Source/RemoteControl/Private/RemoteControlFunctionLibrary.cpp
@@ -63,11 +63,59 @@ bool URemoteControlFunctionLibrary::ApplyColorWheelDelta(UObject* TargetObject,
 		return false;
 	}
 
-	FProperty* Property = TargetObject->GetClass()->FindPropertyByName(FName(*PropertyName));
+	// Split the property name up on dots so we can find nested struct properties
+	TArray<FString> PropertyPathNames;
+	PropertyName.ParseIntoArray(PropertyPathNames, TEXT("."));
+
+	if (PropertyPathNames.IsEmpty())
+	{
+		return false;
+	}
+
+	void* Container = TargetObject;
+	UStruct* ContainerClass = TargetObject->GetClass();
+	if (!ContainerClass)
+	{
+		return false;
+	}
+
+	FProperty* Property = ContainerClass->FindPropertyByName(FName(*PropertyPathNames[0]));
+	if (!Property)
+	{
+		return false;
+	}
+
+	// Walk through the chain of structs until we reach the end (or run out of matching struct properties)
+	for (int32 PathIndex = 1; PathIndex < PropertyPathNames.Num(); ++PathIndex)
+	{
+		if (const FStructProperty* StructProperty = CastField<FStructProperty>(Property))
+		{
+			ContainerClass = StructProperty->Struct;
+			if (ContainerClass == nullptr)
+			{
+				return false;
+			}
+
+			Property = ContainerClass->FindPropertyByName(FName(*PropertyPathNames[PathIndex]));
+			if (Property == nullptr)
+			{
+				return false;
+			}
+
+			Container = StructProperty->ContainerPtrToValuePtr<void>(Container);
+		}
+		else
+		{
+			return false;
+		}
+	}
+
+	ensure(Container && Property);
+
 	if (const FStructProperty* ColorProperty = CastField<FStructProperty>(Property))
 	{
 		FLinearColor Color;
-		ColorProperty->GetValue_InContainer(TargetObject, &Color);
+		ColorProperty->GetValue_InContainer(Container, &Color);
 
 		// Convert to HSV
 		Color = Color.LinearRGBToHSV();
@@ -105,7 +153,7 @@ bool URemoteControlFunctionLibrary::ApplyColorWheelDelta(UObject* TargetObject,
 		TargetObject->Modify();
 #endif
 
-		ColorProperty->SetValue_InContainer(TargetObject, &Color);
+		ColorProperty->SetValue_InContainer(Container, &Color);
 
 #if WITH_EDITOR
 		FPropertyChangedEvent ChangeEvent(Property, bIsInteractive ? EPropertyChangeType::Interactive : EPropertyChangeType::ValueSet);
diff --git a/Engine/Plugins/VirtualProduction/RemoteControl/Source/RemoteControl/Private/RemoteControlModule.cpp b/Engine/Plugins/VirtualProduction/RemoteControl/Source/RemoteControl/Private/RemoteControlModule.cpp
--- a/Engine/Plugins/VirtualProduction/RemoteControl/Source/RemoteControl/Private/RemoteControlModule.cpp
+++ b/Engine/Plugins/VirtualProduction/RemoteControl/Source/RemoteControl/Private/RemoteControlModule.cpp
@@ -2292,6 +2292,9 @@ bool FRemoteControlModule::DeserializeDeltaModificationData(const FRCObjectRefer
 		return false;
 	}
 
+	// The offset between the container we're modifying and the object it belongs to
+	const ptrdiff_t ContainerFromBaseOffset = ((const uint8*)ObjectAccess.ContainerAdress) - ((const uint8*)ObjectAccess.Object.Get());
+
 	// Apply delta operation to each property that was changed (where possible)
 	for (const FPropertyMapStructDeserializerBackendWrapper::FReadPropertyData& ReadProperty : BackendWrapper.GetReadProperties())
 	{
@@ -2299,10 +2302,10 @@ bool FRemoteControlModule::DeserializeDeltaModificationData(const FRCObjectRefer
 		void* OutPropertyValue = ReadProperty.Data;
 
 		// Offset from start of OutData struct to the property that was changed
-		ptrdiff_t Offset = ((const uint8*)OutPropertyValue) - ((const uint8*)OutContainerAddress);
+		const ptrdiff_t Offset = ((const uint8*)OutPropertyValue) - ((const uint8*)OutContainerAddress);
 
 		// Pointer to the equivalent property in the original object
-		const void* BasePropertyValue = ((uint8*)ObjectAccess.Object.Get()) + Offset;
+		const void* BasePropertyValue = ((uint8*)ObjectAccess.Object.Get()) + ContainerFromBaseOffset + Offset;
 
 		if (FNumericProperty* NumericProperty = CastField<FNumericProperty>(ReadProperty.Property))
 		{
