1. Implement Protocol Specific Dynamic Columns

2. Implement Masking Support for Protocols

diff --git a/Engine/Plugins/VirtualProduction/RemoteControl/Source/RemoteControl/Private/Factories/RemoteControlMaskingFactories.cpp b/Engine/Plugins/VirtualProduction/RemoteControl/Source/RemoteControl/Private/Factories/RemoteControlMaskingFactories.cpp
new file mode 100644
--- /dev/null
+++ b/Engine/Plugins/VirtualProduction/RemoteControl/Source/RemoteControl/Private/Factories/RemoteControlMaskingFactories.cpp
@@ -0,0 +1,422 @@
+// Copyright Epic Games, Inc. All Rights Reserved.
+
+#include "RemoteControlMaskingFactories.h"
+
+/**
+ * FVectorMaskingFactory
+ */
+TSharedRef<IRemoteControlMaskingFactory> FVectorMaskingFactory::MakeInstance()
+{
+	return MakeShared<FVectorMaskingFactory>();
+}
+
+void FVectorMaskingFactory::ApplyMaskedValues(const TSharedRef<FRCMaskingOperation>& InMaskingOperation, bool bIsInteractive)
+{
+	if (FStructProperty* ToStructProp = CastField<FStructProperty>(InMaskingOperation->ObjectRef.Property.Get()))
+	{
+		if (UObject* OwningObject = InMaskingOperation->ObjectRef.Object.Get())
+		{
+#if WITH_EDITOR
+			OwningObject->PreEditChange(ToStructProp);
+			OwningObject->Modify();
+#endif // WITH_EDITOR
+
+			if (const FVector* VectorProp = ToStructProp->ContainerPtrToValuePtr<FVector>(OwningObject))
+			{
+				FVector MaskedVector;
+
+				MaskedVector.X = InMaskingOperation->HasMask(ERCMask::MaskA) ? VectorProp->X : InMaskingOperation->PreMaskingCache.X;
+				MaskedVector.Y = InMaskingOperation->HasMask(ERCMask::MaskB) ? VectorProp->Y : InMaskingOperation->PreMaskingCache.Y;
+				MaskedVector.Z = InMaskingOperation->HasMask(ERCMask::MaskC) ? VectorProp->Z : InMaskingOperation->PreMaskingCache.Z;
+
+				ToStructProp->SetValue_InContainer(OwningObject, &MaskedVector);
+			}
+
+#if WITH_EDITOR
+			FPropertyChangedEvent ChangeEvent(ToStructProp, bIsInteractive ? EPropertyChangeType::Interactive : EPropertyChangeType::ValueSet);
+			OwningObject->PostEditChangeProperty(ChangeEvent);
+#endif // WITH_EDITOR
+		}
+	}
+}
+
+void FVectorMaskingFactory::CacheRawValues(const TSharedRef<FRCMaskingOperation>& InMaskingOperation)
+{
+	if (const FStructProperty* FromStructProp = CastField<FStructProperty>(InMaskingOperation->ObjectRef.Property.Get()))
+	{
+		if (const UObject* OwningObject = InMaskingOperation->ObjectRef.Object.Get())
+		{
+			if (const FVector* VectorProp = FromStructProp->ContainerPtrToValuePtr<FVector>(OwningObject))
+			{
+				InMaskingOperation->PreMaskingCache.X = VectorProp->X;
+				InMaskingOperation->PreMaskingCache.Y = VectorProp->Y;
+				InMaskingOperation->PreMaskingCache.Z = VectorProp->Z;
+			}
+		}
+	}
+}
+
+bool FVectorMaskingFactory::SupportsExposedEntity(UScriptStruct* ScriptStruct) const
+{
+	return ScriptStruct == TBaseStructure<FVector>::Get();
+}
+
+/**
+ * FVector4MaskingFactory
+ */
+TSharedRef<IRemoteControlMaskingFactory> FVector4MaskingFactory::MakeInstance()
+{
+	return MakeShared<FVector4MaskingFactory>();
+}
+
+void FVector4MaskingFactory::ApplyMaskedValues(const TSharedRef<FRCMaskingOperation>& InMaskingOperation, bool bIsInteractive)
+{
+	if (FStructProperty* ToStructProp = CastField<FStructProperty>(InMaskingOperation->ObjectRef.Property.Get()))
+	{
+		if (UObject* OwningObject = InMaskingOperation->ObjectRef.Object.Get())
+		{
+#if WITH_EDITOR
+			OwningObject->PreEditChange(ToStructProp);
+			OwningObject->Modify();
+#endif // WITH_EDITOR
+
+			if (const FVector4* Vector4Prop = ToStructProp->ContainerPtrToValuePtr<FVector4>(OwningObject))
+			{
+				FVector4 MaskedVector4;
+
+				MaskedVector4.X = InMaskingOperation->HasMask(ERCMask::MaskA) ? Vector4Prop->X : InMaskingOperation->PreMaskingCache.X;
+				MaskedVector4.Y = InMaskingOperation->HasMask(ERCMask::MaskB) ? Vector4Prop->Y : InMaskingOperation->PreMaskingCache.Y;
+				MaskedVector4.Z = InMaskingOperation->HasMask(ERCMask::MaskC) ? Vector4Prop->Z : InMaskingOperation->PreMaskingCache.Z;
+				MaskedVector4.W = InMaskingOperation->HasMask(ERCMask::MaskD) ? Vector4Prop->W : InMaskingOperation->PreMaskingCache.W;
+
+				ToStructProp->SetValue_InContainer(OwningObject, &MaskedVector4);
+			}
+
+#if WITH_EDITOR
+			FPropertyChangedEvent ChangeEvent(ToStructProp, bIsInteractive ? EPropertyChangeType::Interactive : EPropertyChangeType::ValueSet);
+			OwningObject->PostEditChangeProperty(ChangeEvent);
+#endif // WITH_EDITOR
+		}
+	}
+}
+
+void FVector4MaskingFactory::CacheRawValues(const TSharedRef<FRCMaskingOperation>& InMaskingOperation)
+{
+	if (const FStructProperty* FromStructProp = CastField<FStructProperty>(InMaskingOperation->ObjectRef.Property.Get()))
+	{
+		if (const UObject* OwningObject = InMaskingOperation->ObjectRef.Object.Get())
+		{
+			if (const FVector4* Vector4Prop = FromStructProp->ContainerPtrToValuePtr<FVector4>(OwningObject))
+			{
+				InMaskingOperation->PreMaskingCache.X = Vector4Prop->X;
+				InMaskingOperation->PreMaskingCache.Y = Vector4Prop->Y;
+				InMaskingOperation->PreMaskingCache.Z = Vector4Prop->Z;
+				InMaskingOperation->PreMaskingCache.W = Vector4Prop->W;
+			}
+		}
+	}
+}
+
+bool FVector4MaskingFactory::SupportsExposedEntity(UScriptStruct* ScriptStruct) const
+{
+	return ScriptStruct == TBaseStructure<FVector4>::Get();
+}
+
+/**
+ * FIntVectorMaskingFactory
+ */
+TSharedRef<IRemoteControlMaskingFactory> FIntVectorMaskingFactory::MakeInstance()
+{
+	return MakeShared<FIntVectorMaskingFactory>();
+}
+
+void FIntVectorMaskingFactory::ApplyMaskedValues(const TSharedRef<FRCMaskingOperation>& InMaskingOperation, bool bIsInteractive)
+{
+	if (FStructProperty* ToStructProp = CastField<FStructProperty>(InMaskingOperation->ObjectRef.Property.Get()))
+	{
+		if (UObject* OwningObject = InMaskingOperation->ObjectRef.Object.Get())
+		{
+#if WITH_EDITOR
+			OwningObject->PreEditChange(ToStructProp);
+			OwningObject->Modify();
+#endif // WITH_EDITOR
+
+			if (const FIntVector* IntVectorProp = ToStructProp->ContainerPtrToValuePtr<FIntVector>(OwningObject))
+			{
+				FIntVector MaskedIntVector;
+
+				MaskedIntVector.X = InMaskingOperation->HasMask(ERCMask::MaskA) ? IntVectorProp->X : InMaskingOperation->PreMaskingCache.X;
+				MaskedIntVector.Y = InMaskingOperation->HasMask(ERCMask::MaskB) ? IntVectorProp->Y : InMaskingOperation->PreMaskingCache.Y;
+				MaskedIntVector.Z = InMaskingOperation->HasMask(ERCMask::MaskC) ? IntVectorProp->Z : InMaskingOperation->PreMaskingCache.Z;
+
+				ToStructProp->SetValue_InContainer(OwningObject, &MaskedIntVector);
+			}
+
+#if WITH_EDITOR
+			FPropertyChangedEvent ChangeEvent(ToStructProp, bIsInteractive ? EPropertyChangeType::Interactive : EPropertyChangeType::ValueSet);
+			OwningObject->PostEditChangeProperty(ChangeEvent);
+#endif // WITH_EDITOR
+		}
+	}
+}
+
+void FIntVectorMaskingFactory::CacheRawValues(const TSharedRef<FRCMaskingOperation>& InMaskingOperation)
+{
+	if (const FStructProperty* FromStructProp = CastField<FStructProperty>(InMaskingOperation->ObjectRef.Property.Get()))
+	{
+		if (const UObject* OwningObject = InMaskingOperation->ObjectRef.Object.Get())
+		{
+			if (const FIntVector* IntVectorProp = FromStructProp->ContainerPtrToValuePtr<FIntVector>(OwningObject))
+			{
+				InMaskingOperation->PreMaskingCache.X = IntVectorProp->X;
+				InMaskingOperation->PreMaskingCache.Y = IntVectorProp->Y;
+				InMaskingOperation->PreMaskingCache.Z = IntVectorProp->Z;
+			}
+		}
+	}
+}
+
+bool FIntVectorMaskingFactory::SupportsExposedEntity(UScriptStruct* ScriptStruct) const
+{
+	return ScriptStruct == TBaseStructure<FIntVector>::Get();
+}
+
+/**
+ * FIntVector4MaskingFactory
+ */
+TSharedRef<IRemoteControlMaskingFactory> FIntVector4MaskingFactory::MakeInstance()
+{
+	return MakeShared<FIntVector4MaskingFactory>();
+}
+
+void FIntVector4MaskingFactory::ApplyMaskedValues(const TSharedRef<FRCMaskingOperation>& InMaskingOperation, bool bIsInteractive)
+{
+	if (FStructProperty* ToStructProp = CastField<FStructProperty>(InMaskingOperation->ObjectRef.Property.Get()))
+	{
+		if (UObject* OwningObject = InMaskingOperation->ObjectRef.Object.Get())
+		{
+#if WITH_EDITOR
+			OwningObject->PreEditChange(ToStructProp);
+			OwningObject->Modify();
+#endif // WITH_EDITOR
+
+			if (const FIntVector4* IntVector4Prop = ToStructProp->ContainerPtrToValuePtr<FIntVector4>(OwningObject))
+			{
+				FIntVector4 MaskedIntVector4;
+
+				MaskedIntVector4.X = InMaskingOperation->HasMask(ERCMask::MaskA) ? IntVector4Prop->X : InMaskingOperation->PreMaskingCache.X;
+				MaskedIntVector4.Y = InMaskingOperation->HasMask(ERCMask::MaskB) ? IntVector4Prop->Y : InMaskingOperation->PreMaskingCache.Y;
+				MaskedIntVector4.Z = InMaskingOperation->HasMask(ERCMask::MaskC) ? IntVector4Prop->Z : InMaskingOperation->PreMaskingCache.Z;
+				MaskedIntVector4.W = InMaskingOperation->HasMask(ERCMask::MaskD) ? IntVector4Prop->W : InMaskingOperation->PreMaskingCache.W;
+
+				ToStructProp->SetValue_InContainer(OwningObject, &MaskedIntVector4);
+			}
+
+#if WITH_EDITOR
+			FPropertyChangedEvent ChangeEvent(ToStructProp, bIsInteractive ? EPropertyChangeType::Interactive : EPropertyChangeType::ValueSet);
+			OwningObject->PostEditChangeProperty(ChangeEvent);
+#endif // WITH_EDITOR
+		}
+	}
+}
+
+void FIntVector4MaskingFactory::CacheRawValues(const TSharedRef<FRCMaskingOperation>& InMaskingOperation)
+{
+	if (const FStructProperty* FromStructProp = CastField<FStructProperty>(InMaskingOperation->ObjectRef.Property.Get()))
+	{
+		if (const UObject* OwningObject = InMaskingOperation->ObjectRef.Object.Get())
+		{
+			if (const FIntVector4* IntVector4Prop = FromStructProp->ContainerPtrToValuePtr<FIntVector4>(OwningObject))
+			{
+				InMaskingOperation->PreMaskingCache.X = IntVector4Prop->X;
+				InMaskingOperation->PreMaskingCache.Y = IntVector4Prop->Y;
+				InMaskingOperation->PreMaskingCache.Z = IntVector4Prop->Z;
+				InMaskingOperation->PreMaskingCache.W = IntVector4Prop->W;
+			}
+		}
+	}
+}
+
+bool FIntVector4MaskingFactory::SupportsExposedEntity(UScriptStruct* ScriptStruct) const
+{
+	return ScriptStruct == TBaseStructure<FIntVector4>::Get();
+}
+
+/**
+ * FRotatorMaskingFactory
+ */
+TSharedRef<IRemoteControlMaskingFactory> FRotatorMaskingFactory::MakeInstance()
+{
+	return MakeShared<FRotatorMaskingFactory>();
+}
+
+void FRotatorMaskingFactory::ApplyMaskedValues(const TSharedRef<FRCMaskingOperation>& InMaskingOperation, bool bIsInteractive)
+{
+	if (FStructProperty* ToStructProp = CastField<FStructProperty>(InMaskingOperation->ObjectRef.Property.Get()))
+	{
+		if (UObject* OwningObject = InMaskingOperation->ObjectRef.Object.Get())
+		{
+#if WITH_EDITOR
+			OwningObject->PreEditChange(ToStructProp);
+			OwningObject->Modify();
+#endif // WITH_EDITOR
+
+			if (const FRotator* RotatorProp = ToStructProp->ContainerPtrToValuePtr<FRotator>(OwningObject))
+			{
+				FRotator MaskedRotator;
+
+				MaskedRotator.Roll = InMaskingOperation->HasMask(ERCMask::MaskA) ? RotatorProp->Roll : InMaskingOperation->PreMaskingCache.X;
+				MaskedRotator.Pitch = InMaskingOperation->HasMask(ERCMask::MaskB) ? RotatorProp->Pitch : InMaskingOperation->PreMaskingCache.Y;
+				MaskedRotator.Yaw = InMaskingOperation->HasMask(ERCMask::MaskC) ? RotatorProp->Yaw : InMaskingOperation->PreMaskingCache.Z;
+
+				ToStructProp->SetValue_InContainer(OwningObject, &MaskedRotator);
+			}
+
+#if WITH_EDITOR
+			FPropertyChangedEvent ChangeEvent(ToStructProp, bIsInteractive ? EPropertyChangeType::Interactive : EPropertyChangeType::ValueSet);
+			OwningObject->PostEditChangeProperty(ChangeEvent);
+#endif // WITH_EDITOR
+		}
+	}
+}
+
+void FRotatorMaskingFactory::CacheRawValues(const TSharedRef<FRCMaskingOperation>& InMaskingOperation)
+{
+	if (const FStructProperty* FromStructProp = CastField<FStructProperty>(InMaskingOperation->ObjectRef.Property.Get()))
+	{
+		if (const UObject* OwningObject = InMaskingOperation->ObjectRef.Object.Get())
+		{
+			if (const FRotator* RotatorProp = FromStructProp->ContainerPtrToValuePtr<FRotator>(OwningObject))
+			{
+				InMaskingOperation->PreMaskingCache.X = RotatorProp->Roll;
+				InMaskingOperation->PreMaskingCache.Y = RotatorProp->Pitch;
+				InMaskingOperation->PreMaskingCache.Z = RotatorProp->Yaw;
+			}
+		}
+	}
+}
+
+bool FRotatorMaskingFactory::SupportsExposedEntity(UScriptStruct* ScriptStruct) const
+{
+	return ScriptStruct == TBaseStructure<FRotator>::Get();
+}
+
+/**
+ * FColorMaskingFactory
+ */
+TSharedRef<IRemoteControlMaskingFactory> FColorMaskingFactory::MakeInstance()
+{
+	return MakeShared<FColorMaskingFactory>();
+}
+
+void FColorMaskingFactory::ApplyMaskedValues(const TSharedRef<FRCMaskingOperation>& InMaskingOperation, bool bIsInteractive)
+{
+	if (FStructProperty* ToStructProp = CastField<FStructProperty>(InMaskingOperation->ObjectRef.Property.Get()))
+	{
+		if (UObject* OwningObject = InMaskingOperation->ObjectRef.Object.Get())
+		{
+#if WITH_EDITOR
+			OwningObject->PreEditChange(ToStructProp);
+			OwningObject->Modify();
+#endif // WITH_EDITOR
+
+			if (const FColor* ColorProp = ToStructProp->ContainerPtrToValuePtr<FColor>(OwningObject))
+			{
+				FColor MaskedColor;
+
+				MaskedColor.R = InMaskingOperation->HasMask(ERCMask::MaskA) ? ColorProp->R : InMaskingOperation->PreMaskingCache.X;
+				MaskedColor.G = InMaskingOperation->HasMask(ERCMask::MaskB) ? ColorProp->G : InMaskingOperation->PreMaskingCache.Y;
+				MaskedColor.B = InMaskingOperation->HasMask(ERCMask::MaskC) ? ColorProp->B : InMaskingOperation->PreMaskingCache.Z;
+
+				ToStructProp->SetValue_InContainer(OwningObject, &MaskedColor);
+			}
+
+#if WITH_EDITOR
+			FPropertyChangedEvent ChangeEvent(ToStructProp, bIsInteractive ? EPropertyChangeType::Interactive : EPropertyChangeType::ValueSet);
+			OwningObject->PostEditChangeProperty(ChangeEvent);
+#endif // WITH_EDITOR
+		}
+	}
+}
+
+void FColorMaskingFactory::CacheRawValues(const TSharedRef<FRCMaskingOperation>& InMaskingOperation)
+{
+	if (const FStructProperty* FromStructProp = CastField<FStructProperty>(InMaskingOperation->ObjectRef.Property.Get()))
+	{
+		if (const UObject* OwningObject = InMaskingOperation->ObjectRef.Object.Get())
+		{
+			if (const FColor* ColorProp = FromStructProp->ContainerPtrToValuePtr<FColor>(OwningObject))
+			{
+				InMaskingOperation->PreMaskingCache.X = ColorProp->R;
+				InMaskingOperation->PreMaskingCache.Y = ColorProp->G;
+				InMaskingOperation->PreMaskingCache.Z = ColorProp->B;
+			}
+		}
+	}
+}
+
+bool FColorMaskingFactory::SupportsExposedEntity(UScriptStruct* ScriptStruct) const
+{
+	return ScriptStruct == TBaseStructure<FColor>::Get();
+}
+
+/**
+ * FLinearColorMaskingFactory
+ */
+TSharedRef<IRemoteControlMaskingFactory> FLinearColorMaskingFactory::MakeInstance()
+{
+	return MakeShared<FLinearColorMaskingFactory>();
+}
+
+void FLinearColorMaskingFactory::ApplyMaskedValues(const TSharedRef<FRCMaskingOperation>& InMaskingOperation, bool bIsInteractive)
+{
+	if (FStructProperty* ToStructProp = CastField<FStructProperty>(InMaskingOperation->ObjectRef.Property.Get()))
+	{
+		if (UObject* OwningObject = InMaskingOperation->ObjectRef.Object.Get())
+		{
+#if WITH_EDITOR
+			OwningObject->PreEditChange(ToStructProp);
+			OwningObject->Modify();
+#endif // WITH_EDITOR
+
+			if (const FLinearColor* LinearColorProp = ToStructProp->ContainerPtrToValuePtr<FLinearColor>(OwningObject))
+			{
+				FLinearColor MaskedLinearColor;
+
+				MaskedLinearColor.R = InMaskingOperation->HasMask(ERCMask::MaskA) ? LinearColorProp->R : InMaskingOperation->PreMaskingCache.X;
+				MaskedLinearColor.G = InMaskingOperation->HasMask(ERCMask::MaskB) ? LinearColorProp->G : InMaskingOperation->PreMaskingCache.Y;
+				MaskedLinearColor.B = InMaskingOperation->HasMask(ERCMask::MaskC) ? LinearColorProp->B : InMaskingOperation->PreMaskingCache.Z;
+				MaskedLinearColor.A = InMaskingOperation->HasMask(ERCMask::MaskD) ? LinearColorProp->A : InMaskingOperation->PreMaskingCache.W;
+
+				ToStructProp->SetValue_InContainer(OwningObject, &MaskedLinearColor);
+			}
+
+#if WITH_EDITOR
+			FPropertyChangedEvent ChangeEvent(ToStructProp, bIsInteractive ? EPropertyChangeType::Interactive : EPropertyChangeType::ValueSet);
+			OwningObject->PostEditChangeProperty(ChangeEvent);
+#endif // WITH_EDITOR
+		}
+	}
+}
+
+void FLinearColorMaskingFactory::CacheRawValues(const TSharedRef<FRCMaskingOperation>& InMaskingOperation)
+{
+	if (const FStructProperty* FromStructProp = CastField<FStructProperty>(InMaskingOperation->ObjectRef.Property.Get()))
+	{
+		if (const UObject* OwningObject = InMaskingOperation->ObjectRef.Object.Get())
+		{
+			if (const FLinearColor* LinearColorProp = FromStructProp->ContainerPtrToValuePtr<FLinearColor>(OwningObject))
+			{
+				InMaskingOperation->PreMaskingCache.X = LinearColorProp->R;
+				InMaskingOperation->PreMaskingCache.Y = LinearColorProp->G;
+				InMaskingOperation->PreMaskingCache.Z = LinearColorProp->B;
+				InMaskingOperation->PreMaskingCache.W = LinearColorProp->A;
+			}
+		}
+	}
+}
+
+bool FLinearColorMaskingFactory::SupportsExposedEntity(UScriptStruct* ScriptStruct) const
+{
+	return ScriptStruct == TBaseStructure<FLinearColor>::Get();
+}
diff --git a/Engine/Plugins/VirtualProduction/RemoteControl/Source/RemoteControl/Private/Factories/RemoteControlMaskingFactories.h b/Engine/Plugins/VirtualProduction/RemoteControl/Source/RemoteControl/Private/Factories/RemoteControlMaskingFactories.h
new file mode 100644
--- /dev/null
+++ b/Engine/Plugins/VirtualProduction/RemoteControl/Source/RemoteControl/Private/Factories/RemoteControlMaskingFactories.h
@@ -0,0 +1,90 @@
+// Copyright Epic Games, Inc. All Rights Reserved.
+
+#pragma once
+
+#include "CoreMinimal.h"
+#include "Factories/IRemoteControlMaskingFactory.h"
+
+class FVectorMaskingFactory : public IRemoteControlMaskingFactory
+{
+public:
+	static TSharedRef<IRemoteControlMaskingFactory> MakeInstance();
+
+	//~ Begin IRemoteControlMaskingFactory interface
+	virtual void ApplyMaskedValues(const TSharedRef<FRCMaskingOperation>& InMaskingOperation, bool bIsInteractive) override;
+	virtual void CacheRawValues(const TSharedRef<FRCMaskingOperation>& InMaskingOperation) override;
+	virtual bool SupportsExposedEntity(UScriptStruct* ScriptStruct) const override;
+	//~ End IRemoteControlMaskingFactory interface
+};
+
+class FVector4MaskingFactory : public IRemoteControlMaskingFactory
+{
+public:
+	static TSharedRef<IRemoteControlMaskingFactory> MakeInstance();
+
+	//~ Begin IRemoteControlMaskingFactory interface
+	virtual void ApplyMaskedValues(const TSharedRef<FRCMaskingOperation>& InMaskingOperation, bool bIsInteractive) override;
+	virtual void CacheRawValues(const TSharedRef<FRCMaskingOperation>& InMaskingOperation) override;
+	virtual bool SupportsExposedEntity(UScriptStruct* ScriptStruct) const override;
+	//~ End IRemoteControlMaskingFactory interface
+};
+
+class FIntVectorMaskingFactory : public IRemoteControlMaskingFactory
+{
+public:
+	static TSharedRef<IRemoteControlMaskingFactory> MakeInstance();
+
+	//~ Begin IRemoteControlMaskingFactory interface
+	virtual void ApplyMaskedValues(const TSharedRef<FRCMaskingOperation>& InMaskingOperation, bool bIsInteractive) override;
+	virtual void CacheRawValues(const TSharedRef<FRCMaskingOperation>& InMaskingOperation) override;
+	virtual bool SupportsExposedEntity(UScriptStruct* ScriptStruct) const override;
+	//~ End IRemoteControlMaskingFactory interface
+};
+
+class FIntVector4MaskingFactory : public IRemoteControlMaskingFactory
+{
+public:
+	static TSharedRef<IRemoteControlMaskingFactory> MakeInstance();
+
+	//~ Begin IRemoteControlMaskingFactory interface
+	virtual void ApplyMaskedValues(const TSharedRef<FRCMaskingOperation>& InMaskingOperation, bool bIsInteractive) override;
+	virtual void CacheRawValues(const TSharedRef<FRCMaskingOperation>& InMaskingOperation) override;
+	virtual bool SupportsExposedEntity(UScriptStruct* ScriptStruct) const override;
+	//~ End IRemoteControlMaskingFactory interface
+};
+
+class FRotatorMaskingFactory : public IRemoteControlMaskingFactory
+{
+public:
+	static TSharedRef<IRemoteControlMaskingFactory> MakeInstance();
+
+	//~ Begin IRemoteControlMaskingFactory interface
+	virtual void ApplyMaskedValues(const TSharedRef<FRCMaskingOperation>& InMaskingOperation, bool bIsInteractive) override;
+	virtual void CacheRawValues(const TSharedRef<FRCMaskingOperation>& InMaskingOperation) override;
+	virtual bool SupportsExposedEntity(UScriptStruct* ScriptStruct) const override;
+	//~ End IRemoteControlMaskingFactory interface
+};
+
+class FColorMaskingFactory : public IRemoteControlMaskingFactory
+{
+public:
+	static TSharedRef<IRemoteControlMaskingFactory> MakeInstance();
+
+	//~ Begin IRemoteControlMaskingFactory interface
+	virtual void ApplyMaskedValues(const TSharedRef<FRCMaskingOperation>& InMaskingOperation, bool bIsInteractive) override;
+	virtual void CacheRawValues(const TSharedRef<FRCMaskingOperation>& InMaskingOperation) override;
+	virtual bool SupportsExposedEntity(UScriptStruct* ScriptStruct) const override;
+	//~ End IRemoteControlMaskingFactory interface
+};
+
+class FLinearColorMaskingFactory : public IRemoteControlMaskingFactory
+{
+public:
+	static TSharedRef<IRemoteControlMaskingFactory> MakeInstance();
+
+	//~ Begin IRemoteControlMaskingFactory interface
+	virtual void ApplyMaskedValues(const TSharedRef<FRCMaskingOperation>& InMaskingOperation, bool bIsInteractive) override;
+	virtual void CacheRawValues(const TSharedRef<FRCMaskingOperation>& InMaskingOperation) override;
+	virtual bool SupportsExposedEntity(UScriptStruct* ScriptStruct) const override;
+	//~ End IRemoteControlMaskingFactory interface
+};
diff --git a/Engine/Plugins/VirtualProduction/RemoteControl/Source/RemoteControl/Private/RemoteControlModule.cpp b/Engine/Plugins/VirtualProduction/RemoteControl/Source/RemoteControl/Private/RemoteControlModule.cpp
--- a/Engine/Plugins/VirtualProduction/RemoteControl/Source/RemoteControl/Private/RemoteControlModule.cpp
+++ b/Engine/Plugins/VirtualProduction/RemoteControl/Source/RemoteControl/Private/RemoteControlModule.cpp
@@ -7,6 +7,7 @@
 #include "Backends/CborStructSerializerBackend.h"
 #include "Components/ActorComponent.h"
 #include "Components/LightComponent.h"
+#include "Factories/RemoteControlMaskingFactories.h"
 #include "Features/IModularFeatures.h"
 #include "IRemoteControlInterceptionFeature.h"
 #include "IRemoteControlModule.h"
@@ -591,6 +592,9 @@ void FRemoteControlModule::StartupModule()
 	
 	// Register Property Factories
 	RegisterEntityFactory(FRemoteControlInstanceMaterial::StaticStruct()->GetFName(), FRemoteControlInstanceMaterialFactory::MakeInstance());
+
+	// Register Masking Factories
+	RegisterMaskingFactories();
 }
 
 void FRemoteControlModule::ShutdownModule()
@@ -734,37 +738,46 @@ void FRemoteControlModule::PerformMasking(const TSharedRef<FRCMaskingOperation>&
 		return;
 	}
 
-	if (ActiveMaskingOperations.Contains(InMaskingOperation))
+	if (const FStructProperty* StructProperty = CastField<FStructProperty>(InMaskingOperation->ObjectRef.Property.Get()))
 	{
-		constexpr bool bIsInteractive = true;
+		if (const TSharedPtr<IRemoteControlMaskingFactory>* MaskingFactory = MaskingFactories.Find(StructProperty->Struct))
+		{
+			if (ActiveMaskingOperations.Contains(InMaskingOperation))
+			{
+				constexpr bool bIsInteractive = true;
 
-		ApplyMaskedValues(InMaskingOperation->OperationId, bIsInteractive);
+				(*MaskingFactory)->ApplyMaskedValues(InMaskingOperation, bIsInteractive);
 
-		ActiveMaskingOperations.Remove(InMaskingOperation);
+				ActiveMaskingOperations.Remove(InMaskingOperation);
+			}
+			else
+			{
+				ActiveMaskingOperations.Add(InMaskingOperation);
+
+				(*MaskingFactory)->CacheRawValues(InMaskingOperation);
+			}
+		}
 	}
-	else
-	{
-		ActiveMaskingOperations.Add(InMaskingOperation);
+}
 
-		CacheRawValues(InMaskingOperation->OperationId);
+void FRemoteControlModule::RegisterMaskingFactoryForType(UScriptStruct* RemoteControlPropertyType, const TSharedPtr<IRemoteControlMaskingFactory>& InMaskingFactory)
+{
+	if (!MaskingFactories.Contains(RemoteControlPropertyType))
+	{
+		MaskingFactories.Add(RemoteControlPropertyType, InMaskingFactory);
 	}
 }
 
+void FRemoteControlModule::UnregisterMaskingFactoryForType(UScriptStruct* RemoteControlPropertyType)
+{
+	MaskingFactories.Remove(RemoteControlPropertyType);
+}
+
 bool FRemoteControlModule::SupportsMasking(const FProperty* InProperty) const
 {
-	/** Holds the set of structs supported by masking. */
-	static const TSet<UScriptStruct*> MaskingSupportedStructs = { TBaseStructure<FColor>::Get()
-		, TBaseStructure<FLinearColor>::Get()
-		, TBaseStructure<FRotator>::Get()
-		, TBaseStructure<FVector>::Get()
-		, TBaseStructure<FVector4>::Get()
-		, TBaseStructure<FIntVector>::Get()
-		, TBaseStructure<FIntVector4>::Get()
-	};
-	
 	if (const FStructProperty* StructProperty = CastField<FStructProperty>(InProperty))
 	{
-		return MaskingSupportedStructs.Contains(StructProperty->Struct);
+		return MaskingFactories.Contains(StructProperty->Struct);
 	}
 
 	return false;
@@ -1817,245 +1830,6 @@ int32 FRemoteControlModule::EndManualEditorTransaction(const FGuid& TransactionI
 	return INDEX_NONE;
 }
 
-void FRemoteControlModule::ApplyMaskedValues(const FGuid& FromMaskingOperation, bool bIsInteractive)
-{
-	for (const TSharedPtr<FRCMaskingOperation>& MaskingOperation : ActiveMaskingOperations)
-	{
-		if (MaskingOperation->OperationId == FromMaskingOperation)
-		{
-			const bool bHasMaskA = MaskingOperation->HasMask(ERCMask::MaskA);
-			const bool bHasMaskB = MaskingOperation->HasMask(ERCMask::MaskB);
-			const bool bHasMaskC = MaskingOperation->HasMask(ERCMask::MaskC);
-			const bool bHasMaskD = MaskingOperation->HasMask(ERCMask::MaskD);
-
-			if (FStructProperty* ToStructProp = CastField<FStructProperty>(MaskingOperation->ObjectRef.Property.Get()))
-			{
-				if (UObject* OwningObject = MaskingOperation->ObjectRef.Object.Get())
-				{
-#if WITH_EDITOR
-					OwningObject->PreEditChange(ToStructProp);
-					OwningObject->Modify();
-#endif // WITH_EDITOR
-
-					switch (*ToStructProp->Struct->GetFName().ToEName())
-					{
-						case NAME_Vector:
-						{
-							if (const FVector* VectorProp = ToStructProp->ContainerPtrToValuePtr<FVector>(OwningObject))
-							{
-								FVector MaskedVector;
-
-								MaskedVector.X = bHasMaskA ? VectorProp->X : MaskingOperation->PreMaskingCache.X;
-								MaskedVector.Y = bHasMaskB ? VectorProp->Y : MaskingOperation->PreMaskingCache.Y;
-								MaskedVector.Z = bHasMaskC ? VectorProp->Z : MaskingOperation->PreMaskingCache.Z;
-
-								ToStructProp->SetValue_InContainer(OwningObject, &MaskedVector);
-							}
-						}
-						break;
-						
-						case NAME_Vector4:
-						{
-							if (const FVector4* Vector4Prop = ToStructProp->ContainerPtrToValuePtr<FVector4>(OwningObject))
-							{
-								FVector4 MaskedVector4;
-
-								MaskedVector4.X = bHasMaskA ? Vector4Prop->X : MaskingOperation->PreMaskingCache.X;
-								MaskedVector4.Y = bHasMaskB ? Vector4Prop->Y : MaskingOperation->PreMaskingCache.Y;
-								MaskedVector4.Z = bHasMaskC ? Vector4Prop->Z : MaskingOperation->PreMaskingCache.Z;
-								MaskedVector4.W = bHasMaskD ? Vector4Prop->W : MaskingOperation->PreMaskingCache.W;
-
-								ToStructProp->SetValue_InContainer(OwningObject, &MaskedVector4);
-							}
-						}
-						break;
-
-						case NAME_IntVector:
-						{
-							if (const FIntVector* IntVectorProp = ToStructProp->ContainerPtrToValuePtr<FIntVector>(OwningObject))
-							{
-								FIntVector MaskedIntVector;
-
-								MaskedIntVector.X = bHasMaskA ? IntVectorProp->X : MaskingOperation->PreMaskingCache.X;
-								MaskedIntVector.Y = bHasMaskB ? IntVectorProp->Y : MaskingOperation->PreMaskingCache.Y;
-								MaskedIntVector.Z = bHasMaskC ? IntVectorProp->Z : MaskingOperation->PreMaskingCache.Z;
-
-								ToStructProp->SetValue_InContainer(OwningObject, &MaskedIntVector);
-							}
-						}
-						break;
-
-						case NAME_IntVector4:
-						{
-							if (const FIntVector4* IntVector4Prop = ToStructProp->ContainerPtrToValuePtr<FIntVector4>(OwningObject))
-							{
-								FIntVector4 MaskedIntVector4;
-
-								MaskedIntVector4.X = bHasMaskA ? IntVector4Prop->X : MaskingOperation->PreMaskingCache.X;
-								MaskedIntVector4.Y = bHasMaskB ? IntVector4Prop->Y : MaskingOperation->PreMaskingCache.Y;
-								MaskedIntVector4.Z = bHasMaskC ? IntVector4Prop->Z : MaskingOperation->PreMaskingCache.Z;
-								MaskedIntVector4.W = bHasMaskD ? IntVector4Prop->W : MaskingOperation->PreMaskingCache.W;
-
-								ToStructProp->SetValue_InContainer(OwningObject, &MaskedIntVector4);
-							}
-						}
-						break;
-
-						case NAME_Rotator:
-						{
-							if (const FRotator* RotatorProp = ToStructProp->ContainerPtrToValuePtr<FRotator>(OwningObject))
-							{
-								FRotator MaskedRotator;
-
-								MaskedRotator.Roll = bHasMaskA ? RotatorProp->Roll : MaskingOperation->PreMaskingCache.X;
-								MaskedRotator.Pitch = bHasMaskB ? RotatorProp->Pitch : MaskingOperation->PreMaskingCache.Y;
-								MaskedRotator.Yaw = bHasMaskC ? RotatorProp->Yaw : MaskingOperation->PreMaskingCache.Z;
-
-								ToStructProp->SetValue_InContainer(OwningObject, &MaskedRotator);
-							}
-						}
-						break;
-
-						case NAME_Color:
-						{
-							if (const FColor* ColorProp = ToStructProp->ContainerPtrToValuePtr<FColor>(OwningObject))
-							{
-								FColor MaskedColor;
-
-								MaskedColor.R = bHasMaskA ? ColorProp->R : MaskingOperation->PreMaskingCache.X;
-								MaskedColor.G = bHasMaskB ? ColorProp->G : MaskingOperation->PreMaskingCache.Y;
-								MaskedColor.B = bHasMaskC ? ColorProp->B : MaskingOperation->PreMaskingCache.Z;
-
-								ToStructProp->SetValue_InContainer(OwningObject, &MaskedColor);
-							}
-						}
-						break;
-
-						case NAME_LinearColor:
-						{
-							if (const FLinearColor* LinearColorProp = ToStructProp->ContainerPtrToValuePtr<FLinearColor>(OwningObject))
-							{
-								FLinearColor MaskedLinearColor;
-
-								MaskedLinearColor.R = bHasMaskA ? LinearColorProp->R : MaskingOperation->PreMaskingCache.X;
-								MaskedLinearColor.G = bHasMaskB ? LinearColorProp->G : MaskingOperation->PreMaskingCache.Y;
-								MaskedLinearColor.B = bHasMaskC ? LinearColorProp->B : MaskingOperation->PreMaskingCache.Z;
-								MaskedLinearColor.A = bHasMaskD ? LinearColorProp->A : MaskingOperation->PreMaskingCache.W;
-
-								ToStructProp->SetValue_InContainer(OwningObject, &MaskedLinearColor);
-							}
-						}
-						break;
-					}
-
-#if WITH_EDITOR
-					FPropertyChangedEvent ChangeEvent(ToStructProp, bIsInteractive ? EPropertyChangeType::Interactive : EPropertyChangeType::ValueSet);
-					OwningObject->PostEditChangeProperty(ChangeEvent);
-#endif // WITH_EDITOR
-				}
-			}
-		}
-	}
-}
-
-void FRemoteControlModule::CacheRawValues(const FGuid& FromMaskingOperation)
-{
-	for (const TSharedPtr<FRCMaskingOperation>& MaskingOperation : ActiveMaskingOperations)
-	{
-		if (MaskingOperation->OperationId == FromMaskingOperation)
-		{
-			if (const FStructProperty* FromStructProp = CastField<FStructProperty>(MaskingOperation->ObjectRef.Property.Get()))
-			{
-				if (const UObject* OwningObject = MaskingOperation->ObjectRef.Object.Get())
-				{
-					switch (*FromStructProp->Struct->GetFName().ToEName())
-					{
-						case NAME_Vector:
-						{
-							if (const FVector* VectorProp = FromStructProp->ContainerPtrToValuePtr<FVector>(OwningObject))
-							{
-								MaskingOperation->PreMaskingCache.X = VectorProp->X;
-								MaskingOperation->PreMaskingCache.Y = VectorProp->Y;
-								MaskingOperation->PreMaskingCache.Z = VectorProp->Z;
-							}
-						}
-						break;
-						
-						case NAME_Vector4:
-						{
-							if (const FVector4* Vector4Prop = FromStructProp->ContainerPtrToValuePtr<FVector4>(OwningObject))
-							{
-								MaskingOperation->PreMaskingCache.X = Vector4Prop->X;
-								MaskingOperation->PreMaskingCache.Y = Vector4Prop->Y;
-								MaskingOperation->PreMaskingCache.Z = Vector4Prop->Z;
-								MaskingOperation->PreMaskingCache.W = Vector4Prop->W;
-							}
-						}
-						break;
-						
-						case NAME_IntVector:
-						{
-							if (const FIntVector* IntVectorProp = FromStructProp->ContainerPtrToValuePtr<FIntVector>(OwningObject))
-							{
-								MaskingOperation->PreMaskingCache.X = IntVectorProp->X;
-								MaskingOperation->PreMaskingCache.Y = IntVectorProp->Y;
-								MaskingOperation->PreMaskingCache.Z = IntVectorProp->Z;
-							}
-						}
-						break;
-						
-						case NAME_IntVector4:
-						{
-							if (const FIntVector4* IntVector4Prop = FromStructProp->ContainerPtrToValuePtr<FIntVector4>(OwningObject))
-							{
-								MaskingOperation->PreMaskingCache.X = IntVector4Prop->X;
-								MaskingOperation->PreMaskingCache.Y = IntVector4Prop->Y;
-								MaskingOperation->PreMaskingCache.Z = IntVector4Prop->Z;
-								MaskingOperation->PreMaskingCache.W = IntVector4Prop->W;
-							}
-						}
-						break;
-						
-						case NAME_Rotator:
-						{
-							if (const FRotator* RotatorProp = FromStructProp->ContainerPtrToValuePtr<FRotator>(OwningObject))
-							{
-								MaskingOperation->PreMaskingCache.X = RotatorProp->Roll;
-								MaskingOperation->PreMaskingCache.Y = RotatorProp->Pitch;
-								MaskingOperation->PreMaskingCache.Z = RotatorProp->Yaw;
-							}
-						}
-						break;
-
-						case NAME_Color:
-						{
-							if (const FColor* ColorProp = FromStructProp->ContainerPtrToValuePtr<FColor>(OwningObject))
-							{
-								MaskingOperation->PreMaskingCache.X = ColorProp->R;
-								MaskingOperation->PreMaskingCache.Y = ColorProp->G;
-								MaskingOperation->PreMaskingCache.Z = ColorProp->B;
-							}
-						}
-						break;
-
-						case NAME_LinearColor:
-						{
-							if (const FLinearColor* LinearColorProp = FromStructProp->ContainerPtrToValuePtr<FLinearColor>(OwningObject))
-							{
-								MaskingOperation->PreMaskingCache.X = LinearColorProp->R;
-								MaskingOperation->PreMaskingCache.Y = LinearColorProp->G;
-								MaskingOperation->PreMaskingCache.Z = LinearColorProp->B;
-								MaskingOperation->PreMaskingCache.W = LinearColorProp->A;
-							}
-						}
-						break;
-					}
-				}
-			}
-		}
-	}
-}
-
 void FRemoteControlModule::CachePresets() const
 {
 	TArray<FAssetData> Assets;
@@ -2334,6 +2108,17 @@ bool FRemoteControlModule::DeserializeDeltaModificationData(const FRCObjectRefer
 	return bSuccess;
 }
 
+void FRemoteControlModule::RegisterMaskingFactories()
+{
+	RegisterMaskingFactoryForType(TBaseStructure<FVector>::Get(), FVectorMaskingFactory::MakeInstance());
+	RegisterMaskingFactoryForType(TBaseStructure<FVector4>::Get(), FVector4MaskingFactory::MakeInstance());
+	RegisterMaskingFactoryForType(TBaseStructure<FIntVector>::Get(), FIntVectorMaskingFactory::MakeInstance());
+	RegisterMaskingFactoryForType(TBaseStructure<FIntVector4>::Get(), FIntVector4MaskingFactory::MakeInstance());
+	RegisterMaskingFactoryForType(TBaseStructure<FRotator>::Get(), FRotatorMaskingFactory::MakeInstance());
+	RegisterMaskingFactoryForType(TBaseStructure<FColor>::Get(), FColorMaskingFactory::MakeInstance());
+	RegisterMaskingFactoryForType(TBaseStructure<FLinearColor>::Get(), FLinearColorMaskingFactory::MakeInstance());
+}
+
 #if WITH_EDITOR
 void FRemoteControlModule::EndOngoingModificationIfMismatched(uint32 TypeHash)
 {
diff --git a/Engine/Plugins/VirtualProduction/RemoteControl/Source/RemoteControl/Private/RemoteControlModule.h b/Engine/Plugins/VirtualProduction/RemoteControl/Source/RemoteControl/Private/RemoteControlModule.h
--- a/Engine/Plugins/VirtualProduction/RemoteControl/Source/RemoteControl/Private/RemoteControlModule.h
+++ b/Engine/Plugins/VirtualProduction/RemoteControl/Source/RemoteControl/Private/RemoteControlModule.h
@@ -5,6 +5,7 @@
 #include "AssetRegistry/AssetData.h"
 #include "CoreMinimal.h"
 #include "Engine/EngineTypes.h"
+#include "Factories/IRemoteControlMaskingFactory.h"
 
 class IRemoteControlInterceptionFeatureProcessor;
 
@@ -30,6 +31,8 @@ public:
 	virtual void UnregisterEmbeddedPreset(FName Name) override;
 	virtual void UnregisterEmbeddedPreset(URemoteControlPreset* Preset) override;
 	virtual void PerformMasking(const TSharedRef<FRCMaskingOperation>& InMaskingOperation) override;
+	virtual void RegisterMaskingFactoryForType(UScriptStruct* RemoteControlPropertyType, const TSharedPtr<IRemoteControlMaskingFactory>& InMaskingFactory) override;
+	virtual void UnregisterMaskingFactoryForType(UScriptStruct* RemoteControlPropertyType) override;
 	virtual bool SupportsMasking(const FProperty* InProperty) const override;
 	virtual bool ResolveCall(const FString& ObjectPath, const FString& FunctionName, FRCCallReference& OutCallRef, FString* OutErrorText) override;
 	virtual bool InvokeCall(FRCCall& InCall, ERCPayloadType InPayloadType = ERCPayloadType::Json, const TArray<uint8>& InInterceptPayload = TArray<uint8>()) override;
@@ -65,11 +68,6 @@ public:
 
 private:
 
-	/** Applies masked values to the given struct property. */
-	void ApplyMaskedValues(const FGuid& FromMaskingOperation, bool bIsInteractive);
-	/** Caches premasking values from the given struct property. */
-	void CacheRawValues(const FGuid& ToMaskingOperation);
-
 	/** Cache all presets in the project for the ResolvePreset function. */
 	void CachePresets() const;
 	
@@ -96,6 +94,11 @@ private:
 	 */
 	static bool DeserializeDeltaModificationData(const FRCObjectReference& ObjectAccess, IStructDeserializerBackend& Backend, ERCModifyOperation Operation, TArray<uint8>& OutData);
 
+	/**
+	 * Register(s) masking factories of supported types.
+	 */
+	void RegisterMaskingFactories();
+
 #if WITH_EDITOR
 	/**
 	 * End the ongoing modification if one exists and is a mismatch for the new object to edit.
@@ -190,6 +193,9 @@ private:
 	/** Map of the factories which is responsible for the Remote Control property creation */
 	TMap<FName, TSharedPtr<IRemoteControlPropertyFactory>> EntityFactories;
 
+	/** Map of the factories which is responsible for the Remote Control property masking. */
+	TMap<TWeakObjectPtr<UScriptStruct>, TSharedPtr<IRemoteControlMaskingFactory>> MaskingFactories;
+
 	/** Holds the set of active masking operations. */
 	TSet<TSharedPtr<FRCMaskingOperation>> ActiveMaskingOperations;
 };
diff --git a/Engine/Plugins/VirtualProduction/RemoteControl/Source/RemoteControl/Private/RemoteControlProtocolBinding.cpp b/Engine/Plugins/VirtualProduction/RemoteControl/Source/RemoteControl/Private/RemoteControlProtocolBinding.cpp
--- a/Engine/Plugins/VirtualProduction/RemoteControl/Source/RemoteControl/Private/RemoteControlProtocolBinding.cpp
+++ b/Engine/Plugins/VirtualProduction/RemoteControl/Source/RemoteControl/Private/RemoteControlProtocolBinding.cpp
@@ -3,6 +3,7 @@
 #include "RemoteControlProtocolBinding.h"
 
 #include "CborWriter.h"
+#include "Factories/IRemoteControlMaskingFactory.h"
 #include "IRemoteControlModule.h"
 #include "RemoteControlPreset.h"
 #include "RemoteControlSettings.h"
@@ -14,14 +15,6 @@
 #include "UObject/StructOnScope.h"
 #include "UObject/TextProperty.h"
 
-#if WITH_EDITOR
-
-#include "Widgets/Images/SImage.h"
-#include "Widgets/Input/SButton.h"
-
-#endif // WITH_EDITOR
-
-
 #define LOCTEXT_NAMESPACE "RemoteControl"
 
 namespace EntityInterpolation
@@ -670,6 +663,11 @@ bool FRemoteControlProtocolEntity::ApplyProtocolValueToProperty(double InProtoco
 		return true;
 	}
 
+	if (RemoteControlProperty->GetActiveMasks() == ERCMask::NoMask)
+	{
+		return true;
+	}
+
 	FProperty* Property = RemoteControlProperty->GetProperty();
 	if (!Property)
 	{
@@ -742,65 +740,16 @@ void FRemoteControlProtocolEntity::ResetDefaultBindingState()
 
 #if WITH_EDITOR
 
-TSharedRef<SWidget> FRemoteControlProtocolEntity::GetWidget(const FName& ForColumnName)
+const FName FRemoteControlProtocolEntity::GetPropertyName(const FName& ForColumnName)
 {
-	RegisterWidgets();
-
-	TSharedPtr<SWidget> ProtocolWidget = SNullWidget::NullWidget;
+	RegisterProperties();
 
-	for (const TPair<FName, TSharedPtr<SWidget>>& Widget : Widgets)
+	if (const FName* PropertyName = ColumnsToProperties.Find(ForColumnName))
 	{
-		if (Widget.Key == ForColumnName && Widget.Value.IsValid())
-		{
-			return Widget.Value.ToSharedRef();
-		}
+		return *PropertyName;
 	}
 
-	return ProtocolWidget.ToSharedRef();
-}
-
-void FRemoteControlProtocolEntity::RegisterWidgets()
-{
-	if (!Widgets.Contains(FRemoteControlDefaultProtocolColumns::BindingStatus))
-	{
-		TSharedPtr<SWidget> BindingStatusWidget = SNew(SButton)
-			.ButtonStyle(FAppStyle::Get(), "HoverHintOnly")
-			.ToolTipText(LOCTEXT("RecordingButtonToolTip", "Status of the protocol entity binding"))
-			.ForegroundColor(FSlateColor::UseForeground())
-			.OnClicked_Lambda([this]()
-				{
-					ToggleBindingStatus();
-
-					return FReply::Handled();
-				}
-			)
-			.Content()
-			[
-				SNew(SImage)
-				.ColorAndOpacity_Lambda([this]()
-					{
-						const ERCBindingStatus ActiveBindingStatus = GetBindingStatus();
-				
-						switch (ActiveBindingStatus)
-						{
-							case ERCBindingStatus::Awaiting:
-								return FLinearColor::Red;
-							case ERCBindingStatus::Bound:
-								return FLinearColor::Green;
-							case ERCBindingStatus::Unassigned:
-								return FLinearColor::Gray;
-							default:
-								checkNoEntry();
-						}
-
-						return FLinearColor::Black;
-					}
-				)
-				.Image(FAppStyle::Get().GetBrush(TEXT("Icons.FilledCircle")))
-			];
-
-		Widgets.Add(FRemoteControlDefaultProtocolColumns::BindingStatus, BindingStatusWidget);
-	}
+	return NAME_None;
 }
 
 #endif // WITH_EDITOR
diff --git a/Engine/Plugins/VirtualProduction/RemoteControl/Source/RemoteControl/Public/Factories/IRemoteControlMaskingFactory.h b/Engine/Plugins/VirtualProduction/RemoteControl/Source/RemoteControl/Public/Factories/IRemoteControlMaskingFactory.h
new file mode 100644
--- /dev/null
+++ b/Engine/Plugins/VirtualProduction/RemoteControl/Source/RemoteControl/Public/Factories/IRemoteControlMaskingFactory.h
@@ -0,0 +1,92 @@
+// Copyright Epic Games, Inc. All Rights Reserved.
+
+#pragma once
+
+#include "CoreMinimal.h"
+#include "IRemoteControlModule.h"
+
+/**
+ * Reference to a UObject or one of its properties for the purpose of masking.
+ */
+struct FRCMaskingOperation
+{
+	FRCMaskingOperation() = default;
+
+	explicit FRCMaskingOperation(FRCFieldPathInfo InPathInfo, UObject* InObject)
+		: OperationId(FGuid::NewGuid())
+		, ObjectRef(ERCAccess::NO_ACCESS, InObject, InPathInfo)
+	{
+		check(InObject);
+	}
+
+	explicit FRCMaskingOperation(const FRCObjectReference& InObjectRef)
+		: OperationId(FGuid::NewGuid())
+		, ObjectRef(InObjectRef)
+	{
+	}
+
+	bool HasMask(ERCMask InMaskBit) const
+	{
+		return (Masks & InMaskBit) != ERCMask::NoMask;
+	}
+
+	bool IsValid() const
+	{
+		return OperationId.IsValid() && ObjectRef.IsValid();
+	}
+
+	friend bool operator==(const FRCMaskingOperation& LHS, const FRCMaskingOperation& RHS)
+	{
+		return LHS.OperationId == RHS.OperationId && LHS.ObjectRef == RHS.ObjectRef;
+	}
+
+	friend uint32 GetTypeHash(const FRCMaskingOperation& MaskingOperation)
+	{
+		return HashCombine(GetTypeHash(MaskingOperation.OperationId), GetTypeHash(MaskingOperation.ObjectRef));
+	}
+
+public:
+
+	/** Unique identifier of the operation being performed. */
+	FGuid OperationId;
+
+	/** Masks to be applied. */
+	ERCMask Masks = ERCMask::NoMask;
+
+	/** Holds Object reference. */
+	FRCObjectReference ObjectRef;
+
+	/** Holds the state of this RC property before applying any masking. */
+	FVector4 PreMaskingCache = FVector4::Zero();
+};
+
+/**
+ * Factory which is responsible for masking support for FRemoteControlProperty.
+ */
+class IRemoteControlMaskingFactory : public TSharedFromThis<IRemoteControlMaskingFactory>
+{
+public:
+
+	/** Virtual destructor */
+	virtual ~IRemoteControlMaskingFactory(){}
+
+	/**
+	 * Applies masked values to the given struct property.
+	 * @param InMaskingOperation Shared reference of the masking operation to perform.
+	 */
+	virtual void ApplyMaskedValues(const TSharedRef<FRCMaskingOperation>& InMaskingOperation, bool bIsInteractive) = 0;
+
+	/**
+	 * Caches premasking values from the given struct property.
+	 * @param InMaskingOperation Shared reference of the masking operation to perform.
+	 */
+	virtual void CacheRawValues(const TSharedRef<FRCMaskingOperation>& InMaskingOperation) = 0;
+
+	/**
+	 * Whether the factory support exposed entity.
+	 * @param ScriptStruct Static struct of exposed property.
+	 * @return true if the script struct is supported by given factory
+	 */
+	virtual bool SupportsExposedEntity(UScriptStruct* ScriptStruct) const = 0;
+};
+
diff --git a/Engine/Plugins/VirtualProduction/RemoteControl/Source/RemoteControl/Public/IRemoteControlModule.h b/Engine/Plugins/VirtualProduction/RemoteControl/Source/RemoteControl/Public/IRemoteControlModule.h
--- a/Engine/Plugins/VirtualProduction/RemoteControl/Source/RemoteControl/Public/IRemoteControlModule.h
+++ b/Engine/Plugins/VirtualProduction/RemoteControl/Source/RemoteControl/Public/IRemoteControlModule.h
@@ -11,12 +11,14 @@
 
 REMOTECONTROL_API DECLARE_LOG_CATEGORY_EXTERN(LogRemoteControl, Log, All);
 
+class IRemoteControlMaskingFactory;
 class IStructDeserializerBackend;
 class IStructSerializerBackend;
 class URemoteControlPreset;
 PRAGMA_DISABLE_DEPRECATION_WARNINGS
 struct FExposedProperty;
 PRAGMA_ENABLE_DEPRECATION_WARNINGS
+struct FRCMaskingOperation;
 struct FRemoteControlProperty;
 class IRemoteControlPropertyFactory;
 
@@ -212,61 +214,6 @@ struct FRCObjectReference
 	FRCFieldPathInfo PropertyPathInfo;
 };
 
-/**
- * Reference to a UObject or one of its properties for the purpose of masking.
- */
-struct FRCMaskingOperation
-{
-	FRCMaskingOperation() = default;
-
-	explicit FRCMaskingOperation(FRCFieldPathInfo InPathInfo, UObject* InObject)
-		: OperationId(FGuid::NewGuid())
-		, ObjectRef(ERCAccess::NO_ACCESS, InObject, InPathInfo)
-	{
-		check(InObject);
-	}
-	
-	explicit FRCMaskingOperation(const FRCObjectReference& InObjectRef)
-		: OperationId(FGuid::NewGuid())
-		, ObjectRef(InObjectRef)
-	{
-	}
-
-	bool HasMask(ERCMask InMaskBit) const
-	{
-		return (Masks & InMaskBit) != ERCMask::NoMask;
-	}
-
-	bool IsValid() const
-	{
-		return OperationId.IsValid() && ObjectRef.IsValid();
-	}
-
-	friend bool operator==(const FRCMaskingOperation& LHS, const FRCMaskingOperation& RHS)
-	{
-		return LHS.OperationId == RHS.OperationId && LHS.ObjectRef == RHS.ObjectRef;
-	}
-
-	friend uint32 GetTypeHash(const FRCMaskingOperation& MaskingOperation)
-	{
-		return HashCombine(GetTypeHash(MaskingOperation.OperationId), GetTypeHash(MaskingOperation.ObjectRef));
-	}
-
-public:
-
-	/** Unique identifier of the operation being performed. */
-	FGuid OperationId;
-
-	/** Masks to be applied. */
-	ERCMask Masks = ERCMask::NoMask;
-	
-	/** Holds Object reference. */
-	FRCObjectReference ObjectRef;
-
-	/** Holds the state of this RC property before applying any masking. */
-	FVector4 PreMaskingCache = FVector4::Zero();
-};
-
 /**
  * Interface for the remote control module.
  */
@@ -325,6 +272,16 @@ public:
 	 */
 	virtual void PerformMasking(const TSharedRef<FRCMaskingOperation>& InMaskingOperation) = 0;
 
+	/**
+	 * Register a masking factory to handle that masks the supported properties.
+	 */
+	virtual void RegisterMaskingFactoryForType(UScriptStruct* RemoteControlPropertyType, const TSharedPtr<IRemoteControlMaskingFactory>& InMaskingFactory) = 0;
+
+	/**
+	 * Unregister a previously registered masking factory.
+	 */
+	virtual void UnregisterMaskingFactoryForType(UScriptStruct* RemoteControlPropertyType) = 0;
+
 	/**
 	 * Returns true if the given property can be masked, false otherwise.
 	 */
diff --git a/Engine/Plugins/VirtualProduction/RemoteControl/Source/RemoteControl/Public/RemoteControlProtocolBinding.h b/Engine/Plugins/VirtualProduction/RemoteControl/Source/RemoteControl/Public/RemoteControlProtocolBinding.h
--- a/Engine/Plugins/VirtualProduction/RemoteControl/Source/RemoteControl/Public/RemoteControlProtocolBinding.h
+++ b/Engine/Plugins/VirtualProduction/RemoteControl/Source/RemoteControl/Public/RemoteControlProtocolBinding.h
@@ -25,10 +25,11 @@ using FGetProtocolMappingCallback = TFunctionRef<void(FRemoteControlProtocolMapp
 
 #if WITH_EDITOR
 
-namespace FRemoteControlDefaultProtocolColumns
-{
-	static FName BindingStatus = TEXT("BindingStatus");
-}
+#define EXPOSE_PROTOCOL_PROPERTY(ColumnName, ClassName, MemberName) \
+	if (!ColumnsToProperties.Contains(ColumnName)) \
+	{	\
+		ColumnsToProperties.Add(ColumnName, GET_MEMBER_NAME_CHECKED(ClassName, MemberName)); \
+	}
 
 #endif // WITH_EDITOR
 
@@ -427,11 +428,11 @@ public:
 
 #if WITH_EDITOR
 
-	/** Retrieves the widget corresponding to the given column name. */
-	TSharedRef<SWidget> GetWidget(const FName& ForColumnName);
-
-	/** Register(s) all the widgets of this protocol entity. */
-	virtual void RegisterWidgets();
+	/** Retrieves the name of property corresponding to the given column name. */
+	const FName GetPropertyName(const FName& ForColumnName);
+	
+	/** Register(s) all the properties (to be exposed) of this protocol entity. */
+	virtual void RegisterProperties() {};
 
 #endif // WITH_EDITOR
 
@@ -535,10 +536,14 @@ protected:
 	UPROPERTY()
 	TSet<FRemoteControlProtocolMapping> Mappings;
 	
-	/** 
-	 * Holds column specific widget(s) for this protocol.
+#if WITH_EDITOR
+
+	/**
+	 * Holds column specific properties for this protocol.
 	 */
-	TMap<FName, TSharedPtr<SWidget>> Widgets;
+	TMap<FName, FName> ColumnsToProperties;
+
+#endif // WITH_EDITOR
 
 private:
 	/** Binding status of this protocol entity */
diff --git a/Engine/Plugins/VirtualProduction/RemoteControl/Source/RemoteControl/RemoteControl.Build.cs b/Engine/Plugins/VirtualProduction/RemoteControl/Source/RemoteControl/RemoteControl.Build.cs
--- a/Engine/Plugins/VirtualProduction/RemoteControl/Source/RemoteControl/RemoteControl.Build.cs
+++ b/Engine/Plugins/VirtualProduction/RemoteControl/Source/RemoteControl/RemoteControl.Build.cs
@@ -37,8 +37,6 @@ public class RemoteControl : ModuleRules
 					"AssetTools",
 					"DeveloperSettings",
 					"MessageLog",
-					"Slate",
-					"SlateCore",
 					"UnrealEd",
 				}
 			);
diff --git a/Engine/Plugins/VirtualProduction/RemoteControl/Source/RemoteControlProtocol/Private/RemoteControlProtocol.cpp b/Engine/Plugins/VirtualProduction/RemoteControl/Source/RemoteControlProtocol/Private/RemoteControlProtocol.cpp
--- a/Engine/Plugins/VirtualProduction/RemoteControl/Source/RemoteControlProtocol/Private/RemoteControlProtocol.cpp
+++ b/Engine/Plugins/VirtualProduction/RemoteControl/Source/RemoteControlProtocol/Private/RemoteControlProtocol.cpp
@@ -109,17 +109,6 @@ TFunction<bool(FRemoteControlProtocolEntityWeakPtr InProtocolEntityWeakPtr)> FRe
 	};
 }
 
-#if WITH_EDITOR
-
-void FRemoteControlProtocol::RegisterColumns()
-{
-	REGISTER_COLUMN(FRemoteControlDefaultProtocolColumns::BindingStatus
-		, NSLOCTEXT("RemoteControlProtocol", "RCPresetBindingStatusColumnHeader", "REC")
-		, 33.f);
-}
-
-#endif // WITH_EDITOR
-
 FProperty* IRemoteControlProtocol::GetRangeInputTemplateProperty() const
 {
 	FProperty* RangeInputTemplateProperty = nullptr;
diff --git a/Engine/Plugins/VirtualProduction/RemoteControl/Source/RemoteControlProtocol/Public/IRemoteControlProtocol.h b/Engine/Plugins/VirtualProduction/RemoteControl/Source/RemoteControlProtocol/Public/IRemoteControlProtocol.h
--- a/Engine/Plugins/VirtualProduction/RemoteControl/Source/RemoteControlProtocol/Public/IRemoteControlProtocol.h
+++ b/Engine/Plugins/VirtualProduction/RemoteControl/Source/RemoteControlProtocol/Public/IRemoteControlProtocol.h
@@ -18,7 +18,8 @@ namespace ProtocolColumnConstants
 {
 	static float ColumnSizeMicro = 0.075f;
 	static float ColumnSizeMini = 0.15f;
-	static float ColumnSizeSmall = 0.3f;
+	static float ColumnSizeSmall = 0.21f;
+	static float ColumnSizeNormal = 0.3f;
 	static float ColumnSizeMedium = 0.5f;
 	static float ColumnSizeLarge = 0.8f;
 	static float ColumnSizeExtraLarge = 1.f;
diff --git a/Engine/Plugins/VirtualProduction/RemoteControl/Source/RemoteControlProtocol/Public/RemoteControlProtocol.h b/Engine/Plugins/VirtualProduction/RemoteControl/Source/RemoteControlProtocol/Public/RemoteControlProtocol.h
--- a/Engine/Plugins/VirtualProduction/RemoteControl/Source/RemoteControlProtocol/Public/RemoteControlProtocol.h
+++ b/Engine/Plugins/VirtualProduction/RemoteControl/Source/RemoteControlProtocol/Public/RemoteControlProtocol.h
@@ -47,7 +47,7 @@ public:
 protected:
 
 	/** Populates protocol specific columns. */
-	virtual void RegisterColumns();
+	virtual void RegisterColumns() {};
 
 #endif // WITH_EDITOR
 
diff --git a/Engine/Plugins/VirtualProduction/RemoteControl/Source/RemoteControlUI/Private/UI/SRCPanelExposedEntitiesList.cpp b/Engine/Plugins/VirtualProduction/RemoteControl/Source/RemoteControlUI/Private/UI/SRCPanelExposedEntitiesList.cpp
--- a/Engine/Plugins/VirtualProduction/RemoteControl/Source/RemoteControlUI/Private/UI/SRCPanelExposedEntitiesList.cpp
+++ b/Engine/Plugins/VirtualProduction/RemoteControl/Source/RemoteControlUI/Private/UI/SRCPanelExposedEntitiesList.cpp
@@ -37,9 +37,10 @@
 #define LOCTEXT_NAMESPACE "RemoteControlPanelEntitiesList"
 
 TSet<FName> SRCPanelExposedEntitiesList::DefaultProtocolColumns = {
-//	FRemoteControlPresetColumns::LinkIdentifier,
-	FRemoteControlPresetColumns::Mask,
-	FRemoteControlPresetColumns::Status,
+//	RemoteControlPresetColumns::LinkIdentifier,
+	RemoteControlPresetColumns::Mask,
+	RemoteControlPresetColumns::Status,
+	RemoteControlPresetColumns::BindingStatus,
 };
 
 /**
@@ -92,7 +93,7 @@ public:
 
 		if (Entity.IsValid())
 		{
-			if (Entity->HasChildren() && InColumnName == FRemoteControlPresetColumns::Description)
+			if (Entity->HasChildren() && InColumnName == RemoteControlPresetColumns::Description)
 			{
 				// -- Row is for TreeView --
 				SHorizontalBox::FSlot* InnerContentSlotNativePtr = nullptr;
@@ -125,9 +126,9 @@ public:
 			}
 			else
 			{
-				const bool bAlignCenter = InColumnName == FRemoteControlPresetColumns::Status || InColumnName == FRemoteControlDefaultProtocolColumns::BindingStatus;
+				const bool bAlignCenter = InColumnName == RemoteControlPresetColumns::Status || InColumnName == RemoteControlPresetColumns::BindingStatus;
 				
-				const bool bAlignLeft = InColumnName == FRemoteControlPresetColumns::Mask;
+				const bool bAlignLeft = InColumnName == RemoteControlPresetColumns::Mask;
 
 				return SNew(SBox)
 					.Padding(FMargin(4.f, 2.f))
@@ -270,24 +271,24 @@ void SRCPanelExposedEntitiesList::Construct(const FArguments& InArgs, URemoteCon
 			SAssignNew(FieldsHeaderRow, SHeaderRow)
 			.Style(&RCPanelStyle->HeaderRowStyle)
 
-			+ SHeaderRow::Column(FRemoteControlPresetColumns::DragDropHandle)
+			+ SHeaderRow::Column(RemoteControlPresetColumns::DragDropHandle)
 			.DefaultLabel(LOCTEXT("RCPresetDragDropHandleColumnHeader", ""))
 			.FixedWidth(25.f)
 			.HeaderContentPadding(RCPanelStyle->HeaderRowPadding)
 
-			+ SHeaderRow::Column(FRemoteControlPresetColumns::Description)
+			+ SHeaderRow::Column(RemoteControlPresetColumns::Description)
 			.DefaultLabel(LOCTEXT("RCPresetDescColumnHeader", "Description"))
 			.HAlignHeader(HAlign_Center)
 			.FillWidth(0.5f)
 			.HeaderContentPadding(RCPanelStyle->HeaderRowPadding)
 
-			+ SHeaderRow::Column(FRemoteControlPresetColumns::Value)
+			+ SHeaderRow::Column(RemoteControlPresetColumns::Value)
 			.DefaultLabel(LOCTEXT("RCPresetValueColumnHeader", "Value"))
 			.HAlignHeader(HAlign_Center)
 			.FillWidth(0.5f)
 			.HeaderContentPadding(RCPanelStyle->HeaderRowPadding)
 
-			+ SHeaderRow::Column(FRemoteControlPresetColumns::Reset)
+			+ SHeaderRow::Column(RemoteControlPresetColumns::Reset)
 			.DefaultLabel(LOCTEXT("RCPresetResetButtonColumnHeader", ""))
 			.FixedWidth(48.f)
 			.HeaderContentPadding(RCPanelStyle->HeaderRowPadding)
@@ -1376,7 +1377,7 @@ SHeaderRow::FColumn::FArguments SRCPanelExposedEntitiesList::CreateColumn(const
 {
 	const FText DefaultColumnLabel = GetColumnLabel(ForColumnName);
 
-	if (ForColumnName == FRemoteControlPresetColumns::Status || ForColumnName == FRemoteControlDefaultProtocolColumns::BindingStatus) // Exception for Status & BindingStatus columns as they are fixed ones.
+	if (ForColumnName == RemoteControlPresetColumns::Status || ForColumnName == RemoteControlPresetColumns::BindingStatus) // Exception for Status & BindingStatus columns as they are fixed ones.
 	{
 		const SHeaderRow::FColumn::FArguments ColumnArgs = SHeaderRow::Column(ForColumnName)
 			.DefaultLabel(DefaultColumnLabel)
@@ -1398,25 +1399,21 @@ SHeaderRow::FColumn::FArguments SRCPanelExposedEntitiesList::CreateColumn(const
 
 int32 SRCPanelExposedEntitiesList::GetColumnIndex(const FName& ForColumn) const
 {
-	if (ForColumn == FRemoteControlPresetColumns::LinkIdentifier)
+	if (ForColumn == RemoteControlPresetColumns::LinkIdentifier)
 	{
-		return GetColumnIndex_Internal(ForColumn, FRemoteControlPresetColumns::DragDropHandle, ERCColumn::ERC_After);
+		return GetColumnIndex_Internal(ForColumn, RemoteControlPresetColumns::DragDropHandle, ERCColumn::ERC_After);
 	}
-	else if (ForColumn == FRemoteControlPresetColumns::Mask)
+	else if (ForColumn == RemoteControlPresetColumns::Mask)
 	{
-		return GetColumnIndex_Internal(ForColumn, FRemoteControlPresetColumns::Description, ERCColumn::ERC_After);
+		return GetColumnIndex_Internal(ForColumn, RemoteControlPresetColumns::Description, ERCColumn::ERC_After);
 	}
-	else if (ForColumn == FRemoteControlPresetColumns::Status)
+	else if (ForColumn == RemoteControlPresetColumns::Status)
 	{
-		return GetColumnIndex_Internal(ForColumn, FRemoteControlPresetColumns::Value, ERCColumn::ERC_After);
-	}
-	else if (ForColumn == FRemoteControlPresetColumns::Type)
-	{
-		return GetColumnIndex_Internal(ForColumn, FRemoteControlPresetColumns::Value, ERCColumn::ERC_After);
+		return GetColumnIndex_Internal(ForColumn, RemoteControlPresetColumns::Value, ERCColumn::ERC_After);
 	}
 	else
 	{
-		return GetColumnIndex_Internal(ForColumn, FRemoteControlPresetColumns::Value, ERCColumn::ERC_After);
+		return GetColumnIndex_Internal(ForColumn, RemoteControlPresetColumns::Value, ERCColumn::ERC_After);
 	}
 
 	return INDEX_NONE;
@@ -1476,15 +1473,19 @@ int32 SRCPanelExposedEntitiesList::GetColumnIndex_Internal(const FName& ForColum
 
 FText SRCPanelExposedEntitiesList::GetColumnLabel(const FName& ForColumn) const
 {
-	if (ForColumn == FRemoteControlPresetColumns::LinkIdentifier)
+	if (ForColumn == RemoteControlPresetColumns::BindingStatus)
+	{
+		return LOCTEXT("RCPresetBindingStatusColumnHeader", "REC");
+	}
+	else if (ForColumn == RemoteControlPresetColumns::LinkIdentifier)
 	{
 		return LOCTEXT("RCPresetLinkIDColumnHeader", "Link ID");
 	}
-	else if (ForColumn == FRemoteControlPresetColumns::Mask)
+	else if (ForColumn == RemoteControlPresetColumns::Mask)
 	{
 		return LOCTEXT("RCPresetMaskColumnHeader", "Mask");
 	}
-	else if (ForColumn == FRemoteControlPresetColumns::Status)
+	else if (ForColumn == RemoteControlPresetColumns::Status)
 	{
 		return LOCTEXT("RCPresetStatusColumnHeader", "");
 	}
@@ -1506,10 +1507,9 @@ FText SRCPanelExposedEntitiesList::GetColumnLabel(const FName& ForColumn) const
 
 float SRCPanelExposedEntitiesList::GetColumnSize(const FName ForColumn) const
 {
-	float ColumnSize = ProtocolColumnConstants::ColumnSizeSmall;
+	float ColumnSize = ProtocolColumnConstants::ColumnSizeNormal;
 
-	if (ForColumn == FRemoteControlPresetColumns::LinkIdentifier ||
-		ForColumn == FRemoteControlPresetColumns::Type)
+	if (ForColumn == RemoteControlPresetColumns::LinkIdentifier)
 	{
 		ColumnSize = ProtocolColumnConstants::ColumnSizeMini;
 	}
diff --git a/Engine/Plugins/VirtualProduction/RemoteControl/Source/RemoteControlUI/Private/UI/SRCPanelExposedField.cpp b/Engine/Plugins/VirtualProduction/RemoteControl/Source/RemoteControlUI/Private/UI/SRCPanelExposedField.cpp
--- a/Engine/Plugins/VirtualProduction/RemoteControl/Source/RemoteControlUI/Private/UI/SRCPanelExposedField.cpp
+++ b/Engine/Plugins/VirtualProduction/RemoteControl/Source/RemoteControlUI/Private/UI/SRCPanelExposedField.cpp
@@ -143,13 +143,64 @@ TSharedRef<SWidget> SRCPanelExposedField::GetProtocolWidget(const FName ForColum
 				{
 					if (TSharedPtr<TStructOnScope<FRemoteControlProtocolEntity>> RCProtocolEntityPtr = RCProtocolIter.GetRemoteControlProtocolEntityPtr())
 					{
-						return RCProtocolEntityPtr->Get()->GetWidget(ForColumnName);
+						if (ForColumnName == RemoteControlPresetColumns::BindingStatus)
+						{
+							TSharedPtr<SWidget> BindingStatusWidget = SNew(SButton)
+								.ButtonStyle(FAppStyle::Get(), "HoverHintOnly")
+								.ToolTipText(LOCTEXT("RecordingButtonToolTip", "Status of the protocol entity binding"))
+								.ForegroundColor(FSlateColor::UseForeground())
+								.OnClicked_Lambda([RCProtocolEntityPtr]()
+									{
+										(*RCProtocolEntityPtr)->ToggleBindingStatus();
+
+										return FReply::Handled();
+									}
+								)
+								.Content()
+								[
+									SNew(SImage)
+									.ColorAndOpacity_Lambda([RCProtocolEntityPtr]()
+										{
+											const ERCBindingStatus ActiveBindingStatus = (*RCProtocolEntityPtr)->GetBindingStatus();
+
+											switch (ActiveBindingStatus)
+											{
+												case ERCBindingStatus::Awaiting:
+													return FLinearColor::Red;
+												case ERCBindingStatus::Bound:
+													return FLinearColor::Green;
+												case ERCBindingStatus::Unassigned:
+													return FLinearColor::Gray;
+												default:
+													checkNoEntry();
+											}
+
+											return FLinearColor::Black;
+										}
+									)
+									.Image(FAppStyle::Get().GetBrush(TEXT("Icons.FilledCircle")))
+								];
+
+							return BindingStatusWidget.ToSharedRef();
+						}
+						else if (TSharedPtr<FRCPanelWidgetRegistry> Registry = WidgetRegistry.Pin())
+						{
+							if (TSharedPtr<IDetailTreeNode> Node = Registry->GetStructTreeNode(RCProtocolEntityPtr, (*RCProtocolEntityPtr)->GetPropertyName(ForColumnName).ToString(), ERCFindNodeMethod::Name))
+							{
+								FNodeWidgets NodeWidgets = Node->CreateNodeWidgets();
+
+								if (NodeWidgets.ValueWidget)
+								{
+									return NodeWidgets.ValueWidget.ToSharedRef();
+								}
+							}
+						}
 					}
 				}
 			}
 		}
 	}
-
+	
 	return SRCPanelTreeNode::GetProtocolWidget(ForColumnName, InProtocolName);
 }
 
@@ -217,11 +268,11 @@ TSharedRef<SWidget> SRCPanelExposedField::GetWidget(const FName ForColumnName, c
 {
 	if (HasProtocolExtension())
 	{
-		if (ForColumnName == FRemoteControlPresetColumns::Mask)
+		if (ForColumnName == RemoteControlPresetColumns::Mask)
 		{
 			return SNew(SRCProtocolMask, WeakField);
 		}
-		else if (ForColumnName == FRemoteControlPresetColumns::Status)
+		else if (ForColumnName == RemoteControlPresetColumns::Status)
 		{
 			return SNew(STextBlock)
 				.Text(LOCTEXT("StatusText", "!"))
diff --git a/Engine/Plugins/VirtualProduction/RemoteControl/Source/RemoteControlUI/Private/UI/SRCPanelTreeNode.cpp b/Engine/Plugins/VirtualProduction/RemoteControl/Source/RemoteControlUI/Private/UI/SRCPanelTreeNode.cpp
--- a/Engine/Plugins/VirtualProduction/RemoteControl/Source/RemoteControlUI/Private/UI/SRCPanelTreeNode.cpp
+++ b/Engine/Plugins/VirtualProduction/RemoteControl/Source/RemoteControlUI/Private/UI/SRCPanelTreeNode.cpp
@@ -11,10 +11,10 @@
 #define LOCTEXT_NAMESPACE "RemoteControlPanelNode"
 
 TSet<FName> SRCPanelTreeNode::DefaultColumns = {
-	FRemoteControlPresetColumns::DragDropHandle,
-	FRemoteControlPresetColumns::Description,
-	FRemoteControlPresetColumns::Value,
-	FRemoteControlPresetColumns::Reset
+	RemoteControlPresetColumns::DragDropHandle,
+	RemoteControlPresetColumns::Description,
+	RemoteControlPresetColumns::Value,
+	RemoteControlPresetColumns::Reset
 };
 
 TSharedRef<SWidget> SRCPanelTreeNode::GetProtocolWidget(const FName ForColumnName, const FName InProtocolName)
@@ -39,19 +39,19 @@ const bool SRCPanelTreeNode::SupportsProtocol(const FName& InProtocolName) const
 
 TSharedRef<SWidget> SRCPanelTreeNode::GetWidget(const FName ForColumnName, const FName InActiveProtocol)
 {
-	if (ForColumnName == FRemoteControlPresetColumns::DragDropHandle)
+	if (ForColumnName == RemoteControlPresetColumns::DragDropHandle)
 	{
 		return DragHandleWidget.ToSharedRef();
 	}
-	else if (ForColumnName == FRemoteControlPresetColumns::Description)
+	else if (ForColumnName == RemoteControlPresetColumns::Description)
 	{
 		return NodeNameWidget.ToSharedRef();
 	}
-	else if (ForColumnName == FRemoteControlPresetColumns::Reset)
+	else if (ForColumnName == RemoteControlPresetColumns::Reset)
 	{
 		return ResetValueWidget.ToSharedRef();
 	}
-	else if (ForColumnName == FRemoteControlPresetColumns::Value)
+	else if (ForColumnName == RemoteControlPresetColumns::Value)
 	{
 		return NodeValueWidget.ToSharedRef();
 	}
diff --git a/Engine/Plugins/VirtualProduction/RemoteControl/Source/RemoteControlUI/Private/UI/SRCPanelTreeNode.h b/Engine/Plugins/VirtualProduction/RemoteControl/Source/RemoteControlUI/Private/UI/SRCPanelTreeNode.h
--- a/Engine/Plugins/VirtualProduction/RemoteControl/Source/RemoteControlUI/Private/UI/SRCPanelTreeNode.h
+++ b/Engine/Plugins/VirtualProduction/RemoteControl/Source/RemoteControlUI/Private/UI/SRCPanelTreeNode.h
@@ -15,14 +15,14 @@ struct SRCPanelExposedField;
 struct SRCPanelExposedActor;
 struct SRCPanelExposedMaterial;
 
-namespace FRemoteControlPresetColumns
+namespace RemoteControlPresetColumns
 {
 	static FName DragDropHandle = TEXT("DragDropHandle");
 	static FName LinkIdentifier = TEXT("LinkID");
 	static FName Description = TEXT("Description");
 	static FName Mask = TEXT("Mask");
-	static FName Type = TEXT("Type");
 	static FName Value = TEXT("Value");
+	static FName BindingStatus = TEXT("BindingStatus");
 	static FName Status = TEXT("Status");
 	static FName Reset = TEXT("Reset");
 }
diff --git a/Engine/Plugins/VirtualProduction/RemoteControlProtocolDMX/Source/RemoteControlProtocolDMX/Private/RemoteControlProtocolDMX.cpp b/Engine/Plugins/VirtualProduction/RemoteControlProtocolDMX/Source/RemoteControlProtocolDMX/Private/RemoteControlProtocolDMX.cpp
--- a/Engine/Plugins/VirtualProduction/RemoteControlProtocolDMX/Source/RemoteControlProtocolDMX/Private/RemoteControlProtocolDMX.cpp
+++ b/Engine/Plugins/VirtualProduction/RemoteControlProtocolDMX/Source/RemoteControlProtocolDMX/Private/RemoteControlProtocolDMX.cpp
@@ -15,7 +15,6 @@
 #if WITH_EDITOR
 #include "IRCProtocolBindingList.h"
 #include "IRemoteControlProtocolWidgetsModule.h"
-#include "Widgets/Input/SNumericEntryBox.h"
 #endif
 
 #define LOCTEXT_NAMESPACE "RemoteControlProtocolDMX"
@@ -24,7 +23,7 @@ const FName FRemoteControlProtocolDMX::ProtocolName = TEXT("DMX");
 
 #if WITH_EDITOR
 
-namespace FRemoteControlDMXProtocolColumns
+namespace RemoteControlDMXProtocolColumns
 {
 	static FName Channel = TEXT("Channel");
 	static FName Universe = TEXT("Universe");
@@ -87,33 +86,11 @@ const FString& FRemoteControlDMXProtocolEntity::GetRangePropertyMaxValue() const
 
 #if WITH_EDITOR
 
-void FRemoteControlDMXProtocolEntity::RegisterWidgets()
+void FRemoteControlDMXProtocolEntity::RegisterProperties()
 {
-	FRemoteControlProtocolEntity::RegisterWidgets();
-
-	{ // Universe
-		if (!Widgets.Contains(FRemoteControlDMXProtocolColumns::Universe))
-		{
-			TSharedPtr<SWidget> UniverseWidget = SNew(SNumericEntryBox<int32>)
-				.AllowSpin(true)
-				.Value_Lambda([this]() { return ExtraSetting.Universe; })
-				.OnValueCommitted_Lambda([this](int32 NewValue, ETextCommit::Type InCommitType) { ExtraSetting.Universe = NewValue; });
-
-			Widgets.Add(FRemoteControlDMXProtocolColumns::Universe, UniverseWidget);
-		}
-	}
+	EXPOSE_PROTOCOL_PROPERTY(RemoteControlDMXProtocolColumns::Universe, FRemoteControlDMXProtocolEntityExtraSetting, Universe);
 	
-	{ // Channel
-		if (!Widgets.Contains(FRemoteControlDMXProtocolColumns::Channel))
-		{
-			TSharedPtr<SWidget> ChannelWidget = SNew(SNumericEntryBox<int32>)
-				.AllowSpin(true)
-				.Value_Lambda([this]() { return ExtraSetting.StartingChannel; })
-				.OnValueCommitted_Lambda([this](int32 NewValue, ETextCommit::Type InCommitType) { ExtraSetting.StartingChannel = NewValue; });
-
-			Widgets.Add(FRemoteControlDMXProtocolColumns::Channel, ChannelWidget);
-		}
-	}
+	EXPOSE_PROTOCOL_PROPERTY(RemoteControlDMXProtocolColumns::Channel, FRemoteControlDMXProtocolEntityExtraSetting, StartingChannel);
 }
 
 #endif // WITH_EDITOR
@@ -371,11 +348,11 @@ void FRemoteControlProtocolDMX::RegisterColumns()
 {
 	FRemoteControlProtocol::RegisterColumns();
 
-	REGISTER_COLUMN(FRemoteControlDMXProtocolColumns::Channel
+	REGISTER_COLUMN(RemoteControlDMXProtocolColumns::Channel
 		, LOCTEXT("RCPresetChannelColumnHeader", "Channel")
 		, ProtocolColumnConstants::ColumnSizeMicro);
 
-	REGISTER_COLUMN(FRemoteControlDMXProtocolColumns::Universe
+	REGISTER_COLUMN(RemoteControlDMXProtocolColumns::Universe
 		, LOCTEXT("RCPresetUniverseColumnHeader", "Universe")
 		, ProtocolColumnConstants::ColumnSizeMicro);
 }
diff --git a/Engine/Plugins/VirtualProduction/RemoteControlProtocolDMX/Source/RemoteControlProtocolDMX/Public/RemoteControlProtocolDMX.h b/Engine/Plugins/VirtualProduction/RemoteControlProtocolDMX/Source/RemoteControlProtocolDMX/Public/RemoteControlProtocolDMX.h
--- a/Engine/Plugins/VirtualProduction/RemoteControlProtocolDMX/Source/RemoteControlProtocolDMX/Public/RemoteControlProtocolDMX.h
+++ b/Engine/Plugins/VirtualProduction/RemoteControlProtocolDMX/Source/RemoteControlProtocolDMX/Public/RemoteControlProtocolDMX.h
@@ -72,7 +72,7 @@ public:
 #if WITH_EDITOR
 
 	/** Register(s) all the widgets of this protocol entity. */
-	virtual void RegisterWidgets() override;
+	virtual void RegisterProperties() override;
 
 #endif // WITH_EDITOR
 	//~ End FRemoteControlProtocolEntity interface
diff --git a/Engine/Plugins/VirtualProduction/RemoteControlProtocolDMX/Source/RemoteControlProtocolDMX/RemoteControlProtocolDMX.Build.cs b/Engine/Plugins/VirtualProduction/RemoteControlProtocolDMX/Source/RemoteControlProtocolDMX/RemoteControlProtocolDMX.Build.cs
--- a/Engine/Plugins/VirtualProduction/RemoteControlProtocolDMX/Source/RemoteControlProtocolDMX/RemoteControlProtocolDMX.Build.cs
+++ b/Engine/Plugins/VirtualProduction/RemoteControlProtocolDMX/Source/RemoteControlProtocolDMX/RemoteControlProtocolDMX.Build.cs
@@ -28,8 +28,6 @@ public class RemoteControlProtocolDMX : ModuleRules
 				new string[] {
 					"InputCore",
 					"RemoteControlProtocolWidgets",
-					"Slate",
-					"SlateCore",
 				}
 			);
 		}
diff --git a/Engine/Plugins/VirtualProduction/RemoteControlProtocolMIDI/Source/RemoteControlProtocolMIDI/Private/RemoteControlProtocolMIDI.cpp b/Engine/Plugins/VirtualProduction/RemoteControlProtocolMIDI/Source/RemoteControlProtocolMIDI/Private/RemoteControlProtocolMIDI.cpp
--- a/Engine/Plugins/VirtualProduction/RemoteControlProtocolMIDI/Source/RemoteControlProtocolMIDI/Private/RemoteControlProtocolMIDI.cpp
+++ b/Engine/Plugins/VirtualProduction/RemoteControlProtocolMIDI/Source/RemoteControlProtocolMIDI/Private/RemoteControlProtocolMIDI.cpp
@@ -10,7 +10,6 @@
 #if WITH_EDITOR
 #include "IRCProtocolBindingList.h"
 #include "IRemoteControlProtocolWidgetsModule.h"
-#include "Widgets/Input/SNumericEntryBox.h"
 #endif
 
 #define LOCTEXT_NAMESPACE "FRemoteControlProtocolMIDI"
@@ -19,10 +18,11 @@ const FName FRemoteControlProtocolMIDI::ProtocolName = TEXT("MIDI");
 
 #if WITH_EDITOR
 
-namespace FRemoteControlMIDIProtocolColumns
+namespace RemoteControlMIDIProtocolColumns
 {
 	static FName Channel = TEXT("Channel");
 	static FName Identifier = TEXT("Identifier");
+	static FName Type = TEXT("Type");
 }
 
 #endif // WITH_EDITOR
@@ -145,33 +145,13 @@ bool FRemoteControlMIDIProtocolEntity::IsSame(const FRemoteControlProtocolEntity
 
 #if WITH_EDITOR
 
-void FRemoteControlMIDIProtocolEntity::RegisterWidgets()
+void FRemoteControlMIDIProtocolEntity::RegisterProperties()
 {
-	FRemoteControlProtocolEntity::RegisterWidgets();
+	EXPOSE_PROTOCOL_PROPERTY(RemoteControlMIDIProtocolColumns::Channel, FRemoteControlMIDIProtocolEntity, Channel);
 
-	{ // Identifier
-		if (!Widgets.Contains(FRemoteControlMIDIProtocolColumns::Identifier))
-		{
-			TSharedPtr<SWidget> IdentifierWidget = SNew(SNumericEntryBox<int32>)
-				.AllowSpin(true)
-				.Value_Lambda([this]() { return MessageData1; })
-				.OnValueCommitted_Lambda([this](int32 NewValue, ETextCommit::Type InCommitType) { MessageData1 = NewValue; });
+	EXPOSE_PROTOCOL_PROPERTY(RemoteControlMIDIProtocolColumns::Identifier, FRemoteControlMIDIProtocolEntity, MessageData1);
 
-			Widgets.Add(FRemoteControlMIDIProtocolColumns::Identifier, IdentifierWidget);
-		}
-	}
-
-	{ // Channel
-		if (!Widgets.Contains(FRemoteControlMIDIProtocolColumns::Channel))
-		{
-			TSharedPtr<SWidget> ChannelWidget = SNew(SNumericEntryBox<int32>)
-				.AllowSpin(true)
-				.Value_Lambda([this]() { return Channel; })
-				.OnValueCommitted_Lambda([this](int32 NewValue, ETextCommit::Type InCommitType) { Channel = NewValue; });
-
-			Widgets.Add(FRemoteControlMIDIProtocolColumns::Channel, ChannelWidget);
-		}
-	}
+	EXPOSE_PROTOCOL_PROPERTY(RemoteControlMIDIProtocolColumns::Type, FRemoteControlMIDIProtocolEntity, EventType);
 }
 
 #endif // WITH_EDITOR
@@ -408,13 +388,17 @@ void FRemoteControlProtocolMIDI::RegisterColumns()
 {
 	FRemoteControlProtocol::RegisterColumns();
 
-	REGISTER_COLUMN(FRemoteControlMIDIProtocolColumns::Channel
+	REGISTER_COLUMN(RemoteControlMIDIProtocolColumns::Channel
 		, LOCTEXT("RCPresetChannelColumnHeader", "Channel")
 		, ProtocolColumnConstants::ColumnSizeMicro);
 	
-	REGISTER_COLUMN(FRemoteControlMIDIProtocolColumns::Identifier
+	REGISTER_COLUMN(RemoteControlMIDIProtocolColumns::Identifier
 		, LOCTEXT("RCPresetIdentifierColumnHeader", "ID")
 		, ProtocolColumnConstants::ColumnSizeMicro);
+
+	REGISTER_COLUMN(RemoteControlMIDIProtocolColumns::Type
+		, LOCTEXT("RCPresetTypeColumnHeader", "Type")
+		, ProtocolColumnConstants::ColumnSizeSmall);
 }
 
 #endif // WITH_EDITOR
diff --git a/Engine/Plugins/VirtualProduction/RemoteControlProtocolMIDI/Source/RemoteControlProtocolMIDI/Public/RemoteControlProtocolMIDI.h b/Engine/Plugins/VirtualProduction/RemoteControlProtocolMIDI/Source/RemoteControlProtocolMIDI/Public/RemoteControlProtocolMIDI.h
--- a/Engine/Plugins/VirtualProduction/RemoteControlProtocolMIDI/Source/RemoteControlProtocolMIDI/Public/RemoteControlProtocolMIDI.h
+++ b/Engine/Plugins/VirtualProduction/RemoteControlProtocolMIDI/Source/RemoteControlProtocolMIDI/Public/RemoteControlProtocolMIDI.h
@@ -157,7 +157,7 @@ public:
 #if WITH_EDITOR
 
 	/** Register(s) all the widgets of this protocol entity. */
-	virtual void RegisterWidgets() override;
+	virtual void RegisterProperties() override;
 
 #endif // WITH_EDITOR
 };
diff --git a/Engine/Plugins/VirtualProduction/RemoteControlProtocolMIDI/Source/RemoteControlProtocolMIDI/RemoteControlProtocolMIDI.Build.cs b/Engine/Plugins/VirtualProduction/RemoteControlProtocolMIDI/Source/RemoteControlProtocolMIDI/RemoteControlProtocolMIDI.Build.cs
--- a/Engine/Plugins/VirtualProduction/RemoteControlProtocolMIDI/Source/RemoteControlProtocolMIDI/RemoteControlProtocolMIDI.Build.cs
+++ b/Engine/Plugins/VirtualProduction/RemoteControlProtocolMIDI/Source/RemoteControlProtocolMIDI/RemoteControlProtocolMIDI.Build.cs
@@ -27,8 +27,6 @@ public class RemoteControlProtocolMIDI : ModuleRules
 				new string[] {
 					"InputCore",
 					"RemoteControlProtocolWidgets",
-					"Slate",
-					"SlateCore",
 				}
 			);
 		}
diff --git a/Engine/Plugins/VirtualProduction/RemoteControlProtocolOSC/Source/RemoteControlProtocolOSC/Private/RemoteControlProtocolOSC.cpp b/Engine/Plugins/VirtualProduction/RemoteControlProtocolOSC/Source/RemoteControlProtocolOSC/Private/RemoteControlProtocolOSC.cpp
--- a/Engine/Plugins/VirtualProduction/RemoteControlProtocolOSC/Source/RemoteControlProtocolOSC/Private/RemoteControlProtocolOSC.cpp
+++ b/Engine/Plugins/VirtualProduction/RemoteControlProtocolOSC/Source/RemoteControlProtocolOSC/Private/RemoteControlProtocolOSC.cpp
@@ -7,7 +7,6 @@
 #if WITH_EDITOR
 #include "IRCProtocolBindingList.h"
 #include "IRemoteControlProtocolWidgetsModule.h"
-#include "Widgets/Input/SEditableTextBox.h"
 #endif
 
 #include "OSCManager.h"
@@ -19,7 +18,7 @@ const FName FRemoteControlProtocolOSC::ProtocolName = TEXT("OSC");
 
 #if WITH_EDITOR
 
-namespace FRemoteControlOSCProtocolColumns
+namespace RemoteControlOSCProtocolColumns
 {
 	static FName Address = TEXT("Address");
 }
@@ -38,20 +37,9 @@ bool FRemoteControlOSCProtocolEntity::IsSame(const FRemoteControlProtocolEntity*
 
 #if WITH_EDITOR
 
-void FRemoteControlOSCProtocolEntity::RegisterWidgets()
+void FRemoteControlOSCProtocolEntity::RegisterProperties()
 {
-	FRemoteControlProtocolEntity::RegisterWidgets();
-
-	{ // Address
-		if (!Widgets.Contains(FRemoteControlOSCProtocolColumns::Address))
-		{
-			TSharedPtr<SWidget> AddressWidget = SNew(SEditableTextBox)
-				.Text_Lambda([this]() { return FText::FromName(PathName); })
-				.OnTextCommitted_Lambda([this](const FText& NewValue, ETextCommit::Type InCommitType) { PathName = *NewValue.ToString(); });
-
-			Widgets.Add(FRemoteControlOSCProtocolColumns::Address, AddressWidget);
-		}
-	}
+	EXPOSE_PROTOCOL_PROPERTY(RemoteControlOSCProtocolColumns::Address, FRemoteControlOSCProtocolEntity, PathName);
 }
 
 #endif // WITH_EDITOR
@@ -171,7 +159,7 @@ void FRemoteControlProtocolOSC::RegisterColumns()
 {
 	FRemoteControlProtocol::RegisterColumns();
 
-	REGISTER_COLUMN(FRemoteControlOSCProtocolColumns::Address
+	REGISTER_COLUMN(RemoteControlOSCProtocolColumns::Address
 		, LOCTEXT("RCPresetAddressColumnHeader", "Address")
 		, ProtocolColumnConstants::ColumnSizeMedium);
 }
diff --git a/Engine/Plugins/VirtualProduction/RemoteControlProtocolOSC/Source/RemoteControlProtocolOSC/Private/RemoteControlProtocolOSC.h b/Engine/Plugins/VirtualProduction/RemoteControlProtocolOSC/Source/RemoteControlProtocolOSC/Private/RemoteControlProtocolOSC.h
--- a/Engine/Plugins/VirtualProduction/RemoteControlProtocolOSC/Source/RemoteControlProtocolOSC/Private/RemoteControlProtocolOSC.h
+++ b/Engine/Plugins/VirtualProduction/RemoteControlProtocolOSC/Source/RemoteControlProtocolOSC/Private/RemoteControlProtocolOSC.h
@@ -40,7 +40,7 @@ public:
 #if WITH_EDITOR
 
 	/** Register(s) all the widgets of this protocol entity. */
-	virtual void RegisterWidgets() override;
+	virtual void RegisterProperties() override;
 
 #endif // WITH_EDITOR
 };
diff --git a/Engine/Plugins/VirtualProduction/RemoteControlProtocolOSC/Source/RemoteControlProtocolOSC/RemoteControlProtocolOSC.Build.cs b/Engine/Plugins/VirtualProduction/RemoteControlProtocolOSC/Source/RemoteControlProtocolOSC/RemoteControlProtocolOSC.Build.cs
--- a/Engine/Plugins/VirtualProduction/RemoteControlProtocolOSC/Source/RemoteControlProtocolOSC/RemoteControlProtocolOSC.Build.cs
+++ b/Engine/Plugins/VirtualProduction/RemoteControlProtocolOSC/Source/RemoteControlProtocolOSC/RemoteControlProtocolOSC.Build.cs
@@ -30,8 +30,6 @@ public class RemoteControlProtocolOSC : ModuleRules
 					"InputCore",
 				    "RemoteControlProtocolWidgets",
 					"Settings",
-					"Slate",
-					"SlateCore",
 				}
 			);
 		}
