Remote API Protocol - Record button in main table does not work

diff --git a/Engine/Plugins/VirtualProduction/RemoteControl/Source/RemoteControlProtocolWidgets/Private/Widgets/SRCProtocolBindingList.h b/Engine/Plugins/VirtualProduction/RemoteControl/Source/RemoteControlProtocolWidgets/Private/Widgets/SRCProtocolBindingList.h
--- a/Engine/Plugins/VirtualProduction/RemoteControl/Source/RemoteControlProtocolWidgets/Private/Widgets/SRCProtocolBindingList.h
+++ b/Engine/Plugins/VirtualProduction/RemoteControl/Source/RemoteControlProtocolWidgets/Private/Widgets/SRCProtocolBindingList.h
@@ -1,4 +1,4 @@
-﻿// Copyright Epic Games, Inc. All Rights Reserved.
+// Copyright Epic Games, Inc. All Rights Reserved.
 
 #pragma once
 
@@ -35,6 +35,15 @@ public:
 	{
 		return AwaitingProtocolEntities;
 	}
+
+private:
+
+	/** Start recording incoming protocol message handler */
+	virtual void OnStartRecording(TSharedPtr<TStructOnScope<FRemoteControlProtocolEntity>> InEntity) override;
+
+	/** Stop recording incoming protocol message handler */
+	virtual void OnStopRecording(TSharedPtr<TStructOnScope<FRemoteControlProtocolEntity>> InEntity) override;
+
 	//~ End IRCProtocolBindingList Interface
 
 private:
@@ -85,12 +94,6 @@ private:
 		SecondaryColumnWidth = InWidth;
 	}
 
-	/** Start recording incoming protocol message handler */
-	void OnStartRecording(TSharedPtr<TStructOnScope<FRemoteControlProtocolEntity>> InEntity);
-
-	/** Stop recording incoming protocol message handler */
-	void OnStopRecording(TSharedPtr<TStructOnScope<FRemoteControlProtocolEntity>> InEntity);
-
 	/** Get (mutable) module settings */
 	URemoteControlProtocolWidgetsSettings* GetSettings();
 
diff --git a/Engine/Plugins/VirtualProduction/RemoteControl/Source/RemoteControlProtocolWidgets/Public/IRCProtocolBindingList.h b/Engine/Plugins/VirtualProduction/RemoteControl/Source/RemoteControlProtocolWidgets/Public/IRCProtocolBindingList.h
--- a/Engine/Plugins/VirtualProduction/RemoteControl/Source/RemoteControlProtocolWidgets/Public/IRCProtocolBindingList.h
+++ b/Engine/Plugins/VirtualProduction/RemoteControl/Source/RemoteControlProtocolWidgets/Public/IRCProtocolBindingList.h
@@ -20,4 +20,11 @@ public:
 
 	/** Get the set of entities which is awaiting state and waiting for binding */
 	virtual FRemoteControlProtocolEntitySet GetAwaitingProtocolEntities() const = 0;
+
+	/** Start recording incoming protocol message handler */
+	virtual void OnStartRecording(TSharedPtr<TStructOnScope<FRemoteControlProtocolEntity>> InEntity) = 0;
+
+	/** Stop recording incoming protocol message handler */
+	virtual void OnStopRecording(TSharedPtr<TStructOnScope<FRemoteControlProtocolEntity>> InEntity) = 0;
+
 };
diff --git a/Engine/Plugins/VirtualProduction/RemoteControl/Source/RemoteControlUI/Private/UI/SRCPanelExposedField.cpp b/Engine/Plugins/VirtualProduction/RemoteControl/Source/RemoteControlUI/Private/UI/SRCPanelExposedField.cpp
--- a/Engine/Plugins/VirtualProduction/RemoteControl/Source/RemoteControlUI/Private/UI/SRCPanelExposedField.cpp
+++ b/Engine/Plugins/VirtualProduction/RemoteControl/Source/RemoteControlUI/Private/UI/SRCPanelExposedField.cpp
@@ -5,7 +5,9 @@
 #include "Algo/Transform.h"
 #include "EditorFontGlyphs.h"
 #include "IDetailTreeNode.h"
+#include "IRCProtocolBindingList.h"
 #include "IRemoteControlProtocolModule.h"
+#include "IRemoteControlProtocolWidgetsModule.h"
 #include "IRemoteControlUIModule.h"
 #include "Layout/Visibility.h"
 #include "Materials/MaterialInstanceDynamic.h"
@@ -151,7 +153,25 @@ TSharedRef<SWidget> SRCPanelExposedField::GetProtocolWidget(const FName ForColum
 								.ForegroundColor(FSlateColor::UseForeground())
 								.OnClicked_Lambda([RCProtocolEntityPtr]()
 									{
-										(*RCProtocolEntityPtr)->ToggleBindingStatus();
+										const ERCBindingStatus BindingStatus = (*RCProtocolEntityPtr)->ToggleBindingStatus();
+
+										IRemoteControlProtocolWidgetsModule& RCProtocolWidgetsModule = IRemoteControlProtocolWidgetsModule::Get();
+
+										if (TSharedPtr<IRCProtocolBindingList> RCProtocolBindingList = RCProtocolWidgetsModule.GetProtocolBindingList())
+										{
+											if (BindingStatus == ERCBindingStatus::Awaiting)
+											{
+												RCProtocolBindingList->OnStartRecording(RCProtocolEntityPtr);
+											}
+											else if (BindingStatus == ERCBindingStatus::Bound)
+											{
+												RCProtocolBindingList->OnStopRecording(RCProtocolEntityPtr);
+											}
+											else
+											{
+												checkNoEntry();
+											}
+										}
 
 										return FReply::Handled();
 									}
