Update remote control protocol UI

diff --git a/Engine/Plugins/VirtualProduction/RemoteControl/Source/RemoteControlProtocolWidgets/Private/RemoteControlProtocolWidgetsModule.cpp b/Engine/Plugins/VirtualProduction/RemoteControl/Source/RemoteControlProtocolWidgets/Private/RemoteControlProtocolWidgetsModule.cpp
--- a/Engine/Plugins/VirtualProduction/RemoteControl/Source/RemoteControlProtocolWidgets/Private/RemoteControlProtocolWidgetsModule.cpp
+++ b/Engine/Plugins/VirtualProduction/RemoteControl/Source/RemoteControlProtocolWidgets/Private/RemoteControlProtocolWidgetsModule.cpp
@@ -5,6 +5,7 @@
 #include "IRemoteControlModule.h"
 #include "IRemoteControlProtocol.h"
 #include "IRemoteControlProtocolModule.h"
+#include "RemoteControlProtocolWidgetsSettings.h"
 #include "IRemoteControlModule.h"
 #include "RemoteControlPreset.h"
 #include "Styling/ProtocolPanelStyle.h"
@@ -17,11 +18,15 @@ DEFINE_LOG_CATEGORY(LogRemoteControlProtocolWidgets);
 void FRemoteControlProtocolWidgetsModule::StartupModule()
 {
 	FProtocolPanelStyle::Initialize();
+
+	OnActiveProtocolChanged().AddRaw(this, &FRemoteControlProtocolWidgetsModule::SetActiveProtocolName);
 }
 
 void FRemoteControlProtocolWidgetsModule::ShutdownModule()
 {
 	FProtocolPanelStyle::Shutdown();
+
+	OnActiveProtocolChanged().RemoveAll(this);
 }
 
 TSharedRef<SWidget> FRemoteControlProtocolWidgetsModule::GenerateDetailsForEntity(URemoteControlPreset* InPreset, const FGuid& InFieldId, const EExposedFieldType& InFieldType)
@@ -63,11 +68,21 @@ TSharedPtr<IRCProtocolBindingList> FRemoteControlProtocolWidgetsModule::GetProto
 	return RCProtocolBindingList;
 }
 
+const FName FRemoteControlProtocolWidgetsModule::GetSelectedProtocolName() const
+{
+	return ActiveProtocolName;
+}
+
 FOnProtocolBindingAddedOrRemoved& FRemoteControlProtocolWidgetsModule::OnProtocolBindingAddedOrRemoved()
 {
 	return OnProtocolBindingAddedOrRemovedDelegate;
 }
 
+FOnActiveProtocolChanged& FRemoteControlProtocolWidgetsModule::OnActiveProtocolChanged()
+{
+	return OnActiveProtocolChangedDelegate;
+}
+
 void FRemoteControlProtocolWidgetsModule::OnBindingAdded(TSharedRef<FProtocolBindingViewModel> InBindingViewModel)
 {
 	OnProtocolBindingAddedOrRemovedDelegate.Broadcast(ERCProtocolBinding::Added);
@@ -78,4 +93,27 @@ void FRemoteControlProtocolWidgetsModule::OnBindingRemoved(FGuid InBindingId)
 	OnProtocolBindingAddedOrRemovedDelegate.Broadcast(ERCProtocolBinding::Removed);
 }
 
+void FRemoteControlProtocolWidgetsModule::SetActiveProtocolName(const FName InProtocolName)
+{
+	if (ActiveProtocolName == InProtocolName)
+	{
+		return;
+	}
+
+	ActiveProtocolName = InProtocolName;
+
+	URemoteControlProtocolWidgetsSettings* Settings = GetMutableDefault<URemoteControlProtocolWidgetsSettings>();
+
+	Settings->PreferredProtocol = ActiveProtocolName;
+
+	Settings->PostEditChange();
+
+	Settings->SaveConfig();
+
+	if (RCProtocolBindingList.IsValid())
+	{
+		RCProtocolBindingList->Refresh();
+	}
+}
+
 IMPLEMENT_MODULE(FRemoteControlProtocolWidgetsModule, RemoteControlProtocolWidgets);
diff --git a/Engine/Plugins/VirtualProduction/RemoteControl/Source/RemoteControlProtocolWidgets/Private/RemoteControlProtocolWidgetsModule.h b/Engine/Plugins/VirtualProduction/RemoteControl/Source/RemoteControlProtocolWidgets/Private/RemoteControlProtocolWidgetsModule.h
--- a/Engine/Plugins/VirtualProduction/RemoteControl/Source/RemoteControlProtocolWidgets/Private/RemoteControlProtocolWidgetsModule.h
+++ b/Engine/Plugins/VirtualProduction/RemoteControl/Source/RemoteControlProtocolWidgets/Private/RemoteControlProtocolWidgetsModule.h
@@ -23,7 +23,9 @@ public:
 	virtual TSharedRef<SWidget> GenerateDetailsForEntity(URemoteControlPreset* InPreset, const FGuid& InFieldId, const EExposedFieldType& InFieldType) override;
 	virtual void ResetProtocolBindingList() override;
 	virtual TSharedPtr<IRCProtocolBindingList> GetProtocolBindingList() const override;
+	virtual const FName GetSelectedProtocolName() const override;
 	virtual FOnProtocolBindingAddedOrRemoved& OnProtocolBindingAddedOrRemoved() override;
+	virtual FOnActiveProtocolChanged& OnActiveProtocolChanged() override;
 	//~ End IRemoteControlProtocolWidgetsModule Interface
 
 private:
@@ -34,10 +36,19 @@ private:
 	/** Called when binding is removed. */
 	void OnBindingRemoved(FGuid InBindingId);
 
+	/** Sets the selected protocol in the list and user setting. */
+	void SetActiveProtocolName(const FName InProtocolName);
+
 private:
 	/** Binding list public interface instance */
 	TSharedPtr<IRCProtocolBindingList> RCProtocolBindingList;
 
 	/** Called when binding is added or removed. */
 	FOnProtocolBindingAddedOrRemoved OnProtocolBindingAddedOrRemovedDelegate;
+
+	/** Called when active protocol selection changed. */
+	FOnActiveProtocolChanged OnActiveProtocolChangedDelegate;
+
+	/** Holds the active protocol name. */
+	FName ActiveProtocolName;
 };
diff --git a/Engine/Plugins/VirtualProduction/RemoteControl/Source/RemoteControlProtocolWidgets/Private/Widgets/SRCProtocolBinding.cpp b/Engine/Plugins/VirtualProduction/RemoteControl/Source/RemoteControlProtocolWidgets/Private/Widgets/SRCProtocolBinding.cpp
--- a/Engine/Plugins/VirtualProduction/RemoteControl/Source/RemoteControlProtocolWidgets/Private/Widgets/SRCProtocolBinding.cpp
+++ b/Engine/Plugins/VirtualProduction/RemoteControl/Source/RemoteControlProtocolWidgets/Private/Widgets/SRCProtocolBinding.cpp
@@ -2,7 +2,6 @@
 
 #include "SRCProtocolBinding.h"
 
-#include "EditorFontGlyphs.h"
 #include "SDropTarget.h"
 #include "SRCBindingWarning.h"
 #include "SRCProtocolRangeList.h"
@@ -58,7 +57,7 @@ void SRCProtocolBinding::Construct(const FArguments& InArgs, const TSharedRef<ST
 	RowBox->AddSlot()
 		.AutoWidth()
 		.VAlign(VAlign_Center)
-		.Padding(0)
+		.Padding(2.f)
 		[
 			SNew(SButton)
 			.ButtonStyle(FAppStyle::Get(), "HoverHintOnly")
@@ -67,9 +66,9 @@ void SRCProtocolBinding::Construct(const FArguments& InArgs, const TSharedRef<ST
 			.OnClicked(this, &SRCProtocolBinding::OnDelete)
 			.Content()
 			[
-				SNew(STextBlock)
-				.Font(FAppStyle::Get().GetFontStyle("FontAwesome.10"))
-				.Text(FText::FromString(FString(TEXT("\xf00d"))))
+				SNew(SImage)
+				.ColorAndOpacity(FSlateColor::UseForeground())
+				.Image(FAppStyle::GetBrush("Icons.X"))
 			]
 		];
 
@@ -77,7 +76,7 @@ void SRCProtocolBinding::Construct(const FArguments& InArgs, const TSharedRef<ST
 	RowBox->AddSlot()
 		.AutoWidth()
 		.VAlign(VAlign_Center)
-		.Padding(0)
+		.Padding(2.f)
 		[
 			SNew(SButton)
 			.ButtonStyle(FAppStyle::Get(), "HoverHintOnly")
@@ -86,10 +85,9 @@ void SRCProtocolBinding::Construct(const FArguments& InArgs, const TSharedRef<ST
 			.OnClicked(this, &SRCProtocolBinding::ToggleRecording)
 			.Content()
 			[
-				SNew(STextBlock)
+				SNew(SImage)
 				.ColorAndOpacity_Raw(this, &SRCProtocolBinding::GetRecordingButtonColor)
-				.Font(FAppStyle::Get().GetFontStyle("FontAwesome.10"))
-				.Text(FEditorFontGlyphs::Circle)
+				.Image(FAppStyle::Get().GetBrush(TEXT("Icons.FilledCircle")))
 			]
 		];
 
diff --git a/Engine/Plugins/VirtualProduction/RemoteControl/Source/RemoteControlProtocolWidgets/Private/Widgets/SRCProtocolBindingList.cpp b/Engine/Plugins/VirtualProduction/RemoteControl/Source/RemoteControlProtocolWidgets/Private/Widgets/SRCProtocolBindingList.cpp
--- a/Engine/Plugins/VirtualProduction/RemoteControl/Source/RemoteControlProtocolWidgets/Private/Widgets/SRCProtocolBindingList.cpp
+++ b/Engine/Plugins/VirtualProduction/RemoteControl/Source/RemoteControlProtocolWidgets/Private/Widgets/SRCProtocolBindingList.cpp
@@ -1,4 +1,4 @@
-﻿// Copyright Epic Games, Inc. All Rights Reserved.
+// Copyright Epic Games, Inc. All Rights Reserved.
 
 #include "SRCProtocolBindingList.h"
 
@@ -6,18 +6,15 @@
 #include "EditorFontGlyphs.h"
 #include "Styling/AppStyle.h"
 #include "IRemoteControlProtocolModule.h"
-#include "RemoteControlProtocolWidgetsSettings.h"
+#include "RemoteControlProtocolWidgetsModule.h"
 #include "SRCProtocolBinding.h"
-#include "SRCProtocolList.h"
 #include "Delegates/DelegateSignatureImpl.inl"
-#include "Framework/MultiBox/MultiBoxBuilder.h"
 #include "Types/ISlateMetaData.h"
 #include "ViewModels/ProtocolEntityViewModel.h"
 #include "ViewModels/RCViewModelCommon.h"
 #include "Widgets/SNullWidget.h"
 #include "Widgets/Images/SImage.h"
 #include "Widgets/Input/SButton.h"
-#include "Widgets/Input/SComboButton.h"
 #include "Widgets/Views/SListView.h"
 #include "Widgets/Views/STableRow.h"
 #include "Widgets/Views/STableViewBase.h"
@@ -61,29 +58,9 @@ void SRCProtocolBindingList::Construct(const FArguments& InArgs, TSharedRef<FPro
 		Refresh();
 	});
 
-	// The visibility toggle menu to show/hide protocol types
-	FMenuBuilder ProtocolVisibilityMenu(true, nullptr);
-
 	TSharedRef<SScrollBar> ExternalScrollBar = SNew(SScrollBar);
 	ExternalScrollBar->SetVisibility(TAttribute<EVisibility>(this, &SRCProtocolBindingList::GetScrollBarVisibility));
 
-	ProtocolNames = IRemoteControlProtocolModule::Get().GetProtocolNames();
-	for(const FName& ProtocolName : ProtocolNames)
-	{
-		ProtocolVisibilityMenu.AddMenuEntry(
-			FText::Format(LOCTEXT("ShowProtocolFmt", "Show {0}"), FText::FromName(ProtocolName)),
-			FText::Format(LOCTEXT("ShowProtocolTooltipFmt", "Show all {0} protocol entries."), FText::FromName(ProtocolName)),
-			FSlateIcon(),
-			FUIAction(
-				FExecuteAction::CreateLambda([&](){ this->ToggleShowProtocol(ProtocolName); }),
-				FCanExecuteAction(),
-				FIsActionChecked::CreateLambda([&]() { return this->IsProtocolShown(ProtocolName); })
-			),
-			NAME_None,
-			EUserInterfaceActionType::ToggleButton
-		);
-	}
-
 	SAssignNew(BindingList, SListView<TSharedPtr<IRCTreeNodeViewModel>>)
 	.ListItemsSource(&FilteredBindings)
 	.OnGenerateRow(this, &SRCProtocolBindingList::ConstructBindingWidget)
@@ -103,42 +80,54 @@ void SRCProtocolBindingList::Construct(const FArguments& InArgs, TSharedRef<FPro
 			.AutoWidth()
 			[
 				SNew(SHorizontalBox)
+
 				+ SHorizontalBox::Slot()
 				.Padding(Padding)
 				.AutoWidth()
-				[
-					SAssignNew(ProtocolList, SRCProtocolList)
-				]
-
-				+ SHorizontalBox::Slot()
-				.VAlign(VAlign_Center)
-				.HAlign(HAlign_Left)
 				[
 					SNew(SButton)
-	                .ButtonStyle(FAppStyle::Get(), "NoBorder")
-	                .ToolTipText(LOCTEXT("AddProtocol", "Add Protocol"))
-	                .IsEnabled_Lambda([this](){ return CanAddProtocol(); })
-	                .OnClicked_Lambda([this]()
-	                {
-	                	if(ProtocolList.IsValid() && ProtocolList->GetSelectedProtocolName().IsValid())
-	                	{
-	                		const FName& ProtocolName = *ProtocolList->GetSelectedProtocolName();
-	                		ViewModel->AddBinding(ProtocolName);
-	                		if (!IsProtocolShown(ProtocolName))
-	                		{
-	                			ToggleShowProtocol(ProtocolName);
-	                		}
-	                	}
-                        return FReply::Handled();
-	                })
-	                .ContentPadding(FMargin(2.0f, 2.0f))
-	                .Content()
-	                [
-	                    SNew(STextBlock)
-	                    .TextStyle(FAppStyle::Get(), "NormalText.Important")
-	                    .Font(FAppStyle::Get().GetFontStyle("FontAwesome.10"))
-	                    .Text(FEditorFontGlyphs::Plus)
-	                ]
+					.ButtonStyle(FAppStyle::Get(), "PropertyEditor.AssetComboStyle")
+					.ForegroundColor(FSlateColor::UseForeground())
+					.ToolTipText(this, &SRCProtocolBindingList::HandleAddBindingToolTipText)
+					.IsEnabled_Lambda([this]() { return CanAddProtocol(); })
+					.OnClicked_Lambda([this]()
+						{
+							IRemoteControlProtocolWidgetsModule& RCProtocolWidgetsModule = IRemoteControlProtocolWidgetsModule::Get();
+
+							const FName& ProtocolName = RCProtocolWidgetsModule.GetSelectedProtocolName();
+
+							ViewModel->AddBinding(ProtocolName);
+
+							Refresh(false);
+
+							return FReply::Handled();
+						}
+					)
+					.ContentPadding(FMargin(4.f, 2.f))
+					.Content()
+					[
+						SNew(SHorizontalBox)
+						
+						+SHorizontalBox::Slot()
+						.AutoWidth()
+						.VAlign(VAlign_Center)
+						.Padding(2.f)
+						[
+							SNew(SImage)
+							.ColorAndOpacity(FSlateColor::UseForeground())
+							.Image(FAppStyle::GetBrush("Icons.Plus"))
+						]
+
+						+SHorizontalBox::Slot()
+						.AutoWidth()
+						.VAlign(VAlign_Center)
+						.Padding(2.f)
+						[
+							SNew(STextBlock)
+							.TextStyle(FAppStyle::Get(), "ContentBrowser.TopBar.Font")
+							.Text(LOCTEXT("AddBindingText", "Add Binding"))
+						]
+					]
 				]
 
 				+ SHorizontalBox::Slot()
@@ -159,25 +148,6 @@ void SRCProtocolBindingList::Construct(const FArguments& InArgs, TSharedRef<FPro
 			[
 				SNullWidget::NullWidget
 			]
-
-			+ SHorizontalBox::Slot()
-			.AutoWidth()
-			[
-				SNew(SComboButton)
-				.ContentPadding(0)
-				.ForegroundColor(FSlateColor::UseForeground())
-				.ButtonStyle(FAppStyle::Get(), "ToggleButton")
-				.AddMetaData<FTagMetaData>(FTagMetaData(TEXT("ViewOptions")))
-				.MenuContent()
-				[
-					ProtocolVisibilityMenu.MakeWidget()
-				]
-				.ButtonContent()
-				[
-				 	SNew(SImage)
-				 	.Image(FAppStyle::GetBrush("GenericViewButton"))
-				]
-			]
 		]			 
 
 		+ SVerticalBox::Slot()
@@ -236,14 +206,15 @@ TSharedRef<ITableRow> SRCProtocolBindingList::ConstructBindingWidget(TSharedPtr<
 
 bool SRCProtocolBindingList::CanAddProtocol()
 {
-	const TSharedPtr<FName> SelectedProtocolName = ProtocolList.IsValid() ? ProtocolList->GetSelectedProtocolName() : nullptr;
+	IRemoteControlProtocolWidgetsModule& RCProtocolWidgetsModule = IRemoteControlProtocolWidgetsModule::Get();
+	const FName& SelectedProtocolName = RCProtocolWidgetsModule.GetSelectedProtocolName();
 	if(!ViewModel->IsBound())
 	{
 		StatusMessage = FText::GetEmpty();
 		return false;
 	}
 	
-	if(ViewModel->CanAddBinding(SelectedProtocolName.IsValid() ? *SelectedProtocolName : NAME_None, StatusMessage))
+	if(ViewModel->CanAddBinding(SelectedProtocolName, StatusMessage))
 	{
 		StatusMessage = FText::GetEmpty();
 		return true;
@@ -252,24 +223,18 @@ bool SRCProtocolBindingList::CanAddProtocol()
 	return false;	
 }
 
-void SRCProtocolBindingList::ToggleShowProtocol(const FName& InProtocolName)
+FText SRCProtocolBindingList::HandleAddBindingToolTipText() const
 {
-	if(IsProtocolShown(InProtocolName))
-	{
-		GetSettings()->HiddenProtocolTypeNames.Add(InProtocolName);
-	}
-	else
+	IRemoteControlProtocolWidgetsModule& RCProtocolWidgetsModule = IRemoteControlProtocolWidgetsModule::Get();
+
+	const FName& SelectedProtocolName = RCProtocolWidgetsModule.GetSelectedProtocolName();
+
+	if (!SelectedProtocolName.IsNone())
 	{
-		GetSettings()->HiddenProtocolTypeNames.Remove(InProtocolName);
+		return FText::Format(LOCTEXT("AddBindingToolTip", "Add {0} protocol binding."), { FText::FromName(SelectedProtocolName) });
 	}
-	GetSettings()->SaveConfig();
-
-	Refresh(false);
-}
 
-bool SRCProtocolBindingList::IsProtocolShown(const FName& InProtocolName)
-{
-	return !GetSettings()->HiddenProtocolTypeNames.Contains(InProtocolName);
+	return FText::GetEmpty();
 }
 
 EVisibility SRCProtocolBindingList::GetScrollBarVisibility() const
@@ -288,25 +253,19 @@ void SRCProtocolBindingList::OnStopRecording(TSharedPtr<TStructOnScope<FRemoteCo
 	AwaitingProtocolEntities.Remove(InEntity);
 }
 
-URemoteControlProtocolWidgetsSettings* SRCProtocolBindingList::GetSettings()
-{
-	if(Settings.IsValid())
-	{
-		return Settings.Get();
-	}
-
-	Settings = GetMutableDefault<URemoteControlProtocolWidgetsSettings>();
-	ensure(Settings.IsValid());
-	
-	return Settings.Get();
-}
-
 bool SRCProtocolBindingList::Refresh(bool bNavigateToEnd)
 {
 	FilteredBindings.Empty(FilteredBindings.Num());
-	for(const TSharedPtr<FProtocolBindingViewModel>& ProtocolBinding :  ViewModel->GetFilteredBindings(GetSettings()->HiddenProtocolTypeNames))
+
+	IRemoteControlProtocolWidgetsModule& RCProtocolWidgetsModule = IRemoteControlProtocolWidgetsModule::Get();
+	const FName& SelectedProtocolName = RCProtocolWidgetsModule.GetSelectedProtocolName();
+
+	for (const TSharedPtr<FProtocolBindingViewModel>& ProtocolBinding : ViewModel->GetBindings())
 	{
-		FilteredBindings.Add(ProtocolBinding);
+		if (ProtocolBinding->GetBinding()->GetProtocolName() == SelectedProtocolName)
+		{
+			FilteredBindings.Add(ProtocolBinding);
+		}
 	}
 	
 	if(BindingList.IsValid())
@@ -317,7 +276,6 @@ bool SRCProtocolBindingList::Refresh(bool bNavigateToEnd)
 			return false;
 		}
 
-		
 		GEditor->GetTimerManager()->SetTimerForNextTick(FTimerDelegate::CreateLambda([bNavigateToEnd, WeakListPtr = TWeakPtr<SRCProtocolBindingList>(StaticCastSharedRef<SRCProtocolBindingList>(AsShared()))]()
 		{
 			if (TSharedPtr<SRCProtocolBindingList> ListPtr = WeakListPtr.Pin())
diff --git a/Engine/Plugins/VirtualProduction/RemoteControl/Source/RemoteControlProtocolWidgets/Private/Widgets/SRCProtocolBindingList.h b/Engine/Plugins/VirtualProduction/RemoteControl/Source/RemoteControlProtocolWidgets/Private/Widgets/SRCProtocolBindingList.h
--- a/Engine/Plugins/VirtualProduction/RemoteControl/Source/RemoteControlProtocolWidgets/Private/Widgets/SRCProtocolBindingList.h
+++ b/Engine/Plugins/VirtualProduction/RemoteControl/Source/RemoteControlProtocolWidgets/Private/Widgets/SRCProtocolBindingList.h
@@ -14,9 +14,7 @@ class FProtocolEntityViewModel;
 class IRCTreeNodeViewModel;
 class ITableRow;
 class SRCProtocolBinding;
-class SRCProtocolList;
 class STableViewBase;
-class URemoteControlProtocolWidgetsSettings;
 
 /** The root view for a given entity. A (vertical) list of bindings, where each binding has a protocol. */
 class REMOTECONTROLPROTOCOLWIDGETS_API SRCProtocolBindingList final : public SCompoundWidget, public IRCProtocolBindingList
@@ -44,6 +42,9 @@ private:
 	/** Stop recording incoming protocol message handler */
 	virtual void OnStopRecording(TSharedPtr<TStructOnScope<FRemoteControlProtocolEntity>> InEntity) override;
 
+	/** Refresh protocol binding list. Returns false if the backing widget wasn't valid. */
+	virtual bool Refresh(bool bNavigateToEnd = false) override;
+
 	//~ End IRCProtocolBindingList Interface
 
 private:
@@ -53,11 +54,8 @@ private:
 	/** Check if the bound entity type is supported by Protocol. */
 	bool CanAddProtocol();
 
-	/** Toggles visibility of the given protocol. */
-	void ToggleShowProtocol(const FName& InProtocolName);
-
-	/** Check if the given protocol name is being displayed. */
-	bool IsProtocolShown(const FName& InProtocolName);
+	/** Retrieves the appropriate tooltip based on the active protocol. */
+	FText HandleAddBindingToolTipText() const;
 
 	/** Called to get the visibility of the scrollbar based on options, needs to be dynamic to avoid layout changing on expansion */
 	EVisibility GetScrollBarVisibility() const;
@@ -94,12 +92,6 @@ private:
 		SecondaryColumnWidth = InWidth;
 	}
 
-	/** Get (mutable) module settings */
-	URemoteControlProtocolWidgetsSettings* GetSettings();
-
-	/** Refresh protocol binding list. Returns false if the backing widget wasn't valid. */
-	bool Refresh(bool bNavigateToEnd = false);
-
 private:
 	/** ViewModel for the Protocol Entity. */
 	TSharedPtr<FProtocolEntityViewModel> ViewModel;
@@ -107,12 +99,6 @@ private:
 	/** Current status message, if any */
 	FText StatusMessage;
 
-	/** List of all available protocol names */
-	TArray<FName> ProtocolNames;
-
-	/** Dropdown list of available protocol names to add */
-	TSharedPtr<SRCProtocolList> ProtocolList;
-
 	/** ListView widget for each protocol binding */
 	TSharedPtr<SListView<TSharedPtr<IRCTreeNodeViewModel>>> BindingList;
 
@@ -131,9 +117,6 @@ private:
 	/** Relative width to control secondary splitters */
 	float SecondaryColumnWidth = 0.5f;
 
-	/** Reference to (mutable) settings class for this module */
-	TWeakObjectPtr<URemoteControlProtocolWidgetsSettings> Settings;
-
 	/** Set of protocol entities with awaiting state and waiting for binding */
 	FRemoteControlProtocolEntitySet AwaitingProtocolEntities;
 };
diff --git a/Engine/Plugins/VirtualProduction/RemoteControl/Source/RemoteControlProtocolWidgets/Private/Widgets/SRCProtocolList.cpp b/Engine/Plugins/VirtualProduction/RemoteControl/Source/RemoteControlProtocolWidgets/Private/Widgets/SRCProtocolList.cpp
deleted file mode 100644
diff --git a/Engine/Plugins/VirtualProduction/RemoteControl/Source/RemoteControlProtocolWidgets/Private/Widgets/SRCProtocolList.h b/Engine/Plugins/VirtualProduction/RemoteControl/Source/RemoteControlProtocolWidgets/Private/Widgets/SRCProtocolList.h
deleted file mode 100644
diff --git a/Engine/Plugins/VirtualProduction/RemoteControl/Source/RemoteControlProtocolWidgets/Private/Widgets/SRCProtocolRange.cpp b/Engine/Plugins/VirtualProduction/RemoteControl/Source/RemoteControlProtocolWidgets/Private/Widgets/SRCProtocolRange.cpp
--- a/Engine/Plugins/VirtualProduction/RemoteControl/Source/RemoteControlProtocolWidgets/Private/Widgets/SRCProtocolRange.cpp
+++ b/Engine/Plugins/VirtualProduction/RemoteControl/Source/RemoteControlProtocolWidgets/Private/Widgets/SRCProtocolRange.cpp
@@ -1,4 +1,4 @@
-﻿// Copyright Epic Games, Inc. All Rights Reserved.
+// Copyright Epic Games, Inc. All Rights Reserved.
 
 #include "SRCProtocolRange.h"
 
@@ -29,7 +29,8 @@ void SRCProtocolRange::Construct(const FArguments& InArgs, const TSharedRef<STab
 	SNew(SHorizontalBox)
 	+ SHorizontalBox::Slot()
 	.HAlign(HAlign_Right)
-	.VAlign(VAlign_Top)
+	.VAlign(VAlign_Center)
+	.Padding(4.f, 4.f, 0.f, 0.f)
 	.AutoWidth()
 	[
 		MakeInput()
@@ -41,6 +42,7 @@ void SRCProtocolRange::Construct(const FArguments& InArgs, const TSharedRef<STab
 	+ SHorizontalBox::Slot()
 	.HAlign(HAlign_Left)
 	.VAlign(VAlign_Center)
+	.Padding(4.f, 4.f, 0.f, 0.f)
 	.FillWidth(1.0f)
 	[
 		SNew(SHorizontalBox)
@@ -57,8 +59,7 @@ void SRCProtocolRange::Construct(const FArguments& InArgs, const TSharedRef<STab
 	// Copies current property value
 	+ SHorizontalBox::Slot()
 	.HAlign(HAlign_Left)
-	.VAlign(VAlign_Top)
-	.Padding(0, 6, 0, 0)
+	.VAlign(VAlign_Center)
 	.AutoWidth()
 	[
 		SNew(SButton)
@@ -78,8 +79,7 @@ void SRCProtocolRange::Construct(const FArguments& InArgs, const TSharedRef<STab
 	// Validation warning
 	+ SHorizontalBox::Slot()
 	.HAlign(HAlign_Right)
-	.VAlign(VAlign_Top)
-	.Padding(0, 8, 0, 0)
+	.VAlign(VAlign_Center)
 	.AutoWidth()
 	[
 		SNew(SRCBindingWarning)
@@ -99,8 +99,8 @@ void SRCProtocolRange::Construct(const FArguments& InArgs, const TSharedRef<STab
 	// Delete button container
 	+ SHorizontalBox::Slot()
 	.HAlign(HAlign_Right)
-	.VAlign(VAlign_Top)
-	.Padding(0, 8, 0, 0)
+	.VAlign(VAlign_Center)
+	.Padding(2.f, 0.f)
 	.AutoWidth()
 	[
 		SNew(SButton)
@@ -108,11 +108,12 @@ void SRCProtocolRange::Construct(const FArguments& InArgs, const TSharedRef<STab
 		.ForegroundColor(FSlateColor::UseForeground())
 		.IsFocusable(false)
 		.OnClicked(this, &SRCProtocolRange::OnDelete)
+		.ContentPadding(2.f)
 		.Content()
 		[
-			SNew(STextBlock)
-			.Font(FAppStyle::Get().GetFontStyle("FontAwesome.10"))
-			.Text(FText::FromString(FString(TEXT("\xf00d"))))
+			SNew(SImage)
+			.ColorAndOpacity(FSlateColor::UseForeground())
+			.Image(FAppStyle::GetBrush("Icons.X"))
 		]
 	];
 
diff --git a/Engine/Plugins/VirtualProduction/RemoteControl/Source/RemoteControlProtocolWidgets/Private/Widgets/SRCProtocolRangeList.cpp b/Engine/Plugins/VirtualProduction/RemoteControl/Source/RemoteControlProtocolWidgets/Private/Widgets/SRCProtocolRangeList.cpp
--- a/Engine/Plugins/VirtualProduction/RemoteControl/Source/RemoteControlProtocolWidgets/Private/Widgets/SRCProtocolRangeList.cpp
+++ b/Engine/Plugins/VirtualProduction/RemoteControl/Source/RemoteControlProtocolWidgets/Private/Widgets/SRCProtocolRangeList.cpp
@@ -1,11 +1,11 @@
-﻿// Copyright Epic Games, Inc. All Rights Reserved.
+// Copyright Epic Games, Inc. All Rights Reserved.
 
 #include "SRCProtocolRangeList.h"
-#include "EditorFontGlyphs.h"
 #include "SRCProtocolRange.h"
 #include "ViewModels/ProtocolBindingViewModel.h"
 #include "ViewModels/ProtocolRangeViewModel.h"
 #include "Widgets/SBoxPanel.h"
+#include "Widgets/Images/SImage.h"
 #include "Widgets/Input/SButton.h"
 #include "Widgets/Layout/SSeparator.h"
 #include "Widgets/Layout/SSpacer.h"
@@ -130,13 +130,12 @@ TSharedRef<SWidget> SRCProtocolRangeList::ConstructHeader()
 				ViewModel->AddRangeMapping();
 				return FReply::Handled();
 			})
-			.ContentPadding(FMargin(2.0f, 1.0f))
+			.ContentPadding(FMargin(4.0f, 2.0f))
 			.Content()
 			[
-				SNew(STextBlock)
-				.TextStyle(FAppStyle::Get(), "NormalText.Important")
-				.Font(FAppStyle::Get().GetFontStyle("FontAwesome.10"))
-				.Text(FEditorFontGlyphs::Plus)
+				SNew(SImage)
+				.ColorAndOpacity(FSlateColor::UseForeground())
+				.Image(FAppStyle::GetBrush("Icons.Plus"))
 			]
 		]
 
@@ -154,13 +153,12 @@ TSharedRef<SWidget> SRCProtocolRangeList::ConstructHeader()
 				ViewModel->RemoveAllRangeMappings();
 				return FReply::Handled();
 			})
-			.ContentPadding(FMargin(2.0f, 1.0f))
+			.ContentPadding(FMargin(4.0f, 2.0f))
 			.Content()
 			[
-				SNew(STextBlock)
-				.TextStyle(FAppStyle::Get(), "NormalText.Important")
-				.Font(FAppStyle::Get().GetFontStyle("FontAwesome.10"))
-				.Text(FText::FromString(FString(TEXT("\xf00d"))))
+				SNew(SImage)
+				.ColorAndOpacity(FSlateColor::UseForeground())
+				.Image(FAppStyle::GetBrush("Icons.X"))
 			]
 		];
 
diff --git a/Engine/Plugins/VirtualProduction/RemoteControl/Source/RemoteControlProtocolWidgets/Public/IRCProtocolBindingList.h b/Engine/Plugins/VirtualProduction/RemoteControl/Source/RemoteControlProtocolWidgets/Public/IRCProtocolBindingList.h
--- a/Engine/Plugins/VirtualProduction/RemoteControl/Source/RemoteControlProtocolWidgets/Public/IRCProtocolBindingList.h
+++ b/Engine/Plugins/VirtualProduction/RemoteControl/Source/RemoteControlProtocolWidgets/Public/IRCProtocolBindingList.h
@@ -27,4 +27,6 @@ public:
 	/** Stop recording incoming protocol message handler */
 	virtual void OnStopRecording(TSharedPtr<TStructOnScope<FRemoteControlProtocolEntity>> InEntity) = 0;
 
+	/** Refresh protocol binding list. Returns false if the backing widget wasn't valid. */
+	virtual bool Refresh(bool bNavigateToEnd = false) = 0;
 };
diff --git a/Engine/Plugins/VirtualProduction/RemoteControl/Source/RemoteControlProtocolWidgets/Public/IRemoteControlProtocolWidgetsModule.h b/Engine/Plugins/VirtualProduction/RemoteControl/Source/RemoteControlProtocolWidgets/Public/IRemoteControlProtocolWidgetsModule.h
--- a/Engine/Plugins/VirtualProduction/RemoteControl/Source/RemoteControlProtocolWidgets/Public/IRemoteControlProtocolWidgetsModule.h
+++ b/Engine/Plugins/VirtualProduction/RemoteControl/Source/RemoteControlProtocolWidgets/Public/IRemoteControlProtocolWidgetsModule.h
@@ -13,6 +13,8 @@ class SWidget;
 
 DECLARE_MULTICAST_DELEGATE_OneParam(FOnProtocolBindingAddedOrRemoved, ERCProtocolBinding::Op /* BindingOperation */);
 
+DECLARE_MULTICAST_DELEGATE_OneParam(FOnActiveProtocolChanged, const FName /* NewActiveProtocolName */);
+
 /** A Remote Control module that provides editor widgets for protocol bindings. */
 class IRemoteControlProtocolWidgetsModule : public IModuleInterface
 {
@@ -37,7 +39,13 @@ public:
 
 	/** Get the binding list public reference */
 	virtual TSharedPtr<IRCProtocolBindingList> GetProtocolBindingList() const = 0;
+	
+	/** Get the selected protocol name. */
+	virtual const FName GetSelectedProtocolName() const = 0;
 
 	/** Called when binding is added or removed. */
 	virtual FOnProtocolBindingAddedOrRemoved& OnProtocolBindingAddedOrRemoved() = 0;
+
+	/** Called when active protocol selection changed. */
+	virtual FOnActiveProtocolChanged& OnActiveProtocolChanged() = 0;
 };
diff --git a/Engine/Plugins/VirtualProduction/RemoteControl/Source/RemoteControlUI/Private/UI/SRCPanelExposedEntitiesList.cpp b/Engine/Plugins/VirtualProduction/RemoteControl/Source/RemoteControlUI/Private/UI/SRCPanelExposedEntitiesList.cpp
--- a/Engine/Plugins/VirtualProduction/RemoteControl/Source/RemoteControlUI/Private/UI/SRCPanelExposedEntitiesList.cpp
+++ b/Engine/Plugins/VirtualProduction/RemoteControl/Source/RemoteControlUI/Private/UI/SRCPanelExposedEntitiesList.cpp
@@ -343,6 +343,10 @@ void SRCPanelExposedEntitiesList::Construct(const FArguments& InArgs, URemoteCon
 
 					ActiveProtocol = NewMode.ModeId;
 
+					IRemoteControlProtocolWidgetsModule& RCProtocolWidgetsModule = IRemoteControlProtocolWidgetsModule::Get();
+
+					RCProtocolWidgetsModule.OnActiveProtocolChanged().Broadcast(ActiveProtocol);
+
 					if (bIsInProtocolsMode.Get())
 					{
 						RebuildListWithColumns(EEntitiesListMode::Protocols);
@@ -373,6 +377,10 @@ void SRCPanelExposedEntitiesList::Construct(const FArguments& InArgs, URemoteCon
 			if (bIsDefault)
 			{
 				ActiveProtocol = ProtocolName;
+	
+				IRemoteControlProtocolWidgetsModule& RCProtocolWidgetsModule = IRemoteControlProtocolWidgetsModule::Get();
+
+				RCProtocolWidgetsModule.OnActiveProtocolChanged().Broadcast(ActiveProtocol);
 			}
 
 			SRCModeSwitcher::FRCMode::FArguments NewMode = SRCModeSwitcher::Mode(ProtocolName)
