Remote Control Protocols - Fix Crash by deletion

diff --git a/Engine/Plugins/VirtualProduction/RemoteControl/Source/RemoteControlProtocolWidgets/Private/ViewModels/ProtocolBindingViewModel.cpp b/Engine/Plugins/VirtualProduction/RemoteControl/Source/RemoteControlProtocolWidgets/Private/ViewModels/ProtocolBindingViewModel.cpp
--- a/Engine/Plugins/VirtualProduction/RemoteControl/Source/RemoteControlProtocolWidgets/Private/ViewModels/ProtocolBindingViewModel.cpp
+++ b/Engine/Plugins/VirtualProduction/RemoteControl/Source/RemoteControlProtocolWidgets/Private/ViewModels/ProtocolBindingViewModel.cpp
@@ -190,7 +190,8 @@ void FProtocolBindingViewModel::Initialize()
 	// May be stale as a result of an undo deleting it.
 	if(Binding)
 	{
-		if (!RemoteControlTypeUtilities::IsSupportedMappingType(GetProperty().Get()))
+		const FProperty* MappingProperty = GetProperty().Get();
+		if (!MappingProperty || !RemoteControlTypeUtilities::IsSupportedMappingType(MappingProperty))
 		{
 			return;
 		}
@@ -370,6 +371,11 @@ void FProtocolBindingViewModel::AddDefaultRangeMappings()
 	}
 
 	const FProperty* MappingProperty = GetProperty().Get();
+	if (!MappingProperty)
+	{
+		return;
+	}
+	
 	FName MappingPropertyTypeName = MappingProperty->GetClass()->GetFName();
 	if(const FStructProperty* StructProperty = CastField<FStructProperty>(MappingProperty))
 	{
