Copy/Paste/Duplicate functionality for Remote Control Logic

For original code review please see: https://p4-swarm.epicgames.net/reviews/21749275
(this CL has been migrated over from Main branch, where the review originated)

diff --git a/Engine/Plugins/VirtualProduction/RemoteControl/Source/RemoteControl/Private/RCVirtualProperty.cpp b/Engine/Plugins/VirtualProduction/RemoteControl/Source/RemoteControl/Private/RCVirtualProperty.cpp
--- a/Engine/Plugins/VirtualProduction/RemoteControl/Source/RemoteControl/Private/RCVirtualProperty.cpp
+++ b/Engine/Plugins/VirtualProduction/RemoteControl/Source/RemoteControl/Private/RCVirtualProperty.cpp
@@ -231,6 +231,11 @@ bool URCVirtualPropertyBase::CopyCompleteValue(const FProperty* InTargetProperty
 
 }
 
+bool URCVirtualPropertyBase::CopyCompleteValue(URCVirtualPropertyBase* InVirtualProperty)
+{
+	return CopyCompleteValue(InVirtualProperty->GetProperty(), InVirtualProperty->GetValuePtr());
+}
+
 struct FRCVirtualPropertyCastHelpers
 {
 	template<typename TPropertyType, typename TValueType>
diff --git a/Engine/Plugins/VirtualProduction/RemoteControl/Source/RemoteControl/Private/RCVirtualPropertyContainer.cpp b/Engine/Plugins/VirtualProduction/RemoteControl/Source/RemoteControl/Private/RCVirtualPropertyContainer.cpp
--- a/Engine/Plugins/VirtualProduction/RemoteControl/Source/RemoteControl/Private/RCVirtualPropertyContainer.cpp
+++ b/Engine/Plugins/VirtualProduction/RemoteControl/Source/RemoteControl/Private/RCVirtualPropertyContainer.cpp
@@ -18,7 +18,8 @@ URCVirtualPropertyInContainer* URCVirtualPropertyContainerBase::AddProperty(cons
 	const FName PropertyName = GenerateUniquePropertyName(InPropertyName, InValueType, InValueTypeObject, this);
 	Bag.AddProperty(PropertyName, InValueType, InValueTypeObject);
 
-	const FPropertyBagPropertyDesc* BagPropertyDesc = Bag.FindPropertyDescByName(InPropertyName);
+	// Ensure that the property has been succcefully added to the Bag
+	const FPropertyBagPropertyDesc* BagPropertyDesc = Bag.FindPropertyDescByName(PropertyName);
 	if (!ensure(BagPropertyDesc))
 	{
 		return nullptr;
@@ -39,8 +40,9 @@ URCVirtualPropertyInContainer* URCVirtualPropertyContainerBase::AddProperty(cons
 
 URCVirtualPropertyInContainer* URCVirtualPropertyContainerBase::DuplicateProperty(const FName& InPropertyName, const FProperty* InSourceProperty, TSubclassOf<URCVirtualPropertyInContainer> InPropertyClass)
 {
+	// Ensure that the property being duplicated is not already a part of the Bag (not to be confused with a similar looking ensure performed above in AddProperty, this is the exact inverse!)
 	const FPropertyBagPropertyDesc* BagPropertyDesc = Bag.FindPropertyDescByName(InPropertyName);
-	if (ensure(!BagPropertyDesc))
+	if (!ensure(!BagPropertyDesc))
 	{
 		return nullptr;
 	}
@@ -80,6 +82,31 @@ URCVirtualPropertyInContainer* URCVirtualPropertyContainerBase::DuplicatePropert
 	return VirtualPropertyInContainer;
 }
 
+URCVirtualPropertyInContainer* URCVirtualPropertyContainerBase::DuplicateVirtualProperty(URCVirtualPropertyInContainer* InVirtualProperty)
+{
+	if (URCVirtualPropertyInContainer* NewVirtualProperty = DuplicateObject<URCVirtualPropertyInContainer>(InVirtualProperty, InVirtualProperty->GetOuter()))
+	{
+		NewVirtualProperty->PropertyName = GenerateUniquePropertyName(InVirtualProperty->PropertyName, this);
+
+		// Sync Property Bag
+		Bag.AddProperty(NewVirtualProperty->PropertyName, InVirtualProperty->GetProperty());
+
+		// Ensure that the property has been successfully added to the Bag
+		const FPropertyBagPropertyDesc* BagPropertyDesc = Bag.FindPropertyDescByName(NewVirtualProperty->PropertyName);
+		if (!ensure(BagPropertyDesc))
+		{
+			return nullptr;
+		}
+
+		// Sync Virtual Properties List
+		VirtualProperties.Add(NewVirtualProperty);
+
+		return NewVirtualProperty;
+	}
+
+	return nullptr;
+}
+
 bool URCVirtualPropertyContainerBase::RemoveProperty(const FName& InPropertyName)
 {
 	Bag.RemovePropertyByName(InPropertyName);
@@ -161,6 +188,17 @@ TSharedPtr<FStructOnScope> URCVirtualPropertyContainerBase::CreateStructOnScope(
 }
 
 FName URCVirtualPropertyContainerBase::GenerateUniquePropertyName(const FName& InPropertyName, const EPropertyBagPropertyType InValueType, UObject* InValueTypeObject, const URCVirtualPropertyContainerBase* InContainer)
+{
+	FName BaseName = InPropertyName;
+	if (BaseName.IsNone())
+	{
+		BaseName = URCVirtualPropertyBase::GetVirtualPropertyTypeDisplayName(InValueType, InValueTypeObject);
+	}
+
+	return GenerateUniquePropertyName(BaseName, InContainer);
+}
+
+FName URCVirtualPropertyContainerBase::GenerateUniquePropertyName(const FName& InPropertyName, const URCVirtualPropertyContainerBase* InContainer)
 {
 	auto GetFinalName = [](const FString& InPrefix, const int32 InIndex = 0)
 	{
@@ -174,13 +212,10 @@ FName URCVirtualPropertyContainerBase::GenerateUniquePropertyName(const FName& I
 		return FinalName;
 	};
 
-
-	FName DefaultName = URCVirtualPropertyBase::GetVirtualPropertyTypeDisplayName(InValueType, InValueTypeObject);
-	
 	int32 Index = 0;
-	const FString InitialName = InPropertyName.IsNone() ? DefaultName.ToString() : InPropertyName.ToString();
+	const FString InitialName = InPropertyName.ToString();
 	FString FinalName = InitialName;
-	
+
 	// Recursively search for an available name by incrementing suffix till we find one.
 	const FPropertyBagPropertyDesc* PropertyDesc = InContainer->Bag.FindPropertyDescByName(*FinalName);
 	while (PropertyDesc)
diff --git a/Engine/Plugins/VirtualProduction/RemoteControl/Source/RemoteControl/Private/RemoteControlPreset.cpp b/Engine/Plugins/VirtualProduction/RemoteControl/Source/RemoteControl/Private/RemoteControlPreset.cpp
--- a/Engine/Plugins/VirtualProduction/RemoteControl/Source/RemoteControl/Private/RemoteControlPreset.cpp
+++ b/Engine/Plugins/VirtualProduction/RemoteControl/Source/RemoteControl/Private/RemoteControlPreset.cpp
@@ -574,6 +574,16 @@ bool URemoteControlPreset::RemoveVirtualProperty(const FName& InPropertyName)
 	return false;
 }
 
+URCVirtualPropertyInContainer* URemoteControlPreset::DuplicateVirtualProperty(URCVirtualPropertyInContainer* InVirtualProperty)
+{
+	if (ensure(ControllerContainer))
+	{
+		return ControllerContainer->DuplicateVirtualProperty(InVirtualProperty);
+	}
+
+	return nullptr;
+}
+
 void URemoteControlPreset::ResetVirtualProperties()
 {
 	if (ensure(ControllerContainer))
diff --git a/Engine/Plugins/VirtualProduction/RemoteControl/Source/RemoteControl/Public/RCVirtualProperty.h b/Engine/Plugins/VirtualProduction/RemoteControl/Source/RemoteControl/Public/RCVirtualProperty.h
--- a/Engine/Plugins/VirtualProduction/RemoteControl/Source/RemoteControl/Public/RCVirtualProperty.h
+++ b/Engine/Plugins/VirtualProduction/RemoteControl/Source/RemoteControl/Public/RCVirtualProperty.h
@@ -102,10 +102,20 @@ public:
 
 	/** Copy this virtual property's data onto a given FProperty
 	* 
-	* * @param InTargetProperty - The property onto which our value is to be copied
-	* * @param InTargetValuePtr - The memory location for the target property
+	* @param InTargetProperty - The property onto which our value is to be copied
+	* @param InTargetValuePtr - The memory location for the target property
+	* 
+	* @return true if copied successfully
 	*/
 	bool CopyCompleteValue(const FProperty* InTargetProperty, uint8* InTargetValuePtr);
+
+	/** Copy this virtual property's data onto a given Virtual Property
+	*
+	* @param InVirtualProperty - The virtual property onto which our value is to be copied
+	* 
+	* @return true if copied successfully
+	*/
+	bool CopyCompleteValue(URCVirtualPropertyBase* InVirtualProperty);
 	
 	/** Get Bool value from Virtual Property */
 	UFUNCTION(BlueprintCallable, Category = "Remote Control Behaviour")
diff --git a/Engine/Plugins/VirtualProduction/RemoteControl/Source/RemoteControl/Public/RCVirtualPropertyContainer.h b/Engine/Plugins/VirtualProduction/RemoteControl/Source/RemoteControl/Public/RCVirtualPropertyContainer.h
--- a/Engine/Plugins/VirtualProduction/RemoteControl/Source/RemoteControl/Public/RCVirtualPropertyContainer.h
+++ b/Engine/Plugins/VirtualProduction/RemoteControl/Source/RemoteControl/Public/RCVirtualPropertyContainer.h
@@ -38,7 +38,7 @@ public:
 	virtual URCVirtualPropertyInContainer* AddProperty(const FName& InPropertyName, TSubclassOf<URCVirtualPropertyInContainer> InPropertyClass, const EPropertyBagPropertyType InValueType, UObject* InValueTypeObject = nullptr);
 
 	/**
-	 * Duplicates property from give Property.
+	 * Duplicates a property from given FProperty.
 	 * That will add property to Property Bag and Create Remote Control Virtual Property
 	 *
 	 * @param InPropertyName				Name of the property
@@ -50,7 +50,7 @@ public:
 	virtual URCVirtualPropertyInContainer* DuplicateProperty(const FName& InPropertyName, const FProperty* InSourceProperty, TSubclassOf<URCVirtualPropertyInContainer> InPropertyClass);
 
 	/**
-	 * Duplicates property from give Property and copy the property value
+	 * Duplicates a property from a given FProperty and copies the property value
 	 * That will add property to Property Bag and Create Remote Control Virtual Property
 	 *
 	 * @param InPropertyName				Name of the property
@@ -62,6 +62,17 @@ public:
 	 */
 	virtual URCVirtualPropertyInContainer* DuplicatePropertyWithCopy(const FName& InPropertyName, const FProperty* InSourceProperty, const uint8* InSourceContainerPtr, TSubclassOf<URCVirtualPropertyInContainer> InPropertyClass);
 
+	/**
+	 * Given a virtual property this funciton duplicates it (via DuplicateObject)
+	 * and syncs the internal Controller Container and PropertyBag containers as well.
+	 *
+	 * @param InVirtualProperty				Virtual Property to be duplicated
+	 *
+	 * @return Virtual Property Object
+	 */
+	virtual URCVirtualPropertyInContainer* DuplicateVirtualProperty(URCVirtualPropertyInContainer* InVirtualProperty);
+
+
 	/**
 	 * Removes a property from the container by name if it exists.
 	 */
@@ -100,6 +111,11 @@ public:
 	 */
 	static FName GenerateUniquePropertyName(const FName& InPropertyName, const EPropertyBagPropertyType InValueType, UObject* InValueTypeObject, const URCVirtualPropertyContainerBase* InContainer);
 
+	/**
+	* Generates unique name for the property for specified property container
+	*/
+	static FName GenerateUniquePropertyName(const FName& InPropertyName, const URCVirtualPropertyContainerBase* InContainer);
+
 #if WITH_EDITOR
 	/** Called after applying a transaction to the object. Used to broadcast Undo related container changes to UI. */
 	virtual void PostEditUndo();
@@ -128,7 +144,6 @@ public:
 
 	/** Delegate that notifies changes to the virtual property container*/
 	FOnVirtualPropertyContainerModified OnVirtualPropertyContainerModifiedDelegate;
-
 };
 
 
diff --git a/Engine/Plugins/VirtualProduction/RemoteControl/Source/RemoteControl/Public/RemoteControlPreset.h b/Engine/Plugins/VirtualProduction/RemoteControl/Source/RemoteControl/Public/RemoteControlPreset.h
--- a/Engine/Plugins/VirtualProduction/RemoteControl/Source/RemoteControl/Public/RemoteControlPreset.h
+++ b/Engine/Plugins/VirtualProduction/RemoteControl/Source/RemoteControl/Public/RemoteControlPreset.h
@@ -635,6 +635,9 @@ public:
 	/** Removes a given Virtual Property (by Name) from the Remote Control preset */
 	bool RemoveVirtualProperty(const FName& InPropertyName);
 
+	/** Duplicates a given Virtual Property from the Remote Control preset */
+	URCVirtualPropertyInContainer* DuplicateVirtualProperty(URCVirtualPropertyInContainer* InVirtualProperty);
+
 	/** Removes all virtual properties held by this Remote Control preset*/
 	void ResetVirtualProperties();
 
diff --git a/Engine/Plugins/VirtualProduction/RemoteControl/Source/RemoteControlLogic/Private/Action/RCActionContainer.cpp b/Engine/Plugins/VirtualProduction/RemoteControl/Source/RemoteControlLogic/Private/Action/RCActionContainer.cpp
--- a/Engine/Plugins/VirtualProduction/RemoteControl/Source/RemoteControlLogic/Private/Action/RCActionContainer.cpp
+++ b/Engine/Plugins/VirtualProduction/RemoteControl/Source/RemoteControlLogic/Private/Action/RCActionContainer.cpp
@@ -2,6 +2,7 @@
 
 #include "Action/RCActionContainer.h"
 
+#include "Behaviour/RCBehaviour.h"
 #include "Controller/RCController.h"
 #include "IRemoteControlModule.h"
 #include "RCVirtualProperty.h"
@@ -42,9 +43,9 @@ URCAction* URCActionContainer::AddAction(const TSharedRef<const FRemoteControlFi
 }
 
 
-URCAction* URCActionContainer::AddAction(TRCActionUniquenessTest IsUnique, const TSharedRef<const FRemoteControlField> InRemoteControlField)
+URCAction* URCActionContainer::AddAction(TRCActionUniquenessTest InUniquenessTest, const TSharedRef<const FRemoteControlField> InRemoteControlField)
 {
-	if (!IsUnique(Actions))
+	if (!InUniquenessTest(Actions))
 	{
 		return nullptr;
 	}
@@ -63,6 +64,11 @@ URCAction* URCActionContainer::AddAction(TRCActionUniquenessTest IsUnique, const
 	return NewAction;
 }
 
+URCBehaviour* URCActionContainer::GetParentBehaviour()
+{
+	return Cast<URCBehaviour>(GetOuter());
+}
+
 URCPropertyAction* URCActionContainer::AddPropertyAction(const TSharedRef<const FRemoteControlProperty> InRemoteControlProperty)
 {
 	// Create new Property
diff --git a/Engine/Plugins/VirtualProduction/RemoteControl/Source/RemoteControlLogic/Private/Action/RCPropertyAction.cpp b/Engine/Plugins/VirtualProduction/RemoteControl/Source/RemoteControlLogic/Private/Action/RCPropertyAction.cpp
--- a/Engine/Plugins/VirtualProduction/RemoteControl/Source/RemoteControlLogic/Private/Action/RCPropertyAction.cpp
+++ b/Engine/Plugins/VirtualProduction/RemoteControl/Source/RemoteControlLogic/Private/Action/RCPropertyAction.cpp
@@ -2,6 +2,7 @@
 
 #include "Action/RCPropertyAction.h"
 
+#include "Action/RCActionContainer.h"
 #include "Controller/RCController.h"
 #include "IRemoteControlModule.h"
 #include "IRemoteControlPropertyHandle.h"
@@ -18,6 +19,17 @@ URCPropertyAction::URCPropertyAction()
 	PropertySelfContainer = CreateDefaultSubobject<URCVirtualPropertySelfContainer>(FName("VirtualPropertySelfContainer"));
 }
 
+URCBehaviour* URCAction::GetParentBehaviour()
+{
+	URCActionContainer* ActionContainer = Cast<URCActionContainer>(GetOuter());
+	if (ensure(ActionContainer))
+	{
+		return ActionContainer->GetParentBehaviour();
+	}
+
+	return nullptr;
+}
+
 void URCPropertyAction::Execute() const
 {
 	if (!PresetWeakPtr.IsValid())
diff --git a/Engine/Plugins/VirtualProduction/RemoteControl/Source/RemoteControlLogic/Private/Behaviour/Builtin/Conditional/RCBehaviourConditional.cpp b/Engine/Plugins/VirtualProduction/RemoteControl/Source/RemoteControlLogic/Private/Behaviour/Builtin/Conditional/RCBehaviourConditional.cpp
--- a/Engine/Plugins/VirtualProduction/RemoteControl/Source/RemoteControlLogic/Private/Behaviour/Builtin/Conditional/RCBehaviourConditional.cpp
+++ b/Engine/Plugins/VirtualProduction/RemoteControl/Source/RemoteControlLogic/Private/Behaviour/Builtin/Conditional/RCBehaviourConditional.cpp
@@ -126,6 +126,31 @@ void URCBehaviourConditional::Execute()
 	}
 }
 
+URCAction* URCBehaviourConditional::DuplicateAction(URCAction* InAction, URCBehaviour* InBehaviour)
+{
+	URCBehaviourConditional* InBehaviourConditional = Cast<URCBehaviourConditional>(InBehaviour);
+	if (!ensureMsgf(InBehaviourConditional, TEXT("Expected Behaviour of the same type (Conditional) For CopyAction operation!")))
+	{
+		return nullptr;
+	}
+
+	URCAction* NewAction = Super::DuplicateAction(InAction, InBehaviour);
+	if (ensure(NewAction))
+	{
+		// Copy Action specific data residing in Conditional Behaviour:
+		if (FRCBehaviourCondition* ConditionData = this->Conditions.Find(InAction))
+		{
+			// Virtual Property for the Comparand
+			Comparand = NewObject<URCVirtualPropertySelfContainer>(InBehaviourConditional);
+			Comparand->DuplicatePropertyWithCopy(ConditionData->Comparand);
+
+			InBehaviourConditional->OnActionAdded(NewAction, ConditionData->ConditionType, Comparand);
+		}
+	}
+
+	return NewAction;
+}
+
 void URCBehaviourConditional::OnActionAdded(URCAction* Action, const ERCBehaviourConditionType InConditionType, const TObjectPtr<URCVirtualPropertySelfContainer> InComparand)
 {
 	FRCBehaviourCondition Condition(InConditionType, InComparand);
diff --git a/Engine/Plugins/VirtualProduction/RemoteControl/Source/RemoteControlLogic/Private/Behaviour/Builtin/RangeMap/RCRangeMapBehaviour.cpp b/Engine/Plugins/VirtualProduction/RemoteControl/Source/RemoteControlLogic/Private/Behaviour/Builtin/RangeMap/RCRangeMapBehaviour.cpp
--- a/Engine/Plugins/VirtualProduction/RemoteControl/Source/RemoteControlLogic/Private/Behaviour/Builtin/RangeMap/RCRangeMapBehaviour.cpp
+++ b/Engine/Plugins/VirtualProduction/RemoteControl/Source/RemoteControlLogic/Private/Behaviour/Builtin/RangeMap/RCRangeMapBehaviour.cpp
@@ -274,6 +274,37 @@ void URCRangeMapBehaviour::Execute()
 	}
 }
 
+URCAction* URCRangeMapBehaviour::DuplicateAction(URCAction* InAction, URCBehaviour* InBehaviour)
+{
+	URCRangeMapBehaviour* InBehaviourRangeMap = Cast< URCRangeMapBehaviour>(InBehaviour);
+	if (!ensureMsgf(InBehaviourRangeMap, TEXT("Expected Behaviour of the same type (Range Map) For CopyAction operation!")))
+	{
+		return nullptr;
+	}
+
+	URCAction* NewAction = Super::DuplicateAction(InAction, InBehaviour);
+	if (ensure(NewAction))
+	{
+		// Copy action specific data residing in Range Map Behaviour:
+		if (FRCRangeMapStep* RangeMapStepData = RangeMapActionContainer.Find(InAction))
+		{
+			URCVirtualPropertySelfContainer* ActionValueProperty = nullptr;
+			if (URCPropertyAction* PropertyAction = Cast<URCPropertyAction>(NewAction))
+			{
+				ActionValueProperty = PropertyAction->PropertySelfContainer;
+			}
+
+			if(ensureAlways(ActionValueProperty))
+			{
+				FRCRangeMapStep NewData(RangeMapStepData->StepValue, ActionValueProperty);
+				InBehaviourRangeMap->RangeMapActionContainer.Add(NewAction, NewData);
+			}
+		}
+	}
+
+	return NewAction;
+}
+
 URCAction* URCRangeMapBehaviour::AddAction(const TSharedRef<const FRemoteControlField> InRemoteControlField)
 {
 	double InStepValue = 0.0;
diff --git a/Engine/Plugins/VirtualProduction/RemoteControl/Source/RemoteControlLogic/Private/Behaviour/RCBehaviour.cpp b/Engine/Plugins/VirtualProduction/RemoteControl/Source/RemoteControlLogic/Private/Behaviour/RCBehaviour.cpp
--- a/Engine/Plugins/VirtualProduction/RemoteControl/Source/RemoteControlLogic/Private/Behaviour/RCBehaviour.cpp
+++ b/Engine/Plugins/VirtualProduction/RemoteControl/Source/RemoteControlLogic/Private/Behaviour/RCBehaviour.cpp
@@ -6,8 +6,10 @@
 #include "Action/RCFunctionAction.h"
 #include "Action/RCPropertyAction.h"
 #include "Behaviour/RCBehaviourNode.h"
+#include "Controller/RCController.h"
 #include "Engine/Blueprint.h"
 #include "RemoteControlField.h"
+#include "RemoteControlPreset.h"
 
 URCBehaviour::URCBehaviour()
 {
@@ -37,6 +39,19 @@ URCAction* URCBehaviour::AddAction(const TSharedRef<const FRemoteControlField> I
 	return ActionContainer->AddAction(InRemoteControlField);
 }
 
+URCAction* URCBehaviour::DuplicateAction(URCAction* InAction, URCBehaviour* InBehaviour)
+{
+	URCAction* NewAction = DuplicateObject<URCAction>(InAction, InBehaviour->ActionContainer);
+	if(ensure(NewAction))
+	{
+		InBehaviour->ActionContainer->AddAction(NewAction);
+
+		return NewAction;
+	}
+
+	return nullptr;
+}
+
 int32 URCBehaviour::GetNumActions() const
 {
 	return ActionContainer->GetActions().Num();
@@ -121,4 +136,5 @@ const FText& URCBehaviour::GetBehaviorDescription()
 
 	return GetDefault<URCBehaviourNode>(BehaviourNodeClass)->BehaviorDescription;
 }
+
 #endif
\ No newline at end of file
diff --git a/Engine/Plugins/VirtualProduction/RemoteControl/Source/RemoteControlLogic/Private/Controller/RCController.cpp b/Engine/Plugins/VirtualProduction/RemoteControl/Source/RemoteControlLogic/Private/Controller/RCController.cpp
--- a/Engine/Plugins/VirtualProduction/RemoteControl/Source/RemoteControlLogic/Private/Controller/RCController.cpp
+++ b/Engine/Plugins/VirtualProduction/RemoteControl/Source/RemoteControlLogic/Private/Controller/RCController.cpp
@@ -89,4 +89,15 @@ void URCController::ExecuteBehaviours()
 	}
 }
 
+URCBehaviour* URCController::DuplicateBehaviour(URCController* InController, URCBehaviour* InBehaviour)
+{
+	URCBehaviour* NewBehaviour = DuplicateObject<URCBehaviour>(InBehaviour, InController);
+
+	NewBehaviour->ControllerWeakPtr = InController;
+
+	InController->Behaviours.Add(NewBehaviour);
+
+	return NewBehaviour;
+}
+
 #undef LOCTEXT_NAMESPACE /* RCController */ 
\ No newline at end of file
diff --git a/Engine/Plugins/VirtualProduction/RemoteControl/Source/RemoteControlLogic/Public/Action/RCAction.h b/Engine/Plugins/VirtualProduction/RemoteControl/Source/RemoteControlLogic/Public/Action/RCAction.h
--- a/Engine/Plugins/VirtualProduction/RemoteControl/Source/RemoteControlLogic/Public/Action/RCAction.h
+++ b/Engine/Plugins/VirtualProduction/RemoteControl/Source/RemoteControlLogic/Public/Action/RCAction.h
@@ -5,6 +5,7 @@
 #include "CoreMinimal.h"
 #include "RCAction.generated.h"
 
+class URCBehaviour;
 class URemoteControlPreset;
 
 /**
@@ -21,6 +22,11 @@ public:
 
 	FName GetExposedFieldLabel() const;
 
+	virtual void CopyTo(URCAction* Action) {}
+
+	/** Returns the parent Behaviour associated with this Action */
+	URCBehaviour* GetParentBehaviour();
+
 public:
 	/** Exposed Property or Function field Id*/
 	UPROPERTY()
diff --git a/Engine/Plugins/VirtualProduction/RemoteControl/Source/RemoteControlLogic/Public/Action/RCActionContainer.h b/Engine/Plugins/VirtualProduction/RemoteControl/Source/RemoteControlLogic/Public/Action/RCActionContainer.h
--- a/Engine/Plugins/VirtualProduction/RemoteControl/Source/RemoteControlLogic/Public/Action/RCActionContainer.h
+++ b/Engine/Plugins/VirtualProduction/RemoteControl/Source/RemoteControlLogic/Public/Action/RCActionContainer.h
@@ -9,6 +9,7 @@
 
 class URemoteControlPreset;
 class URCAction;
+class URCBehaviour;
 class URCController;
 class URCFunctionAction;
 class URCPropertyAction;
@@ -40,7 +41,7 @@ public:
 	URCAction* AddAction(const TSharedRef<const FRemoteControlField> InRemoteControlField);
 
 	/** Add remote control property action  */
-	URCAction* AddAction(TRCActionUniquenessTest IsUnique, const TSharedRef<const FRemoteControlField> InRemoteControlField);
+	URCAction* AddAction(TRCActionUniquenessTest InUniquenessTest, const TSharedRef<const FRemoteControlField> InRemoteControlField);
 
 	/** Find Action by given exposed filed id */
 	URCAction* FindActionByFieldId(const FGuid InFieldId) const;
@@ -79,6 +80,9 @@ public:
 	/** Delegate that notifies changes to the list of actions to a single listener*/
 	FOnActionsListModified OnActionsListModified;
 
+	/** Derive the parent Behaviour holding this Action Container */
+	URCBehaviour* GetParentBehaviour();
+
 private:
 	/** Add remote control property action  */
 	URCPropertyAction* AddPropertyAction(const TSharedRef<const FRemoteControlProperty> InRemoteControlProperty);
diff --git a/Engine/Plugins/VirtualProduction/RemoteControl/Source/RemoteControlLogic/Public/Behaviour/Builtin/Conditional/RCBehaviourConditional.h b/Engine/Plugins/VirtualProduction/RemoteControl/Source/RemoteControlLogic/Public/Behaviour/Builtin/Conditional/RCBehaviourConditional.h
--- a/Engine/Plugins/VirtualProduction/RemoteControl/Source/RemoteControlLogic/Public/Behaviour/Builtin/Conditional/RCBehaviourConditional.h
+++ b/Engine/Plugins/VirtualProduction/RemoteControl/Source/RemoteControlLogic/Public/Behaviour/Builtin/Conditional/RCBehaviourConditional.h
@@ -66,7 +66,10 @@ public:
 	virtual void Initialize() override;
 
 	/** Execute the behaviour */
-	virtual void Execute() override;	
+	virtual void Execute() override;
+
+	/** Duplicates an Action belonging to us into a given target behaviour*/
+	virtual URCAction* DuplicateAction(URCAction* InAction, URCBehaviour* InBehaviour) override;
 
 	//~ End URCBehaviour interface
 
diff --git a/Engine/Plugins/VirtualProduction/RemoteControl/Source/RemoteControlLogic/Public/Behaviour/Builtin/RangeMap/RCRangeMapBehaviour.h b/Engine/Plugins/VirtualProduction/RemoteControl/Source/RemoteControlLogic/Public/Behaviour/Builtin/RangeMap/RCRangeMapBehaviour.h
--- a/Engine/Plugins/VirtualProduction/RemoteControl/Source/RemoteControlLogic/Public/Behaviour/Builtin/RangeMap/RCRangeMapBehaviour.h
+++ b/Engine/Plugins/VirtualProduction/RemoteControl/Source/RemoteControlLogic/Public/Behaviour/Builtin/RangeMap/RCRangeMapBehaviour.h
@@ -60,6 +60,9 @@ public:
 
 	/** Executes the Behaviour */
 	virtual void Execute() override;
+
+	/** Duplicates an Action belonging to us into a given target behaviour*/
+	virtual URCAction* DuplicateAction(URCAction* InAction, URCBehaviour* InBehaviour) override;
 	//~ End URCBehaviour interface
 
 	/** Refresh function being called whenever either the Controller or the Properties of the Behaviour Details Panel change */
diff --git a/Engine/Plugins/VirtualProduction/RemoteControl/Source/RemoteControlLogic/Public/Behaviour/RCBehaviour.h b/Engine/Plugins/VirtualProduction/RemoteControl/Source/RemoteControlLogic/Public/Behaviour/RCBehaviour.h
--- a/Engine/Plugins/VirtualProduction/RemoteControl/Source/RemoteControlLogic/Public/Behaviour/RCBehaviour.h
+++ b/Engine/Plugins/VirtualProduction/RemoteControl/Source/RemoteControlLogic/Public/Behaviour/RCBehaviour.h
@@ -44,6 +44,9 @@ public:
 	/** Add a Logic action using a remote control field as input */
 	virtual URCAction* AddAction(const TSharedRef<const FRemoteControlField> InRemoteControlField);
 
+	/** Duplicates an Action belonging to us into a given target behaviour*/
+	virtual URCAction* DuplicateAction(URCAction* InAction, URCBehaviour* InBehaviour);
+
 	/** Get number of action associated with behaviour */
 	int32 GetNumActions() const;
 
diff --git a/Engine/Plugins/VirtualProduction/RemoteControl/Source/RemoteControlLogic/Public/Controller/RCController.h b/Engine/Plugins/VirtualProduction/RemoteControl/Source/RemoteControlLogic/Public/Controller/RCController.h
--- a/Engine/Plugins/VirtualProduction/RemoteControl/Source/RemoteControlLogic/Public/Controller/RCController.h
+++ b/Engine/Plugins/VirtualProduction/RemoteControl/Source/RemoteControlLogic/Public/Controller/RCController.h
@@ -48,6 +48,9 @@ public:
 		ExecuteBehaviours();
 	}
 
+	/** Duplicates an existing behaviour and adds it to the behaviour set of a given Controller */
+	static URCBehaviour* DuplicateBehaviour(URCController* InController, URCBehaviour* InBehaviour);
+
 	/** Delegate that notifies changes to the list of behaviours*/
 	FOnBehaviourListModified OnBehaviourListModified;
 
diff --git a/Engine/Plugins/VirtualProduction/RemoteControl/Source/RemoteControlUI/Private/Commands/RemoteControlCommands.cpp b/Engine/Plugins/VirtualProduction/RemoteControl/Source/RemoteControlUI/Private/Commands/RemoteControlCommands.cpp
--- a/Engine/Plugins/VirtualProduction/RemoteControl/Source/RemoteControlUI/Private/Commands/RemoteControlCommands.cpp
+++ b/Engine/Plugins/VirtualProduction/RemoteControl/Source/RemoteControlUI/Private/Commands/RemoteControlCommands.cpp
@@ -39,6 +39,15 @@ void FRemoteControlCommands::RegisterCommands()
 
 	// Rename Entity
 	UI_COMMAND(RenameEntity, "Rename", "Rename the selected  group/exposed entity.", EUserInterfaceActionType::Button, FInputChord(EKeys::F2));
+
+	// Copy Item
+	UI_COMMAND(CopyItem, "Copy", "Copy the selected UI item", EUserInterfaceActionType::Button, FInputChord(EKeys::C, EModifierKey::Control));
+
+	// Paste Item
+	UI_COMMAND(PasteItem, "Paste", "Paste the selected UI item", EUserInterfaceActionType::Button, FInputChord(EKeys::V, EModifierKey::Control));
+
+	// Duplicate Item
+	UI_COMMAND(DuplicateItem, "Duplicate", "Duplicate the selected UI item", EUserInterfaceActionType::Button, FInputChord(EKeys::D, EModifierKey::Control));
 }
 
 #undef LOCTEXT_NAMESPACE
diff --git a/Engine/Plugins/VirtualProduction/RemoteControl/Source/RemoteControlUI/Private/Commands/RemoteControlCommands.h b/Engine/Plugins/VirtualProduction/RemoteControl/Source/RemoteControlUI/Private/Commands/RemoteControlCommands.h
--- a/Engine/Plugins/VirtualProduction/RemoteControl/Source/RemoteControlUI/Private/Commands/RemoteControlCommands.h
+++ b/Engine/Plugins/VirtualProduction/RemoteControl/Source/RemoteControlUI/Private/Commands/RemoteControlCommands.h
@@ -55,4 +55,19 @@ public:
 	 * Holds the information about UI Command that  renames selected group/exposed entity.
 	 */
 	TSharedPtr< FUICommandInfo > RenameEntity;
+
+	/**
+	 * UI Command for copying a UI item in the Remote Control preset. Currently used for Logic panel
+	 */
+	TSharedPtr< FUICommandInfo > CopyItem;
+
+	/**
+	 * UI Command for pasting a UI item in the Remote Control preset. Currently used for Logic panel
+	 */
+	TSharedPtr< FUICommandInfo > PasteItem;
+
+	/**
+	 * UI Command for duplicating a UI item in the Remote Control preset. Currently used for Logic panel
+	 */
+	TSharedPtr< FUICommandInfo > DuplicateItem;
 };
diff --git a/Engine/Plugins/VirtualProduction/RemoteControl/Source/RemoteControlUI/Private/UI/Action/RCActionModel.cpp b/Engine/Plugins/VirtualProduction/RemoteControl/Source/RemoteControlUI/Private/UI/Action/RCActionModel.cpp
--- a/Engine/Plugins/VirtualProduction/RemoteControl/Source/RemoteControlUI/Private/UI/Action/RCActionModel.cpp
+++ b/Engine/Plugins/VirtualProduction/RemoteControl/Source/RemoteControlUI/Private/UI/Action/RCActionModel.cpp
@@ -232,14 +232,8 @@ FLinearColor FRCPropertyActionType::GetPropertyTypeColor() const
 	return TypeColor;
 }
 
-TSharedPtr<SWidget> FRCActionModel::GetContextMenuWidget()
+void FRCActionModel::AddSpecialContextMenuOptions(FMenuBuilder& MenuBuilder)
 {
-	IMainFrameModule& MainFrame = FModuleManager::Get().LoadModuleChecked<IMainFrameModule>("MainFrame");
-
-	FMenuBuilder MenuBuilder(true, MainFrame.GetMainFrameCommandBindings());
-
-	MenuBuilder.BeginSection("Common");
-
 	if (TSharedPtr<SRemoteControlPanel> RemoteControlPanel = PanelWeakPtr.Pin())
 	{
 		// 1. Edit (if available)
@@ -247,20 +241,7 @@ TSharedPtr<SWidget> FRCActionModel::GetContextMenuWidget()
 		{
 			EditableVirtualPropertyWidget->AddEditContextMenuOption(MenuBuilder);
 		}
-
-		// 2. Delete
-		if (TSharedPtr<SRCActionPanel> ActionPanelList = RemoteControlPanel->GetLogicActionPanel())
-		{
-			FUIAction Action(FExecuteAction::CreateSP(ActionPanelList.Get(), &SRCActionPanel::DeleteSelectedPanelItem));
-			MenuBuilder.AddMenuEntry(LOCTEXT("ContextMenuEdit", "Delete"),
-				LOCTEXT("ContextMenuEditTooltip", "Delete the selected action"),
-				FSlateIcon(), Action);
-		}
 	}
-
-	MenuBuilder.EndSection();
-
-	return MenuBuilder.MakeWidget();
 }
 
 void FRCActionModel::OnSelectionExit()
diff --git a/Engine/Plugins/VirtualProduction/RemoteControl/Source/RemoteControlUI/Private/UI/Action/RCActionModel.h b/Engine/Plugins/VirtualProduction/RemoteControl/Source/RemoteControlUI/Private/UI/Action/RCActionModel.h
--- a/Engine/Plugins/VirtualProduction/RemoteControl/Source/RemoteControlUI/Private/UI/Action/RCActionModel.h
+++ b/Engine/Plugins/VirtualProduction/RemoteControl/Source/RemoteControlUI/Private/UI/Action/RCActionModel.h
@@ -8,6 +8,7 @@
 
 #include "UObject/WeakFieldPtr.h"
 
+class FMenuBuilder;
 class FRCBehaviourModel;
 struct FRCPanelStyle;
 class IDetailTreeNode;
@@ -49,6 +50,9 @@ public:
 	/** Color code for this Action. Customized per action type*/
 	virtual FLinearColor GetActionTypeColor() const = 0;
 
+	/** Allows Logic panels to add special functionality to the Context Menu based on context */
+	virtual void AddSpecialContextMenuOptions(FMenuBuilder& MenuBuilder);
+
 	/** OnGenerateRow delegate for the Actions List View*/
 	TSharedRef<ITableRow> OnGenerateWidgetForList(TSharedPtr<FRCActionModel> InItem, const TSharedRef<STableViewBase>& OwnerTable);
 
@@ -58,12 +62,15 @@ public:
 	/** Chooses the appropriate Action model for the current class and field type*/
 	static TSharedPtr<FRCActionModel> GetModelByActionType(URCAction* InAction, const TSharedPtr<class FRCBehaviourModel> InBehaviourItem, const TSharedPtr<SRemoteControlPanel> InRemoteControlPanel);
 
-	/** Context menu for Actions panel list */
-	virtual TSharedPtr<SWidget> GetContextMenuWidget();
-
 	/** Invoked when this Action item no longer has focus*/
 	virtual void OnSelectionExit();
 
+	/** Returns the Behaviour Item associated with us*/
+	TSharedPtr<FRCBehaviourModel> GetParentBehaviour()
+	{
+		return BehaviourItemWeakPtr.Pin();
+	}
+
 protected:
 	/** The Action (data model) associated with us*/
 	TWeakObjectPtr<URCAction> ActionWeakPtr;
diff --git a/Engine/Plugins/VirtualProduction/RemoteControl/Source/RemoteControlUI/Private/UI/Action/SRCActionPanel.cpp b/Engine/Plugins/VirtualProduction/RemoteControl/Source/RemoteControlUI/Private/UI/Action/SRCActionPanel.cpp
--- a/Engine/Plugins/VirtualProduction/RemoteControl/Source/RemoteControlUI/Private/UI/Action/SRCActionPanel.cpp
+++ b/Engine/Plugins/VirtualProduction/RemoteControl/Source/RemoteControlUI/Private/UI/Action/SRCActionPanel.cpp
@@ -33,10 +33,6 @@
 #include "Widgets/Input/SCheckBox.h"
 #include "Widgets/Layout/SBox.h"
 
-#if WITH_EDITOR
-#include "ScopedTransaction.h"
-#endif
-
 #define LOCTEXT_NAMESPACE "SRCActionPanel"
 
 TSharedPtr<SBox> SRCActionPanel::NoneSelectedWidget = SNew(SBox)
@@ -309,6 +305,8 @@ URCAction* SRCActionPanel::AddAction(const TSharedRef<const FRemoteControlField>
 
 			URCAction* NewAction = BehaviourItem->AddAction(InRemoteControlField);
 
+			AddNewActionToList(NewAction);
+
 			// Broadcast new Action to other panels
 			if (const TSharedPtr<SRemoteControlPanel> RemoteControlPanel = GetRemoteControlPanel())
 			{
@@ -429,6 +427,148 @@ void SRCActionPanel::DeleteSelectedPanelItem()
 	ActionPanelList->DeleteSelectedPanelItem();
 }
 
+TSharedPtr<FRCLogicModeBase> SRCActionPanel::GetSelectedLogicItem()
+{
+	if(ActionPanelList)
+	{
+		return ActionPanelList->GetSelectedLogicItem();
+	}
+
+	return nullptr;
+}
+
+void SRCActionPanel::DuplicateAction(URCAction* InAction)
+{
+	URCBehaviour* BehaviourTarget = nullptr;
+	URCBehaviour* BehaviourSource = nullptr;
+
+	// Behaviour Target - The Behaviour into which the input Action is to copied
+	if (const TSharedPtr<FRCBehaviourModel> BehaviourItem = SelectedBehaviourItemWeakPtr.Pin())
+	{
+		BehaviourTarget = BehaviourItem->GetBehaviour();
+		if (!ensure(BehaviourTarget))
+		{
+			return;
+		}
+	}
+	
+	// Behaviour Source - The Behaviour holding the input Action
+	if (InAction)
+	{
+		BehaviourSource = InAction->GetParentBehaviour();
+		if (!ensure(BehaviourSource))
+		{
+			return;
+		}
+	}
+
+	URCAction* NewAction = BehaviourSource->DuplicateAction(InAction, BehaviourTarget);
+
+	AddNewActionToList(NewAction);
+
+	// Broadcast new Action to other panels
+	if (const TSharedPtr<SRemoteControlPanel> RemoteControlPanel = GetRemoteControlPanel())
+	{
+		RemoteControlPanel->OnActionAdded.Broadcast(NewAction);
+	}
+}
+
+void SRCActionPanel::AddNewActionToList(URCAction* NewAction)
+{
+	if(ActionPanelList.IsValid())
+	{
+		ActionPanelList->AddNewLogicItem(NewAction);
+	}
+}
+
+void SRCActionPanel::DuplicateSelectedPanelItem()
+{
+	if (!ensure(ActionPanelList.IsValid()))
+	{
+		return;
+	}
+
+	if (const TSharedPtr<FRCActionModel> ActionItem = StaticCastSharedPtr<FRCActionModel>(ActionPanelList->GetSelectedLogicItem()))
+	{
+		DuplicateAction(ActionItem->GetAction());
+	}
+}
+
+void SRCActionPanel::CopySelectedPanelItem()
+{
+	if (TSharedPtr<SRemoteControlPanel> RemoteControlPanel = GetRemoteControlPanel())
+	{
+		if (TSharedPtr<FRCActionModel> ActionItem = StaticCastSharedPtr<FRCActionModel>(ActionPanelList->GetSelectedLogicItem()))
+		{
+			RemoteControlPanel->SetLogicClipboardItem(ActionItem->GetAction(), SharedThis(this));
+		}
+	}
+}
+
+void SRCActionPanel::PasteItemFromClipboard()
+{
+	if (TSharedPtr<SRemoteControlPanel> RemoteControlPanel = GetRemoteControlPanel())
+	{
+		if (RemoteControlPanel->LogicClipboardItemSource == SharedThis(this))
+		{
+			if(URCAction* Action = Cast<URCAction>(RemoteControlPanel->GetLogicClipboardItem()))
+			{
+				DuplicateAction(Action);
+			}
+		}
+	}
+}
+
+bool SRCActionPanel::CanPasteClipboardItem(UObject* InLogicClipboardItem)
+{
+	URCAction* LogicClipboardAction = Cast<URCAction>(InLogicClipboardItem);
+
+	URCBehaviour* BehaviourSource = nullptr;
+	URCBehaviour* BehaviourTarget = nullptr;
+
+	if (LogicClipboardAction)
+	{
+		BehaviourSource = LogicClipboardAction->GetParentBehaviour();
+		if (!BehaviourSource)
+		{
+			return false;
+		}
+	}
+
+	if (TSharedPtr<FRCBehaviourModel> BehaviourItemTarget = SelectedBehaviourItemWeakPtr.Pin())
+	{
+		BehaviourTarget = BehaviourItemTarget->GetBehaviour();
+		if (!BehaviourTarget)
+		{
+			return false;
+		}
+	}
+
+	// Copy-paste is allowed between compatible Behaviour types only
+	//
+	return BehaviourSource->GetClass() == BehaviourTarget->GetClass();
+}
+
+FText SRCActionPanel::GetPasteItemMenuEntrySuffix()
+{
+	if (TSharedPtr<SRemoteControlPanel> RemoteControlPanel = GetRemoteControlPanel())
+	{
+		// This function should only have been called if we were the source of the item copied.
+		if (ensure(RemoteControlPanel->LogicClipboardItemSource == SharedThis(this)))
+		{
+			if (URCAction* Action = Cast<URCAction>(RemoteControlPanel->GetLogicClipboardItem()))
+			{
+				if (URCBehaviour* Behaviour = Action->GetParentBehaviour())
+				{
+					return FText::Format(FText::FromString("Action {0}"), Behaviour->GetDisplayName());
+				}
+			}
+		}
+	}
+
+	return FText::GetEmpty();
+}
+
 FReply SRCActionPanel::RequestDeleteAllItems()
 {
 	if (!ActionPanelList.IsValid())
diff --git a/Engine/Plugins/VirtualProduction/RemoteControl/Source/RemoteControlUI/Private/UI/Action/SRCActionPanel.h b/Engine/Plugins/VirtualProduction/RemoteControl/Source/RemoteControlUI/Private/UI/Action/SRCActionPanel.h
--- a/Engine/Plugins/VirtualProduction/RemoteControl/Source/RemoteControlUI/Private/UI/Action/SRCActionPanel.h
+++ b/Engine/Plugins/VirtualProduction/RemoteControl/Source/RemoteControlUI/Private/UI/Action/SRCActionPanel.h
@@ -5,6 +5,7 @@
 #include "UI/BaseLogicUI/SRCLogicPanelBase.h"
 
 struct FRCPanelStyle;
+class FRCActionModel;
 struct FRemoteControlField;
 class FRCBehaviourModel;
 class SBox;
@@ -39,6 +40,24 @@ public:
 	/** Delete Item UI command implementation for this panel */
 	virtual void DeleteSelectedPanelItem() override;
 
+	/** Returns the UI item currently selected by the user (if any)*/
+	virtual TSharedPtr<FRCLogicModeBase> GetSelectedLogicItem() override;
+
+	/** "Duplicate Item" UI command implementation for Action panel*/
+	virtual void DuplicateSelectedPanelItem() override;
+
+	/** "Copy Item" UI command implementation for Action panel*/
+	virtual void CopySelectedPanelItem() override;
+
+	/** "Paste Item" UI command implementation for Action panel*/
+	virtual void PasteItemFromClipboard() override;
+
+	/** Whether a given clipboard item can be successfully pasted into this panel */
+	virtual bool CanPasteClipboardItem(UObject* InLogicClipboardItem) override;
+
+	/** Provides an item suffix for the Paste context menu to provide users with useful context on the nature of the item being pasted */
+	virtual FText GetPasteItemMenuEntrySuffix() override;
+
 	/** Adds an Action for the currently active Behaviour and broadcasts to parent panels */
 	URCAction* AddAction(const TSharedRef<const FRemoteControlField> InRemoteControlField);
 
@@ -56,7 +75,6 @@ protected:
 	virtual FReply RequestDeleteAllItems() override;
 
 private:
-
 	/** Determines the visibility Add All Button. */
 	EVisibility HandleAddAllButtonVisibility() const;
 
@@ -95,6 +113,10 @@ private:
 
 private:
 
+	void DuplicateAction(URCAction* InAction);
+
+	void AddNewActionToList(URCAction* NewAction);
+
 	/** The parent Behaviour that this Action panel is associated with */
 	TWeakPtr<FRCBehaviourModel> SelectedBehaviourItemWeakPtr;
 	
diff --git a/Engine/Plugins/VirtualProduction/RemoteControl/Source/RemoteControlUI/Private/UI/Action/SRCActionPanelList.h b/Engine/Plugins/VirtualProduction/RemoteControl/Source/RemoteControlUI/Private/UI/Action/SRCActionPanelList.h
--- a/Engine/Plugins/VirtualProduction/RemoteControl/Source/RemoteControlUI/Private/UI/Action/SRCActionPanelList.h
+++ b/Engine/Plugins/VirtualProduction/RemoteControl/Source/RemoteControlUI/Private/UI/Action/SRCActionPanelList.h
@@ -31,6 +31,8 @@
 
 #define LOCTEXT_NAMESPACE "SRCActionPanelList"
 
+class FRCActionModel;
+
 /*
 * ~ SRCActionPanelList ~
 *
@@ -50,7 +52,10 @@ public:
 	/** Constructs this widget with InArgs */
 	void Construct(const FArguments& InArgs, const TSharedRef<SRCActionPanel> InActionPanel, TSharedPtr<FRCBehaviourModel> InBehaviourItem)
 	{
-		SRCLogicPanelListBase::Construct(SRCLogicPanelListBase::FArguments());
+		const TSharedPtr<SRemoteControlPanel>& RemoteControlPanel = InActionPanel->GetRemoteControlPanel();
+		check(RemoteControlPanel);
+
+		SRCLogicPanelListBase::Construct(SRCLogicPanelListBase::FArguments(), RemoteControlPanel.ToSharedRef());
 
 		ActionPanelWeakPtr = InActionPanel;
 		BehaviourItemWeakPtr = InBehaviourItem;
@@ -62,7 +67,7 @@ public:
 			.OnGenerateRow(this, &SRCActionPanelList::OnGenerateWidgetForList)
 			.OnSelectionChanged(this, &SRCActionPanelList::OnSelectionChanged)
 			.ListViewStyle(&RCPanelStyle->TableViewStyle)
-			.OnContextMenuOpening(this, &SRCActionPanelList::GetContextMenuWidget)
+			.OnContextMenuOpening(this, &SRCLogicPanelListBase::GetContextMenuWidget)
 			.HeaderRow(ActionType::GetHeaderRow());
 
 		ChildSlot
@@ -79,12 +84,8 @@ public:
 			];
 
 		// Add delegates
-		const TSharedPtr<SRemoteControlPanel> RemoteControlPanel = InActionPanel->GetRemoteControlPanel();
-		if (ensure(RemoteControlPanel))
-		{
-			RemoteControlPanel->OnActionAdded.AddSP(this, &SRCActionPanelList::OnActionAdded);
-			RemoteControlPanel->OnEmptyActions.AddSP(this, &SRCActionPanelList::OnEmptyActions);
-		}
+		RemoteControlPanel->OnActionAdded.AddSP(this, &SRCActionPanelList::OnActionAdded);
+		RemoteControlPanel->OnEmptyActions.AddSP(this, &SRCActionPanelList::OnEmptyActions);
 
 		if (URCBehaviour* Behaviour = InBehaviourItem->GetBehaviour())
 		{
@@ -117,7 +118,7 @@ public:
 	/** Whether the Actions List View currently has focus.*/
 	virtual bool IsListFocused() const override
 	{
-		return ListView->HasAnyUserFocus().IsSet();
+		return ListView->HasAnyUserFocus().IsSet() || ContextMenuWidgetCached.IsValid();
 	}
 
 	/** Deletes currently selected items from the list view*/
@@ -126,6 +127,18 @@ public:
 		DeleteItemFromLogicPanel<ActionType>(ActionItems, ListView->GetSelectedItems());
 	}
 
+	/** Returns the UI item currently selected by the user (if any)*/
+	virtual TSharedPtr<FRCLogicModeBase> GetSelectedLogicItem() override
+	{
+		return GetSelectedActionItem();
+	}
+
+	/** The currently selected Action item*/
+	TSharedPtr<FRCActionModel> GetSelectedActionItem()
+	{
+		return SelectedActionItem;
+	}
+
 	/** Fetches the parent Action panel*/
 	TSharedPtr<SRCActionPanel> GetActionPanel()
 	{
@@ -155,6 +168,51 @@ public:
 		return nullptr;
 	}
 
+	virtual void AddNewLogicItem(UObject* InLogicItem) override
+	{
+		AddRowToList(Cast<URCAction>(InLogicItem));
+
+		RequestRefresh();
+	}
+
+	void AddRowToList(URCAction* InAction)
+	{
+		if (!ensure(InAction))
+		{
+			return;
+		}
+
+		if (TSharedPtr<SRCActionPanel> ActionPanel = GetActionPanel())
+		{
+			if (TSharedPtr<SRemoteControlPanel> RemoteControlPanel = ActionPanel->GetRemoteControlPanel())
+			{
+				if (TSharedPtr<FRCBehaviourModel> BehaviourItem = BehaviourItemWeakPtr.Pin())
+				{
+					TSharedPtr<ActionType> ActionItem = ActionType::GetModelByActionType(InAction, BehaviourItem, RemoteControlPanel);
+
+					if (ensure(ActionItem.IsValid()))
+					{
+						ActionItems.Add(ActionItem);
+					}
+				}
+			}
+		}
+	}
+
+	void RequestRefresh()
+	{
+		ListView->RequestListRefresh();
+	}
+
+	/** Allows Logic panels to add special functionality to the Context Menu based on context */
+	virtual void AddSpecialContextMenuOptions(FMenuBuilder& MenuBuilder) override
+	{
+		if (SelectedActionItem)
+		{
+			SelectedActionItem->AddSpecialContextMenuOptions(MenuBuilder);
+		}
+	}
+
 private:
 
 	/** OnGenerateRow delegate for the Actions List View*/
@@ -174,7 +232,9 @@ private:
 	/** Responds to the selection of a newly created action. Resets UI state*/
 	void OnActionAdded(URCAction* InAction)
 	{
-		Reset();
+		// Historical note: Previously we used to call Reset here after adding an Action
+		// This is now covered by use of AddRowToList function (invoked from SRCActionPanel)
+		// @todo: This improvement needs to be done for "Remove Action" as well.
 	}
 
 	/** Responds to the removal of all actions. Rests UI state*/
@@ -186,28 +246,20 @@ private:
 	/** Refreshes the list from the latest state of the data model*/
 	virtual void Reset() override
 	{
-		if (TSharedPtr<SRCActionPanel> ActionPanel = GetActionPanel())
+		ActionItems.Empty();
+
+		if (TSharedPtr<FRCBehaviourModel> BehaviourItem = BehaviourItemWeakPtr.Pin())
 		{
-			if (TSharedPtr<SRemoteControlPanel> RemoteControlPanel = ActionPanel->GetRemoteControlPanel())
+			if (URCBehaviour* Behaviour = Cast<URCBehaviour>(BehaviourItem->GetBehaviour()))
 			{
-				ActionItems.Empty();
-
-				if (TSharedPtr<FRCBehaviourModel> BehaviourItem = BehaviourItemWeakPtr.Pin())
+				for (URCAction* Action : Behaviour->ActionContainer->GetActions())
 				{
-					if (URCBehaviour* Behaviour = Cast<URCBehaviour>(BehaviourItem->GetBehaviour()))
-					{
-						for (URCAction* Action : Behaviour->ActionContainer->GetActions())
-						{
-							TSharedPtr<ActionType> ActionItem = ActionType::GetModelByActionType(Action, BehaviourItem, RemoteControlPanel);
-
-							ActionItems.Add(ActionItem);
-						}
-					}
+					AddRowToList(Action);
 				}
-
-				ListView->RebuildList();
 			}
 		}
+
+		ListView->RebuildList();
 	}
 
 	/** Handles broadcasting of a successful remove item operation.*/
@@ -345,17 +397,6 @@ private:
 		return false;
 	}
 
-	/** Context menu for Actions panel list */
-	TSharedPtr<SWidget> GetContextMenuWidget()
-	{
-		if (SelectedActionItem)
-		{
-			return SelectedActionItem->GetContextMenuWidget();
-		}
-
-		return SNullWidget::NullWidget;
-	}
-
 	/** OnSelectionChanged delegate for Actions List View */
 	void OnSelectionChanged(TSharedPtr<ActionType> InItem, ESelectInfo::Type)
 	{
@@ -370,7 +411,7 @@ private:
 private:
 
 	/** The currently selected Action item*/
-	TSharedPtr<ActionType> SelectedActionItem;
+	TSharedPtr<FRCActionModel> SelectedActionItem;
 
 	/** The parent Action Panel widget*/
 	TWeakPtr<SRCActionPanel> ActionPanelWeakPtr;
diff --git a/Engine/Plugins/VirtualProduction/RemoteControl/Source/RemoteControlUI/Private/UI/BaseLogicUI/SRCLogicPanelBase.h b/Engine/Plugins/VirtualProduction/RemoteControl/Source/RemoteControlUI/Private/UI/BaseLogicUI/SRCLogicPanelBase.h
--- a/Engine/Plugins/VirtualProduction/RemoteControl/Source/RemoteControlUI/Private/UI/BaseLogicUI/SRCLogicPanelBase.h
+++ b/Engine/Plugins/VirtualProduction/RemoteControl/Source/RemoteControlUI/Private/UI/BaseLogicUI/SRCLogicPanelBase.h
@@ -5,6 +5,7 @@
 #include "CoreMinimal.h"
 #include "Widgets/SCompoundWidget.h"
 
+class FRCLogicModeBase;
 class SRemoteControlPanel;
 class URemoteControlPreset;
 
@@ -37,6 +38,24 @@ public:
 	/** "Delete Item" UI command implementation for panels*/
 	virtual void DeleteSelectedPanelItem() = 0;
 
+	/** "Duplicate Item" UI command implementation for panels*/
+	virtual void DuplicateSelectedPanelItem() {}
+
+	/** "Copy Item" UI command implementation for panels*/
+	virtual void CopySelectedPanelItem() {}
+
+	/** "Paste Item" UI command implementation for panels*/
+	virtual void PasteItemFromClipboard() {}
+
+	/** Whether a given clipboard item can be successfully pasted into this panel */
+	virtual bool CanPasteClipboardItem(UObject* InLogicClipboardItem) { return true; }
+
+	/** Returns the UI item currently selected by the user (if any). To be implemented per child panel*/
+	virtual TSharedPtr<FRCLogicModeBase> GetSelectedLogicItem() = 0;
+
+	/** Provides an item suffix for the Paste context menu to provide users with useful context on the nature of the item being pasted */
+	virtual FText GetPasteItemMenuEntrySuffix() { return FText::GetEmpty(); }
+
 protected:
 
 	/** Warns user before deleting all items in a panel. */
diff --git a/Engine/Plugins/VirtualProduction/RemoteControl/Source/RemoteControlUI/Private/UI/BaseLogicUI/SRCLogicPanelListBase.cpp b/Engine/Plugins/VirtualProduction/RemoteControl/Source/RemoteControlUI/Private/UI/BaseLogicUI/SRCLogicPanelListBase.cpp
--- a/Engine/Plugins/VirtualProduction/RemoteControl/Source/RemoteControlUI/Private/UI/BaseLogicUI/SRCLogicPanelListBase.cpp
+++ b/Engine/Plugins/VirtualProduction/RemoteControl/Source/RemoteControlUI/Private/UI/BaseLogicUI/SRCLogicPanelListBase.cpp
@@ -2,10 +2,64 @@
 
 #include "SRCLogicPanelListBase.h"
 
-BEGIN_SLATE_FUNCTION_BUILD_OPTIMIZATION
+#include "Commands/RemoteControlCommands.h"
+#include "Interfaces/IMainFrameModule.h"
+#include "SRCLogicPanelBase.h"
+#include "UI/SRemoteControlPanel.h"
 
-void SRCLogicPanelListBase::Construct(const FArguments& InArgs)
+#define LOCTEXT_NAMESPACE "SRCLogicPanelListBase"
+
+void SRCLogicPanelListBase::Construct(const FArguments& InArgs, const TSharedRef<SRemoteControlPanel>& InPanel)
+{
+	RemoteControlPanelWeakPtr = InPanel;
+}
+
+TSharedPtr<SWidget> SRCLogicPanelListBase::GetContextMenuWidget()
 {
+	IMainFrameModule& MainFrame = FModuleManager::Get().LoadModuleChecked<IMainFrameModule>("MainFrame");
+
+	FMenuBuilder MenuBuilder(true, MainFrame.GetMainFrameCommandBindings());
+
+	// Special menu options
+	MenuBuilder.BeginSection("Advanced");
+	AddSpecialContextMenuOptions(MenuBuilder);
+	MenuBuilder.EndSection();
+
+	// Generic options (based on UI Commands)
+	MenuBuilder.BeginSection("Common");
+
+	const FRemoteControlCommands& Commands = FRemoteControlCommands::Get();
+
+	// 1. Copy
+	MenuBuilder.AddMenuEntry(Commands.CopyItem);
+
+	// 2. Paste
+	FText PasteItemLabel = LOCTEXT("Paste", "Paste");
+	if (TSharedPtr<SRemoteControlPanel> RemoteControlPanel = RemoteControlPanelWeakPtr.Pin())
+	{
+		if (TSharedPtr<SRCLogicPanelBase> LogicPanel = RemoteControlPanel->LogicClipboardItemSource)
+		{
+			const FText Suffix = LogicPanel->GetPasteItemMenuEntrySuffix();
+			if (!Suffix.IsEmpty())
+			{
+				PasteItemLabel = FText::Format(FText::FromString("{0} ({1})"), PasteItemLabel, Suffix);
+			}
+		}
+	}
+	MenuBuilder.AddMenuEntry(Commands.PasteItem, NAME_None, PasteItemLabel);
+
+	// 2. Duplicate
+	MenuBuilder.AddMenuEntry(Commands.DuplicateItem);
+
+	// 3. Delete
+	MenuBuilder.AddMenuEntry(Commands.DeleteEntity);
+
+	MenuBuilder.EndSection();
+
+	TSharedPtr<SWidget> MenuWidget = MenuBuilder.MakeWidget();
+	ContextMenuWidgetCached = MenuWidget;
+
+	return MenuWidget;
 }
 
-END_SLATE_FUNCTION_BUILD_OPTIMIZATION
+#undef LOCTEXT_NAMESPACE
\ No newline at end of file
diff --git a/Engine/Plugins/VirtualProduction/RemoteControl/Source/RemoteControlUI/Private/UI/BaseLogicUI/SRCLogicPanelListBase.h b/Engine/Plugins/VirtualProduction/RemoteControl/Source/RemoteControlUI/Private/UI/BaseLogicUI/SRCLogicPanelListBase.h
--- a/Engine/Plugins/VirtualProduction/RemoteControl/Source/RemoteControlUI/Private/UI/BaseLogicUI/SRCLogicPanelListBase.h
+++ b/Engine/Plugins/VirtualProduction/RemoteControl/Source/RemoteControlUI/Private/UI/BaseLogicUI/SRCLogicPanelListBase.h
@@ -12,6 +12,8 @@
 #include "UI/Controller/RCControllerModel.h"
 #include "Widgets/SCompoundWidget.h"
 
+class SRemoteControlPanel;
+
 /*
 * ~ SRCLogicPanelListBase ~
 *
@@ -28,7 +30,7 @@ public:
 	SLATE_END_ARGS()
 
 	/** Constructs this widget with InArgs */
-	void Construct(const FArguments& InArgs);
+	void Construct(const FArguments& InArgs, const TSharedRef<SRemoteControlPanel>& InPanel);
 
 	/** Returns true if the underlying list is valid and empty. */
 	virtual bool IsEmpty() const = 0;
@@ -42,6 +44,18 @@ public:
 	/** Deletes currently selected items from the list view*/
 	virtual void DeleteSelectedPanelItem()  = 0;
 
+	/** Returns the UI item currently selected by the user (if any). To be implemented per child panel*/
+	virtual TSharedPtr<FRCLogicModeBase> GetSelectedLogicItem() = 0;
+
+	/** Provides a common entry point for adding Logic related data objects (UObjects) to their repsective panels*/
+	virtual void AddNewLogicItem(UObject* InLogicItem) {}
+
+	/** Builds the right-click context menu populated with generic actions (based on UI Commands) */
+	TSharedPtr<SWidget> GetContextMenuWidget();
+
+	/** Allows Logic panels to add special functionality to the Context Menu based on context */
+	virtual void AddSpecialContextMenuOptions(FMenuBuilder& MenuBuilder) {};
+
 private:
 
 	/** Refreshes the list from the latest state of the model*/
@@ -58,6 +72,7 @@ private:
 	virtual int32 RemoveModel(const TSharedPtr<FRCLogicModeBase> InModel) = 0;
 
 protected:
+	TWeakPtr<SWidget> ContextMenuWidgetCached;
 
 	/** Helper function for handling common Delete Item functionality across all child panels (Actions/Behaviours/Controllers)
 	* Currently invoked from each Panel List child class with appropriate template class*/
@@ -95,4 +110,7 @@ protected:
 			Reset();
 		}
 	}
+
+	/** The parent Remote Control Panel widget*/
+	TWeakPtr<SRemoteControlPanel> RemoteControlPanelWeakPtr;
 };
\ No newline at end of file
diff --git a/Engine/Plugins/VirtualProduction/RemoteControl/Source/RemoteControlUI/Private/UI/Behaviour/SRCBehaviourPanel.cpp b/Engine/Plugins/VirtualProduction/RemoteControl/Source/RemoteControlUI/Private/UI/Behaviour/SRCBehaviourPanel.cpp
--- a/Engine/Plugins/VirtualProduction/RemoteControl/Source/RemoteControlUI/Private/UI/Behaviour/SRCBehaviourPanel.cpp
+++ b/Engine/Plugins/VirtualProduction/RemoteControl/Source/RemoteControlUI/Private/UI/Behaviour/SRCBehaviourPanel.cpp
@@ -74,13 +74,15 @@ void SRCBehaviourPanel::OnControllerSelectionChanged(TSharedPtr<FRCControllerMod
 
 void SRCBehaviourPanel::UpdateWrappedWidget(TSharedPtr<FRCControllerModel> InControllerItem)
 {
-	if (InControllerItem.IsValid())
+	TSharedPtr<SRemoteControlPanel> RemoteControlPanel = GetRemoteControlPanel();
+
+	if (InControllerItem.IsValid() && RemoteControlPanel.IsValid())
 	{
 		// Behaviour Dock Panel
 		TSharedPtr<SRCMinorPanel> BehaviourDockPanel = SNew(SRCMinorPanel)
 			.HeaderLabel(LOCTEXT("BehavioursLabel", "Behavior"))
 			[
-				SAssignNew(BehaviourPanelList, SRCBehaviourPanelList, SharedThis(this), InControllerItem)
+				SAssignNew(BehaviourPanelList, SRCBehaviourPanelList, SharedThis(this), InControllerItem, RemoteControlPanel.ToSharedRef())
 			];
 
 		// Add New Behaviour Button
@@ -239,6 +241,111 @@ void SRCBehaviourPanel::DeleteSelectedPanelItem()
 	BehaviourPanelList->DeleteSelectedPanelItem();
 }
 
+void SRCBehaviourPanel::DuplicateBehaviour(URCBehaviour* InBehaviour)
+{
+	if (URCController* Controller = GetParentController())
+	{
+		URCBehaviour* NewBehaviour = URCController::DuplicateBehaviour(Controller, InBehaviour);
+
+		BehaviourPanelList->AddNewLogicItem(NewBehaviour);
+	}
+}
+
+void SRCBehaviourPanel::DuplicateSelectedPanelItem()
+{
+	FScopedTransaction Transaction(LOCTEXT("DuplicateBehaviour", "Duplicate Behaviour"));
+
+	if (TSharedPtr<FRCBehaviourModel> SelectedBehaviourItem = BehaviourPanelList->GetSelectedBehaviourItem())
+	{
+		DuplicateBehaviour(SelectedBehaviourItem->GetBehaviour());
+	}
+}
+
+void SRCBehaviourPanel::CopySelectedPanelItem()
+{
+	if (TSharedPtr<SRemoteControlPanel> RemoteControlPanel = GetRemoteControlPanel())
+	{
+		if (TSharedPtr<FRCBehaviourModel> SelectedBehaviourItem = BehaviourPanelList->GetSelectedBehaviourItem())
+		{
+			RemoteControlPanel->SetLogicClipboardItem(SelectedBehaviourItem->GetBehaviour(), SharedThis(this));
+		}
+	}
+}
+
+void SRCBehaviourPanel::PasteItemFromClipboard()
+{
+	FScopedTransaction Transaction(LOCTEXT("PasteBehaviour", "Paste Behaviour"));
+
+	if (TSharedPtr<SRemoteControlPanel> RemoteControlPanel = GetRemoteControlPanel())
+	{
+		if (RemoteControlPanel->LogicClipboardItemSource == SharedThis(this))
+		{
+			if(URCBehaviour* Behaviour = Cast<URCBehaviour>(RemoteControlPanel->GetLogicClipboardItem()))
+			{
+				DuplicateBehaviour(Behaviour);
+			}
+		}
+	}
+}
+
+bool SRCBehaviourPanel::CanPasteClipboardItem(UObject* InLogicClipboardItem)
+{
+	URCBehaviour* LogicClipboardBehaviour = Cast<URCBehaviour>(InLogicClipboardItem);
+
+	URCController* ControllerSource = LogicClipboardBehaviour->ControllerWeakPtr.Get();
+	if (!ControllerSource)
+	{
+		return false;
+	}
+
+	URCController* ControllerTarget = GetParentController();
+	if (!ControllerTarget)
+	{
+		return false;
+	}
+
+	// Copy-Paste is only permitted between Controllers of the same type
+	//
+	return ControllerSource->GetValueType() == ControllerTarget->GetValueType();
+}
+
+FText SRCBehaviourPanel::GetPasteItemMenuEntrySuffix()
+{
+	if (TSharedPtr<SRemoteControlPanel> RemoteControlPanel = GetRemoteControlPanel())
+	{
+		// This function should only have been called if we were the source of the item copied.
+		if (ensure(RemoteControlPanel->LogicClipboardItemSource == SharedThis(this)))
+		{
+			if (URCBehaviour* Behaviour = Cast<URCBehaviour>(RemoteControlPanel->GetLogicClipboardItem()))
+			{
+				return FText::Format(FText::FromString("Behaviour {0}"), Behaviour->GetDisplayName());
+			}
+		}
+	}
+
+	return FText::GetEmpty();
+}
+
+TSharedPtr<FRCLogicModeBase> SRCBehaviourPanel::GetSelectedLogicItem()
+{
+	if (BehaviourPanelList)
+	{
+		return BehaviourPanelList->GetSelectedBehaviourItem();
+	}
+
+	return nullptr;
+}
+
+URCController* SRCBehaviourPanel::GetParentController()
+{
+	if (TSharedPtr<FRCControllerModel> ControllerItem = SelectedControllerItemWeakPtr.Pin())
+	{
+		return Cast<URCController>(ControllerItem->GetVirtualProperty());
+	}
+
+	return nullptr;
+}
+
 FReply SRCBehaviourPanel::RequestDeleteAllItems()
 {
 	if (!BehaviourPanelList.IsValid())
diff --git a/Engine/Plugins/VirtualProduction/RemoteControl/Source/RemoteControlUI/Private/UI/Behaviour/SRCBehaviourPanel.h b/Engine/Plugins/VirtualProduction/RemoteControl/Source/RemoteControlUI/Private/UI/Behaviour/SRCBehaviourPanel.h
--- a/Engine/Plugins/VirtualProduction/RemoteControl/Source/RemoteControlUI/Private/UI/Behaviour/SRCBehaviourPanel.h
+++ b/Engine/Plugins/VirtualProduction/RemoteControl/Source/RemoteControlUI/Private/UI/Behaviour/SRCBehaviourPanel.h
@@ -6,9 +6,12 @@
 #include "UI/BaseLogicUI/SRCLogicPanelBase.h"
 
 struct FRCPanelStyle;
+class FRCBehaviourModel;
 class FRCControllerModel;
 class SBox;
 class SRCBehaviourPanel;
+class URCController;
+
 /*
 * ~ SRCBehaviourPanel ~
 *
@@ -36,12 +39,35 @@ public:
 	/** Delete Item UI command implementation for this panel */
 	virtual void DeleteSelectedPanelItem() override;
 
+	/** "Duplicate Item" UI command implementation for Behaviour panel*/
+	virtual void DuplicateSelectedPanelItem() override;
+
+	/** "Copy Item" UI command implementation for Behaviour panel*/
+	virtual void CopySelectedPanelItem() override;
+
+	/** "Paste Item" UI command implementation for Behaviour panel*/
+	virtual void PasteItemFromClipboard() override;
+
+	/** Whether a given clipboard item can be successfully pasted into this panel */
+	virtual bool CanPasteClipboardItem(UObject* InLogicClipboardItem) override;
+
+	/** Provides an item suffix for the Paste context menu to provide users with useful context on the nature of the item being pasted */
+	virtual FText GetPasteItemMenuEntrySuffix() override;
+
+	/** Returns the UI item currently selected by the user (if any)*/
+	virtual TSharedPtr<FRCLogicModeBase> GetSelectedLogicItem() override;
+
+	/** Returns the parent Controller associated with this behaviour*/
+	URCController* GetParentController();
+
 protected:
 
 	/** Warns user before deleting all items in a panel. */
 	virtual FReply RequestDeleteAllItems() override;
 
 private:
+	/** Duplicates a given Behaviour object*/
+	void DuplicateBehaviour(URCBehaviour* InBehaviour);
 	
 	/**
 	* Controller list selection change listener.
diff --git a/Engine/Plugins/VirtualProduction/RemoteControl/Source/RemoteControlUI/Private/UI/Behaviour/SRCBehaviourPanelList.cpp b/Engine/Plugins/VirtualProduction/RemoteControl/Source/RemoteControlUI/Private/UI/Behaviour/SRCBehaviourPanelList.cpp
--- a/Engine/Plugins/VirtualProduction/RemoteControl/Source/RemoteControlUI/Private/UI/Behaviour/SRCBehaviourPanelList.cpp
+++ b/Engine/Plugins/VirtualProduction/RemoteControl/Source/RemoteControlUI/Private/UI/Behaviour/SRCBehaviourPanelList.cpp
@@ -32,11 +32,9 @@ namespace FRemoteControlBehaviourColumns
 	const FName Behaviours = TEXT("Behaviors");
 }
 
-BEGIN_SLATE_FUNCTION_BUILD_OPTIMIZATION
-
-void SRCBehaviourPanelList::Construct(const FArguments& InArgs, TSharedRef<SRCBehaviourPanel> InBehaviourPanel, TSharedPtr<FRCControllerModel> InControllerItem)
+void SRCBehaviourPanelList::Construct(const FArguments& InArgs, TSharedRef<SRCBehaviourPanel> InBehaviourPanel, TSharedPtr<FRCControllerModel> InControllerItem, const TSharedRef<SRemoteControlPanel> InRemoteControlPanel)
 {
-	SRCLogicPanelListBase::Construct(SRCLogicPanelListBase::FArguments());
+	SRCLogicPanelListBase::Construct(SRCLogicPanelListBase::FArguments(), InRemoteControlPanel);
 	
 	BehaviourPanelWeakPtr = InBehaviourPanel;
 	ControllerItemWeakPtr = InControllerItem;
@@ -48,6 +46,7 @@ void SRCBehaviourPanelList::Construct(const FArguments& InArgs, TSharedRef<SRCBe
 		.OnSelectionChanged(this, &SRCBehaviourPanelList::OnTreeSelectionChanged)
 		.OnGenerateRow(this, &SRCBehaviourPanelList::OnGenerateWidgetForList)
 		.ListViewStyle(&RCPanelStyle->TableViewStyle)
+		.OnContextMenuOpening(this, &SRCLogicPanelListBase::GetContextMenuWidget)
 		.HeaderRow(
 			SNew(SHeaderRow)
 			.Style(&RCPanelStyle->HeaderRowStyle)
@@ -69,7 +68,7 @@ void SRCBehaviourPanelList::Construct(const FArguments& InArgs, TSharedRef<SRCBe
 	RemoteControlPanel->OnBehaviourAdded.AddSP(this, &SRCBehaviourPanelList::OnBehaviourAdded);
 	RemoteControlPanel->OnEmptyBehaviours.AddSP(this, &SRCBehaviourPanelList::OnEmptyBehaviours);
 
-	if (URCController* Controller = Cast< URCController>(InControllerItem->GetVirtualProperty()))
+	if (URCController* Controller = Cast<URCController>(InControllerItem->GetVirtualProperty()))
 	{
 		Controller->OnBehaviourListModified.AddSP(this, &SRCBehaviourPanelList::OnBehaviourListModified);
 	}
@@ -87,44 +86,61 @@ int32 SRCBehaviourPanelList::Num() const
 	return BehaviourItems.Num();
 }
 
-END_SLATE_FUNCTION_BUILD_OPTIMIZATION
+void SRCBehaviourPanelList::AddNewLogicItem(UObject* InLogicItem)
+{
+	AddBehaviourToList(Cast<URCBehaviour>(InLogicItem));
+
+	RequestRefresh();
+}
+
+void SRCBehaviourPanelList::AddBehaviourToList(URCBehaviour* InBehaviour)
+{
+	if (!BehaviourPanelWeakPtr.IsValid())
+	{
+		return;
+	}
+
+	const TSharedPtr<SRemoteControlPanel> RemoteControlPanel = BehaviourPanelWeakPtr.Pin()->GetRemoteControlPanel();
+
+	if (URCIsEqualBehaviour* IsEqualBehaviour = Cast<URCIsEqualBehaviour>(InBehaviour))
+	{
+		BehaviourItems.Add(MakeShared<FRCIsEqualBehaviourModel>(IsEqualBehaviour));
+	}
+	else if (URCBehaviourConditional* ConditionalBehaviour = Cast<URCBehaviourConditional>(InBehaviour))
+	{
+		BehaviourItems.Add(MakeShared<FRCBehaviourConditionalModel>(ConditionalBehaviour));
+	}
+	else if (URCSetAssetByPathBehaviour* SetAssetByPathBehaviour = Cast<URCSetAssetByPathBehaviour>(InBehaviour))
+	{
+		BehaviourItems.Add(MakeShared<FRCSetAssetByPathBehaviourModel>(SetAssetByPathBehaviour));
+	}
+	else if (URCBehaviourBind* BindBehaviour = Cast<URCBehaviourBind>(InBehaviour))
+	{
+		BehaviourItems.Add(MakeShared<FRCBehaviourBindModel>(BindBehaviour, RemoteControlPanel));
+	}
+	else if (URCRangeMapBehaviour* RangeMapBehaviour = Cast<URCRangeMapBehaviour>(InBehaviour))
+	{
+		BehaviourItems.Add(MakeShared<FRCRangeMapBehaviourModel>(RangeMapBehaviour));
+	}
+	else
+	{
+		BehaviourItems.Add(MakeShared<FRCBehaviourModel>(InBehaviour, RemoteControlPanel));
+	}
+}
+
+
 
 void SRCBehaviourPanelList::Reset()
 {
 	BehaviourItems.Empty();
 
-	const TSharedPtr<SRemoteControlPanel> RemoteControlPanel = BehaviourPanelWeakPtr.IsValid() ? BehaviourPanelWeakPtr.Pin()->GetRemoteControlPanel() : nullptr;
-
 	if (TSharedPtr<FRCControllerModel> ControllerItem = ControllerItemWeakPtr.Pin())
 	{
 		if (URCController* Controller = Cast<URCController>(ControllerItem->GetVirtualProperty()))
 		{
 			for (URCBehaviour* Behaviour : Controller->Behaviours)
 			{
-				if (URCIsEqualBehaviour* IsEqualBehaviour = Cast<URCIsEqualBehaviour>(Behaviour))
-				{
-					BehaviourItems.Add(MakeShared<FRCIsEqualBehaviourModel>(IsEqualBehaviour));
-				}
-				else if (URCBehaviourConditional* ConditionalBehaviour = Cast<URCBehaviourConditional>(Behaviour))
-				{
-					BehaviourItems.Add(MakeShared<FRCBehaviourConditionalModel>(ConditionalBehaviour));
-				}
-				else if (URCSetAssetByPathBehaviour* SetAssetByPathBehaviour = Cast<URCSetAssetByPathBehaviour>(Behaviour))
-				{
-					BehaviourItems.Add(MakeShared<FRCSetAssetByPathBehaviourModel>(SetAssetByPathBehaviour));
-				}
-				else if (URCBehaviourBind* BindBehaviour = Cast<URCBehaviourBind>(Behaviour))
-				{
-					BehaviourItems.Add(MakeShared<FRCBehaviourBindModel>(BindBehaviour, RemoteControlPanel));
-				}
-				else if (URCRangeMapBehaviour* RangeMapBehaviour = Cast<URCRangeMapBehaviour>(Behaviour))
-				{
-					BehaviourItems.Add(MakeShared<FRCRangeMapBehaviourModel>(RangeMapBehaviour));
-				}
-				else
-				{
-					BehaviourItems.Add(MakeShared<FRCBehaviourModel>(Behaviour, RemoteControlPanel));
-				}
+				AddBehaviourToList(Behaviour);
 			}
 		}
 	}
@@ -222,7 +238,7 @@ void SRCBehaviourPanelList::OnBehaviourListModified()
 
 bool SRCBehaviourPanelList::IsListFocused() const
 {
-	return ListView->HasAnyUserFocus().IsSet();
+	return ListView->HasAnyUserFocus().IsSet() || ContextMenuWidgetCached.IsValid();;
 }
 
 void SRCBehaviourPanelList::DeleteSelectedPanelItem()
@@ -230,4 +246,9 @@ void SRCBehaviourPanelList::DeleteSelectedPanelItem()
 	DeleteItemFromLogicPanel<FRCBehaviourModel>(BehaviourItems, ListView->GetSelectedItems());
 }
 
-#undef LOCTEXT_NAMESPACE
+void SRCBehaviourPanelList::RequestRefresh()
+{
+	ListView->RequestListRefresh();
+}
+
+#undef LOCTEXT_NAMESPACE
\ No newline at end of file
diff --git a/Engine/Plugins/VirtualProduction/RemoteControl/Source/RemoteControlUI/Private/UI/Behaviour/SRCBehaviourPanelList.h b/Engine/Plugins/VirtualProduction/RemoteControl/Source/RemoteControlUI/Private/UI/Behaviour/SRCBehaviourPanelList.h
--- a/Engine/Plugins/VirtualProduction/RemoteControl/Source/RemoteControlUI/Private/UI/Behaviour/SRCBehaviourPanelList.h
+++ b/Engine/Plugins/VirtualProduction/RemoteControl/Source/RemoteControlUI/Private/UI/Behaviour/SRCBehaviourPanelList.h
@@ -11,6 +11,7 @@ class FRCBehaviourModel;
 class ITableRow;
 class ITableBase;
 class SRCBehaviourPanel;
+class SRemoteControlPanel;
 class STableViewBase;
 class URemoteControlPreset;
 template <typename ItemType> class SListView;
@@ -31,7 +32,7 @@ public:
 	SLATE_END_ARGS()
 
 	/** Constructs this widget with InArgs */
-	void Construct(const FArguments& InArgs, const TSharedRef<SRCBehaviourPanel> InBehaviourPanel, TSharedPtr<FRCControllerModel> InControllerItem);
+	void Construct(const FArguments& InArgs, const TSharedRef<SRCBehaviourPanel> InBehaviourPanel, TSharedPtr<FRCControllerModel> InControllerItem, const TSharedRef<SRemoteControlPanel> InRemoteControlPanel);
 	
 	/** Returns true if the underlying list is valid and empty. */
 	virtual bool IsEmpty() const override;
@@ -45,8 +46,25 @@ public:
 	/** Deletes currently selected items from the list view*/
 	virtual void DeleteSelectedPanelItem() override;
 
+	/** Returns the UI item currently selected by the user (if any)*/
+	virtual TSharedPtr<FRCLogicModeBase> GetSelectedLogicItem() override
+	{
+		return GetSelectedBehaviourItem();
+	}
+
+	TSharedPtr<FRCBehaviourModel> GetSelectedBehaviourItem()
+	{
+		return SelectedBehaviourItemWeakPtr.Pin();
+	}
+
+	void RequestRefresh();
+
+	virtual void AddNewLogicItem(UObject* InLogicItem) override;
+
 private:
 
+	void AddBehaviourToList(URCBehaviour* InBehaviour);
+
 	/** OnGenerateRow delegate for the Behaviours List View */
 	TSharedRef<ITableRow> OnGenerateWidgetForList( TSharedPtr<FRCBehaviourModel> InItem, const TSharedRef<STableViewBase>& OwnerTable );
 	
diff --git a/Engine/Plugins/VirtualProduction/RemoteControl/Source/RemoteControlUI/Private/UI/Controller/SRCControllerPanel.cpp b/Engine/Plugins/VirtualProduction/RemoteControl/Source/RemoteControlUI/Private/UI/Controller/SRCControllerPanel.cpp
--- a/Engine/Plugins/VirtualProduction/RemoteControl/Source/RemoteControlUI/Private/UI/Controller/SRCControllerPanel.cpp
+++ b/Engine/Plugins/VirtualProduction/RemoteControl/Source/RemoteControlUI/Private/UI/Controller/SRCControllerPanel.cpp
@@ -39,7 +39,7 @@ void SRCControllerPanel::Construct(const FArguments& InArgs, const TSharedRef<SR
 	TSharedPtr<SRCMinorPanel> ControllerDockPanel = SNew(SRCMinorPanel)
 		.HeaderLabel(LOCTEXT("ControllersLabel", "Controller"))
 		[
-			SAssignNew(ControllerPanelList, SRCControllerPanelList, SharedThis(this))
+			SAssignNew(ControllerPanelList, SRCControllerPanelList, SharedThis(this), InPanel)
 		];
 
 	// Add New Controller Button
@@ -134,6 +134,93 @@ void SRCControllerPanel::DeleteSelectedPanelItem()
 	}
 }
 
+void SRCControllerPanel::DuplicateSelectedPanelItem()
+{
+	if (ControllerPanelList)
+	{
+		if(TSharedPtr<FRCControllerModel> ControllerItem = ControllerPanelList->GetSelectedControllerItem())
+		{
+			if (URCController* Controller = Cast<URCController>(ControllerItem->GetVirtualProperty()))
+			{
+				DuplicateController(Controller);
+			}
+		}
+	}
+}
+
+void SRCControllerPanel::CopySelectedPanelItem()
+{
+	if (TSharedPtr<SRemoteControlPanel> RemoteControlPanel = GetRemoteControlPanel())
+	{
+		if (TSharedPtr<FRCControllerModel> ControllerItem = ControllerPanelList->GetSelectedControllerItem())
+		{
+			if (URCController* Controller = Cast<URCController>(ControllerItem->GetVirtualProperty()))
+			{
+				RemoteControlPanel->SetLogicClipboardItem(Controller, SharedThis(this));
+			}
+		}
+	}
+}
+
+void SRCControllerPanel::PasteItemFromClipboard()
+{
+	if (TSharedPtr<SRemoteControlPanel> RemoteControlPanel = GetRemoteControlPanel())
+	{
+		if(RemoteControlPanel->LogicClipboardItemSource == SharedThis(this))
+		{
+			if (URCController* Controller = Cast<URCController>(RemoteControlPanel->GetLogicClipboardItem()))
+			{
+				DuplicateController(Controller);
+			}
+		}
+	}
+}
+
+FText SRCControllerPanel::GetPasteItemMenuEntrySuffix()
+{
+	if (TSharedPtr<SRemoteControlPanel> RemoteControlPanel = GetRemoteControlPanel())
+	{
+		// This function should only have been called if we were the source of the item copied.
+		if (ensure(RemoteControlPanel->LogicClipboardItemSource == SharedThis(this)))
+		{
+			if (URCController* Controller = Cast<URCController>(RemoteControlPanel->GetLogicClipboardItem()))
+			{
+				return FText::Format(FText::FromString("Controller {0}"), FText::FromName(Controller->DisplayName));
+			}
+		}
+	}
+
+	return FText::GetEmpty();
+}
+
+TSharedPtr<FRCLogicModeBase> SRCControllerPanel::GetSelectedLogicItem()
+{
+	if (ControllerPanelList)
+	{
+		return ControllerPanelList->GetSelectedControllerItem();
+	}
+
+	return nullptr;
+}
+
+void SRCControllerPanel::DuplicateController(URCController* InController)
+{
+	if (!ensure(InController))
+	{
+		return;
+	}
+
+	if (URemoteControlPreset* Preset = GetPreset())
+	{
+		if (URCController* NewController = Cast<URCController>(Preset->DuplicateVirtualProperty(InController)))
+		{
+			NewController->SetDisplayIndex(ControllerPanelList->NumControllerItems());
+
+			ControllerPanelList->RequestRefresh();
+		}
+	}
+}
+
 FReply SRCControllerPanel::RequestDeleteAllItems()
 {
 	if (!ControllerPanelList.IsValid())
@@ -246,7 +333,7 @@ void SRCControllerPanel::OnAddControllerClicked(const EPropertyBagPropertyType I
 
 		if (ControllerPanelList.IsValid())
 		{
-			NewVirtualProperty->DisplayIndex = ControllerPanelList->GetNumControllerItems();
+			NewVirtualProperty->DisplayIndex = ControllerPanelList->NumControllerItems();
 		}
 
 		// Refresh list
@@ -282,4 +369,4 @@ FReply SRCControllerPanel::OnClickEmptyButton()
 	return FReply::Handled();
 }
 
-#undef LOCTEXT_NAMESPACE
+#undef LOCTEXT_NAMESPACE
\ No newline at end of file
diff --git a/Engine/Plugins/VirtualProduction/RemoteControl/Source/RemoteControlUI/Private/UI/Controller/SRCControllerPanel.h b/Engine/Plugins/VirtualProduction/RemoteControl/Source/RemoteControlUI/Private/UI/Controller/SRCControllerPanel.h
--- a/Engine/Plugins/VirtualProduction/RemoteControl/Source/RemoteControlUI/Private/UI/Controller/SRCControllerPanel.h
+++ b/Engine/Plugins/VirtualProduction/RemoteControl/Source/RemoteControlUI/Private/UI/Controller/SRCControllerPanel.h
@@ -6,8 +6,10 @@
 #include "UI/BaseLogicUI/SRCLogicPanelBase.h"
 
 enum class EPropertyBagPropertyType : uint8;
+class FRCControllerModel;
 struct FRCPanelStyle;
 class SRCControllerPanel;
+class URCController;
 
 /*
 * ~ SRCControllerPanel ~
@@ -33,6 +35,21 @@ public:
 	/** Delete Item UI command implementation for this panel */
 	virtual void DeleteSelectedPanelItem() override;
 
+	/** "Duplicate Item" UI command implementation for Controller panel*/
+	virtual void DuplicateSelectedPanelItem() override;
+
+	/** "Copy Item" UI command implementation for Controller panel*/
+	virtual void CopySelectedPanelItem() override;
+
+	/** "Paste Item" UI command implementation for Controller panel*/
+	virtual void PasteItemFromClipboard() override;
+
+	/** Provides an item suffix for the Paste context menu to provide users with useful context on the nature of the item being pasted */
+	virtual FText GetPasteItemMenuEntrySuffix() override;
+
+	/** Returns the UI item currently selected by the user (if any)*/
+	virtual TSharedPtr<FRCLogicModeBase> GetSelectedLogicItem() override;
+
 	void EnterRenameMode();
 
 protected:
@@ -41,6 +58,8 @@ protected:
 	virtual FReply RequestDeleteAllItems() override;
 
 private:
+	/** Duplicates a given Controller object*/
+	void DuplicateController(URCController* InController);
 
 	/** Builds a menu containing the list of all possible Controllers*/
 	TSharedRef<SWidget> GetControllerMenuContentWidget() const;
diff --git a/Engine/Plugins/VirtualProduction/RemoteControl/Source/RemoteControlUI/Private/UI/Controller/SRCControllerPanelList.cpp b/Engine/Plugins/VirtualProduction/RemoteControl/Source/RemoteControlUI/Private/UI/Controller/SRCControllerPanelList.cpp
--- a/Engine/Plugins/VirtualProduction/RemoteControl/Source/RemoteControlUI/Private/UI/Controller/SRCControllerPanelList.cpp
+++ b/Engine/Plugins/VirtualProduction/RemoteControl/Source/RemoteControlUI/Private/UI/Controller/SRCControllerPanelList.cpp
@@ -2,8 +2,10 @@
 
 #include "SRCControllerPanelList.h"
 
+#include "Commands/RemoteControlCommands.h"
 #include "Controller/RCControllerContainer.h"
 #include "IDetailTreeNode.h"
+#include "Interfaces/IMainFrameModule.h"
 #include "IPropertyRowGenerator.h"
 #include "RCControllerModel.h"
 #include "RCVirtualProperty.h"
@@ -157,11 +159,9 @@ private:
 	};
 } 
 
-BEGIN_SLATE_FUNCTION_BUILD_OPTIMIZATION
-
-void SRCControllerPanelList::Construct(const FArguments& InArgs, const TSharedRef<SRCControllerPanel> InControllerPanel)
+void SRCControllerPanelList::Construct(const FArguments& InArgs, const TSharedRef<SRCControllerPanel> InControllerPanel, const TSharedRef<SRemoteControlPanel> InRemoteControlPanel)
 {
-	SRCLogicPanelListBase::Construct(SRCLogicPanelListBase::FArguments());
+	SRCLogicPanelListBase::Construct(SRCLogicPanelListBase::FArguments(), InRemoteControlPanel);
 	
 	ControllerPanelWeakPtr = InControllerPanel;
 	
@@ -172,6 +172,7 @@ void SRCControllerPanelList::Construct(const FArguments& InArgs, const TSharedRe
 		.OnSelectionChanged(this, &SRCControllerPanelList::OnTreeSelectionChanged)
 		.OnGenerateRow(this, &SRCControllerPanelList::OnGenerateWidgetForList)
 		.SelectionMode(ESelectionMode::Single) // Current setup supports only single selection (and related display) of a Controller in the list
+		.OnContextMenuOpening(this, &SRCLogicPanelListBase::GetContextMenuWidget)
 		.HeaderRow(
 			SNew(SHeaderRow)
 			.Style(&RCPanelStyle->HeaderRowStyle)
@@ -207,7 +208,6 @@ void SRCControllerPanelList::Construct(const FArguments& InArgs, const TSharedRe
 		RemoteControlPanel->OnEmptyControllers.AddSP(this, &SRCControllerPanelList::OnEmptyControllers);
 
 		Preset->OnVirtualPropertyContainerModified().AddSP(this, &SRCControllerPanelList::OnControllerContainerModified);
-
 	}
 
 	FPropertyRowGeneratorArgs Args;
@@ -217,8 +217,6 @@ void SRCControllerPanelList::Construct(const FArguments& InArgs, const TSharedRe
 	Reset();
 }
 
-END_SLATE_FUNCTION_BUILD_OPTIMIZATION
-
 bool SRCControllerPanelList::IsEmpty() const
 {
 	return ControllerItems.IsEmpty();
@@ -226,7 +224,7 @@ bool SRCControllerPanelList::IsEmpty() const
 
 int32 SRCControllerPanelList::Num() const
 {
-	return ControllerItems.Num();
+	return NumControllerItems();
 }
 
 void SRCControllerPanelList::Reset()
@@ -397,7 +395,7 @@ int32 SRCControllerPanelList::RemoveModel(const TSharedPtr<FRCLogicModeBase> InM
 
 bool SRCControllerPanelList::IsListFocused() const
 {
-	return ListView->HasAnyUserFocus().IsSet();
+	return ListView->HasAnyUserFocus().IsSet() || ContextMenuWidgetCached.IsValid();
 }
 
 void SRCControllerPanelList::DeleteSelectedPanelItem()
diff --git a/Engine/Plugins/VirtualProduction/RemoteControl/Source/RemoteControlUI/Private/UI/Controller/SRCControllerPanelList.h b/Engine/Plugins/VirtualProduction/RemoteControl/Source/RemoteControlUI/Private/UI/Controller/SRCControllerPanelList.h
--- a/Engine/Plugins/VirtualProduction/RemoteControl/Source/RemoteControlUI/Private/UI/Controller/SRCControllerPanelList.h
+++ b/Engine/Plugins/VirtualProduction/RemoteControl/Source/RemoteControlUI/Private/UI/Controller/SRCControllerPanelList.h
@@ -13,6 +13,7 @@ class IPropertyRowGenerator;
 class ITableRow;
 class ITableBase;
 class SRCControllerPanel;
+class SRemoteControlPanel;
 class STableViewBase;
 class URCController;
 class URemoteControlPreset;
@@ -65,7 +66,7 @@ public:
 	SLATE_END_ARGS()
 
 	/** Constructs this widget with InArgs */
-	void Construct(const FArguments& InArgs, const TSharedRef<SRCControllerPanel> InControllerPanel);
+	void Construct(const FArguments& InArgs, const TSharedRef<SRCControllerPanel> InControllerPanel, const TSharedRef<SRemoteControlPanel> InRemoteControlPanel);
 
 	/** Returns true if the underlying list is valid and empty. */
 	virtual bool IsEmpty() const override;
@@ -79,9 +80,15 @@ public:
 	/** Deletes currently selected items from the list view*/
 	virtual void DeleteSelectedPanelItem() override;
 
+	/** Returns the UI item currently selected by the user (if any)*/
+	virtual TSharedPtr<FRCLogicModeBase> GetSelectedLogicItem() override
+	{
+		return GetSelectedControllerItem();
+	}
+
 	void EnterRenameMode();
 
-	int32 GetNumControllerItems() const
+	int32 NumControllerItems() const
 	{
 		return ControllerItems.Num();
 	}
@@ -93,6 +100,18 @@ public:
 	* and pushes all other rows below */
 	void ReorderControllerItem(TSharedRef<FRCControllerModel> ItemToMove, TSharedRef<FRCControllerModel> AnchorItem);
 
+	/** Returns the currently selected Controller UI Item*/
+	TSharedPtr<FRCControllerModel> GetSelectedControllerItem() const
+	{
+		return SelectedControllerItemWeakPtr.Pin();
+	}
+
+	/** Requests the panel to refresh its contents from the latest list of Controllers*/
+	void RequestRefresh()
+	{
+		Reset();
+	}
+
 private:
 
 	/** OnGenerateRow delegate for the Actions List View */
diff --git a/Engine/Plugins/VirtualProduction/RemoteControl/Source/RemoteControlUI/Private/UI/SRemoteControlPanel.cpp b/Engine/Plugins/VirtualProduction/RemoteControl/Source/RemoteControlUI/Private/UI/SRemoteControlPanel.cpp
--- a/Engine/Plugins/VirtualProduction/RemoteControl/Source/RemoteControlUI/Private/UI/SRemoteControlPanel.cpp
+++ b/Engine/Plugins/VirtualProduction/RemoteControl/Source/RemoteControlUI/Private/UI/SRemoteControlPanel.cpp
@@ -1257,6 +1257,21 @@ void SRemoteControlPanel::BindRemoteControlCommands()
 		Commands.RenameEntity,
 		FExecuteAction::CreateSP(this, &SRemoteControlPanel::RenameEntity_Execute),
 		FCanExecuteAction::CreateSP(this, &SRemoteControlPanel::CanRenameEntity));
+
+	ActionList.MapAction(
+		Commands.CopyItem,
+		FExecuteAction::CreateSP(this, &SRemoteControlPanel::CopyItem_Execute),
+		FCanExecuteAction::CreateSP(this, &SRemoteControlPanel::CanCopyItem));
+
+	ActionList.MapAction(
+		Commands.PasteItem,
+		FExecuteAction::CreateSP(this, &SRemoteControlPanel::PasteItem_Execute),
+		FCanExecuteAction::CreateSP(this, &SRemoteControlPanel::CanPasteItem));
+
+	ActionList.MapAction(
+		Commands.DuplicateItem,
+		FExecuteAction::CreateSP(this, &SRemoteControlPanel::DuplicateItem_Execute),
+		FCanExecuteAction::CreateSP(this, &SRemoteControlPanel::CanDuplicateItem));
 }
 
 void SRemoteControlPanel::UnbindRemoteControlCommands()
@@ -2184,7 +2199,7 @@ void SRemoteControlPanel::DeleteEntity_Execute()
 	// ~ Delete Logic Item ~
 	// 
 	// If the user focus is currently active on a Logic panel then route the Delete command to it and return.
-	if (TSharedPtr<class SRCLogicPanelBase> ActiveLogicPanel = GetActiveLogicPanel())
+	if (TSharedPtr<SRCLogicPanelBase> ActiveLogicPanel = GetActiveLogicPanel())
 	{
 		ActiveLogicPanel->DeleteSelectedPanelItem();
 
@@ -2216,7 +2231,7 @@ void SRemoteControlPanel::DeleteEntity_Execute()
 
 bool SRemoteControlPanel::CanDeleteEntity() const
 {
-	if (TSharedPtr<class SRCLogicPanelBase> ActiveLogicPanel = GetActiveLogicPanel())
+	if (TSharedPtr<SRCLogicPanelBase> ActiveLogicPanel = GetActiveLogicPanel())
 	{
 		return true; // User has focus on a logic panel
 	}
@@ -2262,6 +2277,75 @@ bool SRemoteControlPanel::CanRenameEntity() const
 	return false;
 }
 
+void SRemoteControlPanel::SetLogicClipboardItem(UObject* InItem, TSharedPtr<SRCLogicPanelBase> InSourcePanel)
+{
+	LogicClipboardItem = InItem;
+	LogicClipboardItemSource = InSourcePanel;
+}
+
+void SRemoteControlPanel::CopyItem_Execute()
+{
+	if (TSharedPtr<SRCLogicPanelBase> ActiveLogicPanel = GetActiveLogicPanel())
+	{
+		ActiveLogicPanel->CopySelectedPanelItem();
+	}
+}
+
+bool SRemoteControlPanel::CanCopyItem() const
+{
+	if (TSharedPtr<SRCLogicPanelBase> ActiveLogicPanel = GetActiveLogicPanel())
+	{
+		return ActiveLogicPanel->GetSelectedLogicItem().IsValid();
+	}
+
+	return false;
+}
+
+void SRemoteControlPanel::PasteItem_Execute()
+{
+	if (TSharedPtr<SRCLogicPanelBase> ActiveLogicPanel = GetActiveLogicPanel())
+	{
+		ActiveLogicPanel->PasteItemFromClipboard();
+	}
+}
+
+bool SRemoteControlPanel::CanPasteItem() const
+{
+	if(LogicClipboardItem)
+	{
+		if (TSharedPtr<SRCLogicPanelBase> ActiveLogicPanel = GetActiveLogicPanel())
+		{
+			// Currently we only support pasting items between panels of exactly the same type.
+			if (LogicClipboardItemSource == ActiveLogicPanel)
+			{
+				return ActiveLogicPanel->CanPasteClipboardItem(LogicClipboardItem);
+			}
+		}
+	}
+
+	return false;
+}
+
+void SRemoteControlPanel::DuplicateItem_Execute()
+{
+	if (TSharedPtr<SRCLogicPanelBase> ActiveLogicPanel = GetActiveLogicPanel())
+	{
+		ActiveLogicPanel->DuplicateSelectedPanelItem();
+
+		return;
+	}
+}
+
+bool SRemoteControlPanel::CanDuplicateItem() const
+{
+	if (TSharedPtr<SRCLogicPanelBase> ActiveLogicPanel = GetActiveLogicPanel())
+	{
+		return ActiveLogicPanel->GetSelectedLogicItem().IsValid();
+	}
+
+	return false;
+}
+
 void SRemoteControlPanel::OnSearchTextChanged(const FText& InFilterText)
 {
 	SearchTextFilter->SetRawFilterText(InFilterText);
diff --git a/Engine/Plugins/VirtualProduction/RemoteControl/Source/RemoteControlUI/Private/UI/SRemoteControlPanel.h b/Engine/Plugins/VirtualProduction/RemoteControl/Source/RemoteControlUI/Private/UI/SRemoteControlPanel.h
--- a/Engine/Plugins/VirtualProduction/RemoteControl/Source/RemoteControlUI/Private/UI/SRemoteControlPanel.h
+++ b/Engine/Plugins/VirtualProduction/RemoteControl/Source/RemoteControlUI/Private/UI/SRemoteControlPanel.h
@@ -33,11 +33,12 @@ class SClassViewer;
 class SComboButton;
 struct SRCPanelTreeNode;
 class SRCPanelFunctionPicker;
-class SRemoteControlPanel;
 class SRCActionPanel;
+class SRCLogicPanelBase;
 class SRCPanelDrawer;
 class SRCPanelExposedEntitiesList;
 class SRCPanelFilter;
+class SRemoteControlPanel;
 class SSearchBox;
 class STextBlock;
 class URemoteControlPreset;
@@ -148,6 +149,15 @@ public:
 		return ActionPanel;
 	}
 
+	/** For Copy UI command - Sets the logic clipboard item and source */
+	void SetLogicClipboardItem(UObject* InItem, TSharedPtr<SRCLogicPanelBase> SourcePanel);
+
+	/** Fetches the last UI item copied to Logic clipboard by the user */
+	UObject* GetLogicClipboardItem()
+	{
+		return LogicClipboardItem;
+	}
+
 private:
 
 	//~ Remote Control Commands
@@ -291,6 +301,24 @@ private:
 
 	/** Called to test if user is able to rename a group/exposed entity. */
 	bool CanRenameEntity() const;
+
+	/** Called when user attempts to Copy a logic UI item. */
+	void CopyItem_Execute();
+
+	/** Called to test if user is able to Copy a logic UI item. */
+	bool CanCopyItem() const;
+
+	/** Called when user attempts to Paste a logic UI item. */
+	void PasteItem_Execute();
+
+	/** Called to test if user is able to Paste a logic UI item. */
+	bool CanPasteItem() const;
+
+	/** Called when user attempts to Duplicate a logic UI item. */
+	void DuplicateItem_Execute();
+
+	/** Called to test if user is able to Duplicate a logic UI item. */
+	bool CanDuplicateItem() const;
 	
 	// Exposed Entities filtering. (Filters the Exposed Entities view)
 	void OnSearchTextChanged(const FText& InFilterText);
@@ -307,7 +335,7 @@ private:
 	void SaveSettings();
 
 	/** Retrieves active logic panel. */
-	TSharedPtr<class SRCLogicPanelBase> GetActiveLogicPanel() const;
+	TSharedPtr<SRCLogicPanelBase> GetActiveLogicPanel() const;
 
 private:
 	static const FName DefaultRemoteControlPanelToolBarName;
@@ -394,6 +422,14 @@ private:
 	/** Action panel UI widget for Remote Control Logic*/
 	TSharedPtr<class SRCActionPanel> ActionPanel;
 
+	/** LogicClipboardItem - Holds the latest item copied from a Logic panel
+	* 
+	* Note: We track UObjects (Data Model) here rather than the UI Models as the latter are swept away the moment the user navigates to a different Controller.
+	* For example if the user copies an action from a behaviour in a given Controller but then navigates to another Controller, we can no longer rely on the previous UI objects
+	* as they would have been discarded in favor of a new data set for the actively selected Controller */
+	UPROPERTY(Transient)
+	TObjectPtr<UObject> LogicClipboardItem;
+
 public:
 
 	static const float MinimumPanelWidth;
@@ -406,4 +442,7 @@ public:
 	FOnEmptyControllers OnEmptyControllers;
 	FOnEmptyBehaviours OnEmptyBehaviours;
 	FOnEmptyActions OnEmptyActions;
+
+	/** The panel from which the latest Logic UI item was copied*/
+	TSharedPtr<SRCLogicPanelBase> LogicClipboardItemSource;
 };
\ No newline at end of file
