Delete trashcan icon should delete selected item only
UE-163228

diff --git a/Engine/Plugins/VirtualProduction/RemoteControl/Source/RemoteControlUI/Private/UI/Action/SRCActionPanel.cpp b/Engine/Plugins/VirtualProduction/RemoteControl/Source/RemoteControlUI/Private/UI/Action/SRCActionPanel.cpp
--- a/Engine/Plugins/VirtualProduction/RemoteControl/Source/RemoteControlUI/Private/UI/Action/SRCActionPanel.cpp
+++ b/Engine/Plugins/VirtualProduction/RemoteControl/Source/RemoteControlUI/Private/UI/Action/SRCActionPanel.cpp
@@ -145,16 +145,16 @@ void SRCActionPanel::UpdateWrappedWidget(TSharedPtr<FRCBehaviourModel> InBehavio
 				]
 			];
 		
-		// Empty All Button
-		TSharedRef<SWidget> EmptyAllButton = SNew(SButton)
-			.AddMetaData<FTagMetaData>(FTagMetaData(TEXT("Empty All Actions")))
+	// Delete Selected Action Button
+		TSharedRef<SWidget> DeleteSelectedActionButton = SNew(SButton)
+			.AddMetaData<FTagMetaData>(FTagMetaData(TEXT("Delete Selected Action")))
 			.HAlign(HAlign_Center)
 			.VAlign(VAlign_Center)
 			.ForegroundColor(FSlateColor::UseForeground())
 			.ButtonStyle(&RCPanelStyle->FlatButtonStyle)
-			.ToolTipText(LOCTEXT("EmptyAllToolTip", "Deletes all the actions."))
-			.OnClicked(this, &SRCActionPanel::RequestDeleteAllItems)
-			.Visibility_Lambda([this]() { return ActionPanelList.IsValid() && !ActionPanelList->IsEmpty() ? EVisibility::Visible : EVisibility::Collapsed; })
+			.ToolTipText(LOCTEXT("DeleteSelectedActionToolTip", "Deletes all selected Action."))
+			.OnClicked(this, &SRCActionPanel::RequestDeleteSelectedItem)
+			.IsEnabled_Lambda([this]() { return ActionPanelList.IsValid() && ActionPanelList->NumSelectedLogicItems() > 0; })
 			[
 				SNew(SBox)
 				.WidthOverride(RCPanelStyle->IconSize.X)
@@ -168,7 +168,7 @@ void SRCActionPanel::UpdateWrappedWidget(TSharedPtr<FRCBehaviourModel> InBehavio
 
 		ActionDockPanel->AddHeaderToolbarItem(EToolbar::Left, AddNewActionButton);
 		ActionDockPanel->AddHeaderToolbarItem(EToolbar::Right, AddAllActionsButton);
-		ActionDockPanel->AddHeaderToolbarItem(EToolbar::Right, EmptyAllButton);
+		ActionDockPanel->AddHeaderToolbarItem(EToolbar::Right, DeleteSelectedActionButton);
 
 		// Duplicate Dock Panel
 		TSharedPtr<SRCMinorPanel> DuplicateDockPanel = SNew(SRCMinorPanel)
@@ -569,6 +569,25 @@ FText SRCActionPanel::GetPasteItemMenuEntrySuffix()
 	return FText::GetEmpty();
 }
 
+FReply SRCActionPanel::RequestDeleteSelectedItem()
+{
+	if (!ActionPanelList.IsValid())
+	{
+		return FReply::Unhandled();
+	}
+
+	const FText WarningMessage = LOCTEXT("DeleteActionWarning", "Delete the selected Actions?");
+
+	EAppReturnType::Type UserResponse = FMessageDialog::Open(EAppMsgType::YesNo, WarningMessage);
+
+	if (UserResponse == EAppReturnType::Yes)
+	{
+		DeleteSelectedPanelItem();
+	}
+
+	return FReply::Handled();
+}
+
 FReply SRCActionPanel::RequestDeleteAllItems()
 {
 	if (!ActionPanelList.IsValid())
@@ -576,7 +595,7 @@ FReply SRCActionPanel::RequestDeleteAllItems()
 		return FReply::Unhandled();
 	}
 
-	const FText WarningMessage = FText::Format(LOCTEXT("DeleteAllWarning", "You are about to delete '{0}' actions. This action might not be undone.\nAre you sure you want to proceed?"), ActionPanelList->Num());
+	const FText WarningMessage = FText::Format(LOCTEXT("DeleteAllWarning", "You are about to delete {0} actions. Are you sure you want to proceed?"), ActionPanelList->Num());
 
 	EAppReturnType::Type UserResponse = FMessageDialog::Open(EAppMsgType::YesNo, WarningMessage);
 
diff --git a/Engine/Plugins/VirtualProduction/RemoteControl/Source/RemoteControlUI/Private/UI/Action/SRCActionPanel.h b/Engine/Plugins/VirtualProduction/RemoteControl/Source/RemoteControlUI/Private/UI/Action/SRCActionPanel.h
--- a/Engine/Plugins/VirtualProduction/RemoteControl/Source/RemoteControlUI/Private/UI/Action/SRCActionPanel.h
+++ b/Engine/Plugins/VirtualProduction/RemoteControl/Source/RemoteControlUI/Private/UI/Action/SRCActionPanel.h
@@ -70,6 +70,8 @@ public:
 	}
 
 protected:
+	/** Warns user before deleting a selected panel item. */
+	virtual FReply RequestDeleteSelectedItem() override;
 
 	/** Warns user before deleting all items in a panel. */
 	virtual FReply RequestDeleteAllItems() override;
diff --git a/Engine/Plugins/VirtualProduction/RemoteControl/Source/RemoteControlUI/Private/UI/Action/SRCActionPanelList.h b/Engine/Plugins/VirtualProduction/RemoteControl/Source/RemoteControlUI/Private/UI/Action/SRCActionPanelList.h
--- a/Engine/Plugins/VirtualProduction/RemoteControl/Source/RemoteControlUI/Private/UI/Action/SRCActionPanelList.h
+++ b/Engine/Plugins/VirtualProduction/RemoteControl/Source/RemoteControlUI/Private/UI/Action/SRCActionPanelList.h
@@ -55,7 +55,7 @@ public:
 		const TSharedPtr<SRemoteControlPanel>& RemoteControlPanel = InActionPanel->GetRemoteControlPanel();
 		check(RemoteControlPanel);
 
-		SRCLogicPanelListBase::Construct(SRCLogicPanelListBase::FArguments(), RemoteControlPanel.ToSharedRef());
+		SRCLogicPanelListBase::Construct(SRCLogicPanelListBase::FArguments(), InActionPanel, RemoteControlPanel.ToSharedRef());
 
 		ActionPanelWeakPtr = InActionPanel;
 		BehaviourItemWeakPtr = InBehaviourItem;
@@ -115,6 +115,12 @@ public:
 		return ActionItems.Num();
 	}
 
+	/** The number of Controllers currently selected*/
+	virtual int32 NumSelectedLogicItems() const override
+	{
+		return ListView->GetNumItemsSelected();
+	}
+
 	/** Whether the Actions List View currently has focus.*/
 	virtual bool IsListFocused() const override
 	{
diff --git a/Engine/Plugins/VirtualProduction/RemoteControl/Source/RemoteControlUI/Private/UI/BaseLogicUI/SRCLogicPanelBase.h b/Engine/Plugins/VirtualProduction/RemoteControl/Source/RemoteControlUI/Private/UI/BaseLogicUI/SRCLogicPanelBase.h
--- a/Engine/Plugins/VirtualProduction/RemoteControl/Source/RemoteControlUI/Private/UI/BaseLogicUI/SRCLogicPanelBase.h
+++ b/Engine/Plugins/VirtualProduction/RemoteControl/Source/RemoteControlUI/Private/UI/BaseLogicUI/SRCLogicPanelBase.h
@@ -56,7 +56,8 @@ public:
 	/** Provides an item suffix for the Paste context menu to provide users with useful context on the nature of the item being pasted */
 	virtual FText GetPasteItemMenuEntrySuffix() { return FText::GetEmpty(); }
 
-protected:
+	/** Warns user before deleting a selected panel item. */
+	virtual FReply RequestDeleteSelectedItem() = 0;
 
 	/** Warns user before deleting all items in a panel. */
 	virtual FReply RequestDeleteAllItems() = 0;
diff --git a/Engine/Plugins/VirtualProduction/RemoteControl/Source/RemoteControlUI/Private/UI/BaseLogicUI/SRCLogicPanelListBase.cpp b/Engine/Plugins/VirtualProduction/RemoteControl/Source/RemoteControlUI/Private/UI/BaseLogicUI/SRCLogicPanelListBase.cpp
--- a/Engine/Plugins/VirtualProduction/RemoteControl/Source/RemoteControlUI/Private/UI/BaseLogicUI/SRCLogicPanelListBase.cpp
+++ b/Engine/Plugins/VirtualProduction/RemoteControl/Source/RemoteControlUI/Private/UI/BaseLogicUI/SRCLogicPanelListBase.cpp
@@ -9,8 +9,9 @@
 
 #define LOCTEXT_NAMESPACE "SRCLogicPanelListBase"
 
-void SRCLogicPanelListBase::Construct(const FArguments& InArgs, const TSharedRef<SRemoteControlPanel>& InPanel)
+void SRCLogicPanelListBase::Construct(const FArguments& InArgs, const TSharedRef<SRCLogicPanelBase>& InLogicParentPanel, const TSharedRef<SRemoteControlPanel>& InPanel)
 {
+	LogicPanelWeakPtr = InLogicParentPanel;
 	RemoteControlPanelWeakPtr = InPanel;
 }
 
@@ -54,6 +55,13 @@ TSharedPtr<SWidget> SRCLogicPanelListBase::GetContextMenuWidget()
 	// 3. Delete
 	MenuBuilder.AddMenuEntry(Commands.DeleteEntity);
 
+	// 4. Delete All
+	FUIAction Action(FExecuteAction::CreateSP(this, &SRCLogicPanelListBase::RequestDeleteAllItems));
+
+	MenuBuilder.AddMenuEntry(LOCTEXT("DeleteAll", "Delete All"),
+		LOCTEXT("ContextMenuEditTooltip", "Delete all the rows in this list"),
+		FSlateIcon(), Action);
+
 	MenuBuilder.EndSection();
 
 	TSharedPtr<SWidget> MenuWidget = MenuBuilder.MakeWidget();
@@ -62,4 +70,12 @@ TSharedPtr<SWidget> SRCLogicPanelListBase::GetContextMenuWidget()
 	return MenuWidget;
 }
 
+void SRCLogicPanelListBase::RequestDeleteAllItems()
+{
+	if (TSharedPtr<SRCLogicPanelBase> LogicParentPanel = LogicPanelWeakPtr.Pin())
+	{
+		LogicParentPanel->RequestDeleteAllItems();
+	}
+}
+
 #undef LOCTEXT_NAMESPACE
\ No newline at end of file
diff --git a/Engine/Plugins/VirtualProduction/RemoteControl/Source/RemoteControlUI/Private/UI/BaseLogicUI/SRCLogicPanelListBase.h b/Engine/Plugins/VirtualProduction/RemoteControl/Source/RemoteControlUI/Private/UI/BaseLogicUI/SRCLogicPanelListBase.h
--- a/Engine/Plugins/VirtualProduction/RemoteControl/Source/RemoteControlUI/Private/UI/BaseLogicUI/SRCLogicPanelListBase.h
+++ b/Engine/Plugins/VirtualProduction/RemoteControl/Source/RemoteControlUI/Private/UI/BaseLogicUI/SRCLogicPanelListBase.h
@@ -12,6 +12,7 @@
 #include "UI/Controller/RCControllerModel.h"
 #include "Widgets/SCompoundWidget.h"
 
+class SRCLogicPanelBase;
 class SRemoteControlPanel;
 
 /*
@@ -30,7 +31,7 @@ public:
 	SLATE_END_ARGS()
 
 	/** Constructs this widget with InArgs */
-	void Construct(const FArguments& InArgs, const TSharedRef<SRemoteControlPanel>& InPanel);
+	void Construct(const FArguments& InArgs, const TSharedRef<SRCLogicPanelBase>& InLogicParentPanel, const TSharedRef<SRemoteControlPanel>& InPanel);
 
 	/** Returns true if the underlying list is valid and empty. */
 	virtual bool IsEmpty() const = 0;
@@ -38,6 +39,9 @@ public:
 	/** Returns number of items in the list. */
 	virtual int32 Num() const = 0;
 
+	/** The number of items currently selected in a Logic list panel */
+	virtual int32 NumSelectedLogicItems() const = 0;
+
 	/** Whether the List View currently has focus.*/
 	virtual bool IsListFocused() const = 0;
 
@@ -71,6 +75,9 @@ private:
 	/** Removes the given UI model item from the list of UI models for this panel list*/
 	virtual int32 RemoveModel(const TSharedPtr<FRCLogicModeBase> InModel) = 0;
 
+	/** Requests the parent logic panel to delete all items in this list*/
+	void RequestDeleteAllItems();
+
 protected:
 	TWeakPtr<SWidget> ContextMenuWidgetCached;
 
@@ -113,4 +120,7 @@ protected:
 
 	/** The parent Remote Control Panel widget*/
 	TWeakPtr<SRemoteControlPanel> RemoteControlPanelWeakPtr;
+
+	/** The parent Logic Panel*/
+	TWeakPtr<SRCLogicPanelBase> LogicPanelWeakPtr;
 };
\ No newline at end of file
diff --git a/Engine/Plugins/VirtualProduction/RemoteControl/Source/RemoteControlUI/Private/UI/Behaviour/SRCBehaviourPanel.cpp b/Engine/Plugins/VirtualProduction/RemoteControl/Source/RemoteControlUI/Private/UI/Behaviour/SRCBehaviourPanel.cpp
--- a/Engine/Plugins/VirtualProduction/RemoteControl/Source/RemoteControlUI/Private/UI/Behaviour/SRCBehaviourPanel.cpp
+++ b/Engine/Plugins/VirtualProduction/RemoteControl/Source/RemoteControlUI/Private/UI/Behaviour/SRCBehaviourPanel.cpp
@@ -111,16 +111,16 @@ void SRCBehaviourPanel::UpdateWrappedWidget(TSharedPtr<FRCControllerModel> InCon
 				GetBehaviourMenuContentWidget()
 			];
 
-		// Empty All Button
-		TSharedRef<SWidget> EmptyAllButton = SNew(SButton)
-			.AddMetaData<FTagMetaData>(FTagMetaData(TEXT("Empty Behaviors")))
+		// Delete Selected Behaviour Button
+		TSharedRef<SWidget> DeleteSelectedBehaviourButton = SNew(SButton)
+			.AddMetaData<FTagMetaData>(FTagMetaData(TEXT("Delete Selected Behaviour")))
 			.HAlign(HAlign_Center)
 			.VAlign(VAlign_Center)
 			.ForegroundColor(FSlateColor::UseForeground())
 			.ButtonStyle(&RCPanelStyle->FlatButtonStyle)
-			.ToolTipText(LOCTEXT("EmptyAllToolTip", "Deletes all the behaviors."))
-			.OnClicked(this, &SRCBehaviourPanel::RequestDeleteAllItems)
-			.Visibility_Lambda([this]() { return BehaviourPanelList.IsValid() && !BehaviourPanelList->IsEmpty() ? EVisibility::Visible : EVisibility::Collapsed; })
+			.ToolTipText(LOCTEXT("DeleteSelectedBehaviourToolTip", "Deletes the selected behaviour."))
+			.OnClicked(this, &SRCBehaviourPanel::RequestDeleteSelectedItem)
+			.IsEnabled_Lambda([this]() { return BehaviourPanelList.IsValid() && BehaviourPanelList->NumSelectedLogicItems(); })
 			[
 				SNew(SBox)
 				.WidthOverride(RCPanelStyle->IconSize.X)
@@ -133,7 +133,7 @@ void SRCBehaviourPanel::UpdateWrappedWidget(TSharedPtr<FRCControllerModel> InCon
 			];
 
 		BehaviourDockPanel->AddHeaderToolbarItem(EToolbar::Left, AddNewBehaviourButton);
-		BehaviourDockPanel->AddHeaderToolbarItem(EToolbar::Right, EmptyAllButton);
+		BehaviourDockPanel->AddHeaderToolbarItem(EToolbar::Right, DeleteSelectedBehaviourButton);
 
 		WrappedBoxWidget->SetContent(BehaviourDockPanel.ToSharedRef());
 	}
@@ -218,6 +218,7 @@ FReply SRCBehaviourPanel::OnClickEmptyButton()
 		if (URCController* Controller = Cast<URCController>(ControllerItem->GetVirtualProperty()))
 		{
 			FScopedTransaction Transaction(LOCTEXT("EmptyBehaviours", "Empty Behaviours"));
+			Controller->Modify();
 
 			Controller->EmptyBehaviours();
 		}
@@ -346,6 +347,25 @@ URCController* SRCBehaviourPanel::GetParentController()
 	return nullptr;
 }
 
+FReply SRCBehaviourPanel::RequestDeleteSelectedItem()
+{
+	if (!BehaviourPanelList.IsValid())
+	{
+		return FReply::Unhandled();
+	}
+
+	const FText WarningMessage = LOCTEXT("DeleteBehaviourlWarning", "Delete the selected Behaviours?");
+
+	EAppReturnType::Type UserResponse = FMessageDialog::Open(EAppMsgType::YesNo, WarningMessage);
+
+	if (UserResponse == EAppReturnType::Yes)
+	{
+		DeleteSelectedPanelItem();
+	}
+
+	return FReply::Handled();
+}
+
 FReply SRCBehaviourPanel::RequestDeleteAllItems()
 {
 	if (!BehaviourPanelList.IsValid())
@@ -353,7 +373,7 @@ FReply SRCBehaviourPanel::RequestDeleteAllItems()
 		return FReply::Unhandled();
 	}
 
-	const FText WarningMessage = FText::Format(LOCTEXT("DeleteAllWarning", "You are about to delete '{0}' behaviors. This action might not be undone.\nAre you sure you want to proceed?"), BehaviourPanelList->Num());
+	const FText WarningMessage = FText::Format(LOCTEXT("DeleteAllWarning", "You are about to delete {0} behaviors. Are you sure you want to proceed?"), BehaviourPanelList->Num());
 
 	EAppReturnType::Type UserResponse = FMessageDialog::Open(EAppMsgType::YesNo, WarningMessage);
 
diff --git a/Engine/Plugins/VirtualProduction/RemoteControl/Source/RemoteControlUI/Private/UI/Behaviour/SRCBehaviourPanel.h b/Engine/Plugins/VirtualProduction/RemoteControl/Source/RemoteControlUI/Private/UI/Behaviour/SRCBehaviourPanel.h
--- a/Engine/Plugins/VirtualProduction/RemoteControl/Source/RemoteControlUI/Private/UI/Behaviour/SRCBehaviourPanel.h
+++ b/Engine/Plugins/VirtualProduction/RemoteControl/Source/RemoteControlUI/Private/UI/Behaviour/SRCBehaviourPanel.h
@@ -62,6 +62,9 @@ public:
 
 protected:
 
+	/** Warns user before deleting a selected panel item. */
+	virtual FReply RequestDeleteSelectedItem() override;
+
 	/** Warns user before deleting all items in a panel. */
 	virtual FReply RequestDeleteAllItems() override;
 
diff --git a/Engine/Plugins/VirtualProduction/RemoteControl/Source/RemoteControlUI/Private/UI/Behaviour/SRCBehaviourPanelList.cpp b/Engine/Plugins/VirtualProduction/RemoteControl/Source/RemoteControlUI/Private/UI/Behaviour/SRCBehaviourPanelList.cpp
--- a/Engine/Plugins/VirtualProduction/RemoteControl/Source/RemoteControlUI/Private/UI/Behaviour/SRCBehaviourPanelList.cpp
+++ b/Engine/Plugins/VirtualProduction/RemoteControl/Source/RemoteControlUI/Private/UI/Behaviour/SRCBehaviourPanelList.cpp
@@ -34,7 +34,7 @@ namespace FRemoteControlBehaviourColumns
 
 void SRCBehaviourPanelList::Construct(const FArguments& InArgs, TSharedRef<SRCBehaviourPanel> InBehaviourPanel, TSharedPtr<FRCControllerModel> InControllerItem, const TSharedRef<SRemoteControlPanel> InRemoteControlPanel)
 {
-	SRCLogicPanelListBase::Construct(SRCLogicPanelListBase::FArguments(), InRemoteControlPanel);
+	SRCLogicPanelListBase::Construct(SRCLogicPanelListBase::FArguments(), InBehaviourPanel, InRemoteControlPanel);
 	
 	BehaviourPanelWeakPtr = InBehaviourPanel;
 	ControllerItemWeakPtr = InControllerItem;
@@ -86,6 +86,11 @@ int32 SRCBehaviourPanelList::Num() const
 	return BehaviourItems.Num();
 }
 
+int32 SRCBehaviourPanelList::NumSelectedLogicItems() const
+{
+	return ListView->GetNumItemsSelected();
+}
+
 void SRCBehaviourPanelList::AddNewLogicItem(UObject* InLogicItem)
 {
 	AddBehaviourToList(Cast<URCBehaviour>(InLogicItem));
diff --git a/Engine/Plugins/VirtualProduction/RemoteControl/Source/RemoteControlUI/Private/UI/Behaviour/SRCBehaviourPanelList.h b/Engine/Plugins/VirtualProduction/RemoteControl/Source/RemoteControlUI/Private/UI/Behaviour/SRCBehaviourPanelList.h
--- a/Engine/Plugins/VirtualProduction/RemoteControl/Source/RemoteControlUI/Private/UI/Behaviour/SRCBehaviourPanelList.h
+++ b/Engine/Plugins/VirtualProduction/RemoteControl/Source/RemoteControlUI/Private/UI/Behaviour/SRCBehaviourPanelList.h
@@ -40,6 +40,9 @@ public:
 	/** Returns number of items in the list. */
 	virtual int32 Num() const override;
 
+	/** The number of Controllers currently selected*/
+	virtual int32 NumSelectedLogicItems() const override;
+
 	/** Whether the Behaviours List View currently has focus.*/
 	virtual bool IsListFocused() const override;
 
diff --git a/Engine/Plugins/VirtualProduction/RemoteControl/Source/RemoteControlUI/Private/UI/Controller/SRCControllerPanel.cpp b/Engine/Plugins/VirtualProduction/RemoteControl/Source/RemoteControlUI/Private/UI/Controller/SRCControllerPanel.cpp
--- a/Engine/Plugins/VirtualProduction/RemoteControl/Source/RemoteControlUI/Private/UI/Controller/SRCControllerPanel.cpp
+++ b/Engine/Plugins/VirtualProduction/RemoteControl/Source/RemoteControlUI/Private/UI/Controller/SRCControllerPanel.cpp
@@ -68,16 +68,16 @@ void SRCControllerPanel::Construct(const FArguments& InArgs, const TSharedRef<SR
 			GetControllerMenuContentWidget()
 		];
 
-	// Empty All Button
-	TSharedRef<SWidget> EmptyAllButton = SNew(SButton)
-		.AddMetaData<FTagMetaData>(FTagMetaData(TEXT("Empty Controllers")))
+	// Delete Selected Controller Button
+	TSharedRef<SWidget> DeleteSelectedControllerButton = SNew(SButton)
+		.AddMetaData<FTagMetaData>(FTagMetaData(TEXT("Delete Selected Controller")))
 		.HAlign(HAlign_Center)
 		.VAlign(VAlign_Center)
 		.ForegroundColor(FSlateColor::UseForeground())
 		.ButtonStyle(&RCPanelStyle->FlatButtonStyle)
-		.ToolTipText(LOCTEXT("EmptyAllToolTip", "Deletes all the controllers."))
-		.OnClicked(this, &SRCControllerPanel::RequestDeleteAllItems)
-		.Visibility_Lambda([this]() { return ControllerPanelList.IsValid() && !ControllerPanelList->IsEmpty() ? EVisibility::Visible : EVisibility::Collapsed; })
+		.ToolTipText(LOCTEXT("DeleteSelectedControllerToolTip", "Deletes the selected controller."))
+		.OnClicked(this, &SRCControllerPanel::RequestDeleteSelectedItem)
+		.IsEnabled_Lambda([this]() { return ControllerPanelList.IsValid() && ControllerPanelList->NumSelectedLogicItems() > 0; })
 		[
 			SNew(SBox)
 			.WidthOverride(RCPanelStyle->IconSize.X)
@@ -90,7 +90,7 @@ void SRCControllerPanel::Construct(const FArguments& InArgs, const TSharedRef<SR
 		];
 
 	ControllerDockPanel->AddHeaderToolbarItem(EToolbar::Left, AddNewControllerButton);
-	ControllerDockPanel->AddHeaderToolbarItem(EToolbar::Right, EmptyAllButton);
+	ControllerDockPanel->AddHeaderToolbarItem(EToolbar::Right, DeleteSelectedControllerButton);
 
 	ChildSlot
 		.Padding(RCPanelStyle->PanelPadding)
@@ -221,6 +221,25 @@ void SRCControllerPanel::DuplicateController(URCController* InController)
 	}
 }
 
+FReply SRCControllerPanel::RequestDeleteSelectedItem()
+{
+	if (!ControllerPanelList.IsValid())
+	{
+		return FReply::Unhandled();
+	}
+
+	const FText WarningMessage = LOCTEXT("DeleteSelectedItemlWarning", "Delete the selected Controller?");
+
+	EAppReturnType::Type UserResponse = FMessageDialog::Open(EAppMsgType::YesNo, WarningMessage);
+
+	if (UserResponse == EAppReturnType::Yes)
+	{
+		DeleteSelectedPanelItem();
+	}
+
+	return FReply::Handled();
+}
+
 FReply SRCControllerPanel::RequestDeleteAllItems()
 {
 	if (!ControllerPanelList.IsValid())
@@ -228,7 +247,7 @@ FReply SRCControllerPanel::RequestDeleteAllItems()
 		return FReply::Unhandled();
 	}
 
-	const FText WarningMessage = FText::Format(LOCTEXT("DeleteAllWarning", "You are about to delete '{0}' controllers. This action might not be undone.\nAre you sure you want to proceed?"), ControllerPanelList->Num());
+	const FText WarningMessage = FText::Format(LOCTEXT("DeleteAllWarning", "You are about to delete {0} controllers. Are you sure you want to proceed?"), ControllerPanelList->Num());
 
 	EAppReturnType::Type UserResponse = FMessageDialog::Open(EAppMsgType::YesNo, WarningMessage);
 
@@ -355,7 +374,7 @@ FReply SRCControllerPanel::OnClickEmptyButton()
 {
 	if (URemoteControlPreset* Preset = GetPreset())
 	{
-		FScopedTransaction Transaction(LOCTEXT("Empty Controllers", "EmptyControllers"));
+		FScopedTransaction Transaction(LOCTEXT("EmptyControllers", "Empty Controllers"));
 		Preset->Modify();
 
 		Preset->ResetVirtualProperties();
diff --git a/Engine/Plugins/VirtualProduction/RemoteControl/Source/RemoteControlUI/Private/UI/Controller/SRCControllerPanel.h b/Engine/Plugins/VirtualProduction/RemoteControl/Source/RemoteControlUI/Private/UI/Controller/SRCControllerPanel.h
--- a/Engine/Plugins/VirtualProduction/RemoteControl/Source/RemoteControlUI/Private/UI/Controller/SRCControllerPanel.h
+++ b/Engine/Plugins/VirtualProduction/RemoteControl/Source/RemoteControlUI/Private/UI/Controller/SRCControllerPanel.h
@@ -53,6 +53,8 @@ public:
 	void EnterRenameMode();
 
 protected:
+	/** Warns user before deleting a selected panel item. */
+	virtual FReply RequestDeleteSelectedItem() override;
 
 	/** Warns user before deleting all items in a panel. */
 	virtual FReply RequestDeleteAllItems() override;
diff --git a/Engine/Plugins/VirtualProduction/RemoteControl/Source/RemoteControlUI/Private/UI/Controller/SRCControllerPanelList.cpp b/Engine/Plugins/VirtualProduction/RemoteControl/Source/RemoteControlUI/Private/UI/Controller/SRCControllerPanelList.cpp
--- a/Engine/Plugins/VirtualProduction/RemoteControl/Source/RemoteControlUI/Private/UI/Controller/SRCControllerPanelList.cpp
+++ b/Engine/Plugins/VirtualProduction/RemoteControl/Source/RemoteControlUI/Private/UI/Controller/SRCControllerPanelList.cpp
@@ -161,7 +161,7 @@ private:
 
 void SRCControllerPanelList::Construct(const FArguments& InArgs, const TSharedRef<SRCControllerPanel> InControllerPanel, const TSharedRef<SRemoteControlPanel> InRemoteControlPanel)
 {
-	SRCLogicPanelListBase::Construct(SRCLogicPanelListBase::FArguments(), InRemoteControlPanel);
+	SRCLogicPanelListBase::Construct(SRCLogicPanelListBase::FArguments(), InControllerPanel, InRemoteControlPanel);
 	
 	ControllerPanelWeakPtr = InControllerPanel;
 	
@@ -227,6 +227,11 @@ int32 SRCControllerPanelList::Num() const
 	return NumControllerItems();
 }
 
+int32 SRCControllerPanelList::NumSelectedLogicItems() const
+{
+	return ListView->GetNumItemsSelected();
+}
+
 void SRCControllerPanelList::Reset()
 {
 	ControllerItems.Empty();
diff --git a/Engine/Plugins/VirtualProduction/RemoteControl/Source/RemoteControlUI/Private/UI/Controller/SRCControllerPanelList.h b/Engine/Plugins/VirtualProduction/RemoteControl/Source/RemoteControlUI/Private/UI/Controller/SRCControllerPanelList.h
--- a/Engine/Plugins/VirtualProduction/RemoteControl/Source/RemoteControlUI/Private/UI/Controller/SRCControllerPanelList.h
+++ b/Engine/Plugins/VirtualProduction/RemoteControl/Source/RemoteControlUI/Private/UI/Controller/SRCControllerPanelList.h
@@ -74,6 +74,9 @@ public:
 	/** Returns number of items in the list. */
 	virtual int32 Num() const override;
 
+	/** The number of Controllers currently selected*/
+	virtual int32 NumSelectedLogicItems() const override;
+
 	/** Whether the Controllers List View currently has focus.*/
 	virtual bool IsListFocused() const override;
 
diff --git a/Engine/Plugins/VirtualProduction/RemoteControl/Source/RemoteControlUI/Private/UI/SRemoteControlPanel.cpp b/Engine/Plugins/VirtualProduction/RemoteControl/Source/RemoteControlUI/Private/UI/SRemoteControlPanel.cpp
--- a/Engine/Plugins/VirtualProduction/RemoteControl/Source/RemoteControlUI/Private/UI/SRemoteControlPanel.cpp
+++ b/Engine/Plugins/VirtualProduction/RemoteControl/Source/RemoteControlUI/Private/UI/SRemoteControlPanel.cpp
@@ -2201,7 +2201,7 @@ void SRemoteControlPanel::DeleteEntity_Execute()
 	// If the user focus is currently active on a Logic panel then route the Delete command to it and return.
 	if (TSharedPtr<SRCLogicPanelBase> ActiveLogicPanel = GetActiveLogicPanel())
 	{
-		ActiveLogicPanel->DeleteSelectedPanelItem();
+		ActiveLogicPanel->RequestDeleteSelectedItem();
 
 		return; // handled
 	}
@@ -2233,7 +2233,7 @@ bool SRemoteControlPanel::CanDeleteEntity() const
 {
 	if (TSharedPtr<SRCLogicPanelBase> ActiveLogicPanel = GetActiveLogicPanel())
 	{
-		return true; // User has focus on a logic panel
+		return ActiveLogicPanel->GetSelectedLogicItem() != nullptr; // User has focus on a logic panel
 	}
 
 	if (SelectedEntity.IsValid() && Preset.IsValid() && bIsInEditMode)
