[EpicStageApp] Disable Stage App beacon unless the WebSocket server is running

diff --git a/Engine/Plugins/VirtualProduction/EpicStageApp/Source/EpicStageApp/Private/EpicStageApp.cpp b/Engine/Plugins/VirtualProduction/EpicStageApp/Source/EpicStageApp/Private/EpicStageApp.cpp
--- a/Engine/Plugins/VirtualProduction/EpicStageApp/Source/EpicStageApp/Private/EpicStageApp.cpp
+++ b/Engine/Plugins/VirtualProduction/EpicStageApp/Source/EpicStageApp/Private/EpicStageApp.cpp
@@ -10,7 +10,18 @@ void FEpicStageAppModule::StartupModule()
 
 	if (!IsRunningCommandlet())
 	{
-		StageAppBeaconReceiver.Startup();
+		if (WebRemoteControl.IsWebSocketServerRunning())
+		{
+			StartupBeaconReceiver();
+		}
+		else
+		{
+			// Wait for the WebSocket server to start
+			WebSocketServerStartedDelegate = WebRemoteControl.OnWebSocketServerStarted().AddLambda([this](uint32)
+			{
+				StartupBeaconReceiver();
+			});
+		}
 	}
 }
 
@@ -19,10 +30,30 @@ void FEpicStageAppModule::ShutdownModule()
 	if (IWebRemoteControlModule* WebRemoteControl = FModuleManager::GetModulePtr<IWebRemoteControlModule>("WebRemoteControl"))
 	{
 		RouteHandler.UnregisterRoutes(*WebRemoteControl);
+
+		if (WebSocketServerStartedDelegate.IsValid())
+		{
+			WebRemoteControl->OnWebSocketServerStarted().Remove(WebSocketServerStartedDelegate);
+		}
 	}
 
-	if (!IsRunningCommandlet())
+	ShutdownBeaconReceiver();
+}
+
+void FEpicStageAppModule::StartupBeaconReceiver()
+{
+	// Make sure it's not already running
+	ShutdownBeaconReceiver();
+
+	StageAppBeaconReceiver.Startup();
+	bIsBeaconReceiverRunning = true;
+}
+
+void FEpicStageAppModule::ShutdownBeaconReceiver()
+{
+	if (bIsBeaconReceiverRunning)
 	{
 		StageAppBeaconReceiver.Shutdown();
+		bIsBeaconReceiverRunning = false;
 	}
 }
diff --git a/Engine/Plugins/VirtualProduction/EpicStageApp/Source/EpicStageApp/Private/EpicStageApp.h b/Engine/Plugins/VirtualProduction/EpicStageApp/Source/EpicStageApp/Private/EpicStageApp.h
--- a/Engine/Plugins/VirtualProduction/EpicStageApp/Source/EpicStageApp/Private/EpicStageApp.h
+++ b/Engine/Plugins/VirtualProduction/EpicStageApp/Source/EpicStageApp/Private/EpicStageApp.h
@@ -21,6 +21,20 @@ class FEpicStageAppModule : public IModuleInterface
 
 	/** Receiver that responds to beacon messages from the stage app. */
 	FStageAppBeaconReceiver StageAppBeaconReceiver;
+
+private:
+	/** Start the beacon receiver so the engine can be seen by the app. */
+	void StartupBeaconReceiver();
+
+	/** Shut down the beacon receiver so the engine can no longer be seen by the app. */
+	void ShutdownBeaconReceiver();
+
+private:
+	/** Whether the beacon receiver has been started and has not been stopped. */
+	bool bIsBeaconReceiverRunning = false;
+
+	/** Delegate for when the websocket server starts. */
+	FDelegateHandle WebSocketServerStartedDelegate;
 };
 
 IMPLEMENT_MODULE(FEpicStageAppModule, EpicStageApp);
\ No newline at end of file
diff --git a/Engine/Plugins/VirtualProduction/RemoteControl/Source/WebRemoteControl/Private/WebRemoteControl.cpp b/Engine/Plugins/VirtualProduction/RemoteControl/Source/WebRemoteControl/Private/WebRemoteControl.cpp
--- a/Engine/Plugins/VirtualProduction/RemoteControl/Source/WebRemoteControl/Private/WebRemoteControl.cpp
+++ b/Engine/Plugins/VirtualProduction/RemoteControl/Source/WebRemoteControl/Private/WebRemoteControl.cpp
@@ -469,6 +469,7 @@ void FWebRemoteControlModule::StartHttpServer()
 
 		FHttpServerModule::Get().StartAllListeners();
 
+		bIsHttpServerRunning = true;
 		OnHttpServerStartedDelegate.Broadcast(HttpServerPort);
 	}
 }
@@ -494,6 +495,8 @@ void FWebRemoteControlModule::StopHttpServer()
 	}
 
 	HttpRouter.Reset();
+	bIsHttpServerRunning = false;
+
 	OnHttpServerStoppedDelegate.Broadcast();
 }
 
diff --git a/Engine/Plugins/VirtualProduction/RemoteControl/Source/WebRemoteControl/Private/WebRemoteControl.h b/Engine/Plugins/VirtualProduction/RemoteControl/Source/WebRemoteControl/Private/WebRemoteControl.h
--- a/Engine/Plugins/VirtualProduction/RemoteControl/Source/WebRemoteControl/Private/WebRemoteControl.h
+++ b/Engine/Plugins/VirtualProduction/RemoteControl/Source/WebRemoteControl/Private/WebRemoteControl.h
@@ -45,8 +45,10 @@ public:
 	virtual void UnregisterRequestPreprocessor(const FDelegateHandle& RequestPreprocessorHandle) override;
 	virtual FOnWebServerStarted& OnHttpServerStarted() override { return OnHttpServerStartedDelegate; }
 	virtual FSimpleMulticastDelegate& OnHttpServerStopped() override { return OnHttpServerStoppedDelegate; }
+	virtual bool IsHttpServerRunning() override { return bIsHttpServerRunning; }
 	virtual FOnWebServerStarted& OnWebSocketServerStarted() override { return OnWebSocketServerStartedDelegate; }
 	virtual FSimpleMulticastDelegate& OnWebSocketServerStopped() override { return OnWebSocketServerStoppedDelegate; }
+	virtual bool IsWebSocketServerRunning() override { return WebSocketServer.IsRunning(); }
 	virtual FOnWebSocketConnectionClosed& OnWebSocketConnectionOpened() override { return WebSocketServer.OnConnectionOpened(); }
 	virtual FOnWebSocketConnectionClosed& OnWebSocketConnectionClosed() override { return WebSocketServer.OnConnectionClosed(); }
 	virtual void RegisterWebsocketRoute(const FRemoteControlWebsocketRoute& Route) override;
@@ -185,6 +187,9 @@ private:
 	/** Holds the client currently making a request. */
 	FGuid ActingClientId;
 
+	/** Whether the HTTP server has been started and has not been stopped. */
+	bool bIsHttpServerRunning = false;
+
 	/** List of preprocessor delegates that need to be registered when the server is started. */
 	TMap<FDelegateHandle, FHttpRequestHandler> PreprocessorsToRegister;
 
diff --git a/Engine/Plugins/VirtualProduction/RemoteControl/Source/WebRemoteControl/Public/IWebRemoteControlModule.h b/Engine/Plugins/VirtualProduction/RemoteControl/Source/WebRemoteControl/Public/IWebRemoteControlModule.h
--- a/Engine/Plugins/VirtualProduction/RemoteControl/Source/WebRemoteControl/Public/IWebRemoteControlModule.h
+++ b/Engine/Plugins/VirtualProduction/RemoteControl/Source/WebRemoteControl/Public/IWebRemoteControlModule.h
@@ -42,6 +42,11 @@ public:
 	 * Event triggered when the http server stops.
 	 */
 	virtual FSimpleMulticastDelegate& OnHttpServerStopped() = 0;
+
+	/**
+	 * Returns whether the http server is currently running.
+	 */
+	virtual bool IsHttpServerRunning() = 0;
 	
 	
 	/** 
@@ -54,6 +59,11 @@ public:
 	 */
 	virtual FSimpleMulticastDelegate& OnWebSocketServerStopped() = 0;
 
+	/**
+	 * Returns whether the websocket server is currently running.
+	 */
+	virtual bool IsWebSocketServerRunning() = 0;
+
 	/**
 	 * Event triggered when a connection to the websocket server is opened.
 	 */
