[EpicStageApp] Add function to remotely change color grading values based on a wheel position

diff --git a/Engine/Plugins/VirtualProduction/RemoteControl/Source/RemoteControl/Private/RemoteControlFunctionLibrary.cpp b/Engine/Plugins/VirtualProduction/RemoteControl/Source/RemoteControl/Private/RemoteControlFunctionLibrary.cpp
--- a/Engine/Plugins/VirtualProduction/RemoteControl/Source/RemoteControl/Private/RemoteControlFunctionLibrary.cpp
+++ b/Engine/Plugins/VirtualProduction/RemoteControl/Source/RemoteControl/Private/RemoteControlFunctionLibrary.cpp
@@ -24,6 +24,123 @@ namespace RemoteControlFunctionLibrary
 		}
 		return GroupId;
 	}
+
+	/**
+	 * Traverse the nested struct properties on an object to find the property with the given period-separated path.
+	 * If found, return true and set OutProperty to the property matching the name and OutContainer to the object containing its value (within TargetObject).
+	 */
+	bool GetNestedPropertyAndContainerFromPath(UObject* Object, const FString& PropertyName, FProperty*& OutProperty, void*& OutContainer)
+	{
+		if (!Object)
+		{
+			return false;
+		}
+
+		// Split the property name up on dots so we can find nested struct properties
+		TArray<FString> PropertyPathNames;
+		PropertyName.ParseIntoArray(PropertyPathNames, TEXT("."));
+
+		if (PropertyPathNames.IsEmpty())
+		{
+			return false;
+		}
+
+		void* Container = Object;
+		UStruct* ContainerClass = Object->GetClass();
+		if (!ContainerClass)
+		{
+			return false;
+		}
+
+		FProperty* Property = ContainerClass->FindPropertyByName(FName(*PropertyPathNames[0]));
+		if (!Property)
+		{
+			return false;
+		}
+
+		// Walk through the chain of structs until we reach the end (or run out of matching struct properties)
+		for (int32 PathIndex = 1; PathIndex < PropertyPathNames.Num(); ++PathIndex)
+		{
+			if (const FStructProperty* StructProperty = CastField<FStructProperty>(Property))
+			{
+				ContainerClass = StructProperty->Struct;
+				if (ContainerClass == nullptr)
+				{
+					return false;
+				}
+
+				Property = ContainerClass->FindPropertyByName(FName(*PropertyPathNames[PathIndex]));
+				if (Property == nullptr)
+				{
+					return false;
+				}
+
+				Container = StructProperty->ContainerPtrToValuePtr<void>(Container);
+			}
+			else
+			{
+				return false;
+			}
+		}
+
+		ensure(Container && Property);
+
+		OutProperty = Property;
+		OutContainer = Container;
+		return true;
+	}
+
+	/**
+	 * Apply a wheel-based color delta to an RGB linear color.
+	 */
+	void ApplyWheelColorBaseDelta(FLinearColor& InOutColor, const FColorWheelColorBase& DeltaValue, const FColorWheelColorBase& ReferenceColor, float MaxValue)
+	{
+		// Convert to HSV
+		InOutColor = InOutColor.LinearRGBToHSV();
+
+		FVector2D Position;
+
+		if (InOutColor.B > UE_DOUBLE_KINDA_SMALL_NUMBER)
+		{
+			// Determine direction as a unit vector based on the calculated hue
+			const double HueRadians = FMath::DegreesToRadians(InOutColor.R);
+			Position = FVector2D(FMath::Cos(HueRadians), FMath::Sin(HueRadians));
+
+			// Multiply the unit vector by saturation to determine the current position in the color wheel
+			Position *= InOutColor.G;
+		}
+		else
+		{
+			// Color's value is too low to determine the hue and saturation. Fall back to the reference color's position
+			Position = ReferenceColor.Position;
+		}
+
+		// Apply the delta to the position, then convert back to hue and saturation
+		Position += FVector2D(DeltaValue.Position.X, DeltaValue.Position.Y);
+		InOutColor.R = FMath::Fmod(FMath::RadiansToDegrees(FMath::Atan2(Position.Y, Position.X)) + 360.0, 360.0);
+		InOutColor.G = FMath::Clamp(Position.Length(), 0.0, MaxValue);
+
+		// Apply the value change directly
+		InOutColor.B = FMath::Clamp(InOutColor.B + DeltaValue.Value, 0.0, MaxValue);
+
+		InOutColor = InOutColor.HSVToLinearRGB();
+	}
+
+	/** Set a property on an object and fire relevant events. */
+	void SetPropertyAndFireEvents(UObject* Object, FProperty* Property, void* Container, void* Value, bool bIsInteractive)
+	{
+#if WITH_EDITOR
+		Object->PreEditChange(Property);
+		Object->Modify();
+#endif
+
+		Property->SetValue_InContainer(Container, Value);
+
+#if WITH_EDITOR
+		FPropertyChangedEvent ChangeEvent(Property, bIsInteractive ? EPropertyChangeType::Interactive : EPropertyChangeType::ValueSet);
+		Object->PostEditChangeProperty(ChangeEvent);
+#endif
+	}
 }
 
 bool URemoteControlFunctionLibrary::ExposeProperty(URemoteControlPreset* Preset, UObject* SourceObject, const FString& Property, FRemoteControlOptionalExposeArgs Args)
@@ -58,107 +175,66 @@ bool URemoteControlFunctionLibrary::ExposeActor(URemoteControlPreset* Preset, AA
 
 bool URemoteControlFunctionLibrary::ApplyColorWheelDelta(UObject* TargetObject, const FString& PropertyName, const FColorWheelColor& DeltaValue, const FColorWheelColor& ReferenceColor, bool bIsInteractive)
 {
-	if (!TargetObject)
+	FProperty* Property;
+	void* Container;
+	if (!RemoteControlFunctionLibrary::GetNestedPropertyAndContainerFromPath(TargetObject, PropertyName, Property, Container))
 	{
 		return false;
 	}
 
-	// Split the property name up on dots so we can find nested struct properties
-	TArray<FString> PropertyPathNames;
-	PropertyName.ParseIntoArray(PropertyPathNames, TEXT("."));
-
-	if (PropertyPathNames.IsEmpty())
+	if (const FStructProperty* ColorProperty = CastField<FStructProperty>(Property))
 	{
-		return false;
-	}
+		if (ColorProperty->Struct != TBaseStructure<FLinearColor>::Get())
+		{
+			// This isn't a color
+			return false;
+		}
 
-	void* Container = TargetObject;
-	UStruct* ContainerClass = TargetObject->GetClass();
-	if (!ContainerClass)
-	{
-		return false;
-	}
+		FLinearColor Color;
+		ColorProperty->GetValue_InContainer(Container, &Color);
 
-	FProperty* Property = ContainerClass->FindPropertyByName(FName(*PropertyPathNames[0]));
-	if (!Property)
-	{
-		return false;
-	}
+		RemoteControlFunctionLibrary::ApplyWheelColorBaseDelta(Color, DeltaValue, ReferenceColor, 1.0);
 
-	// Walk through the chain of structs until we reach the end (or run out of matching struct properties)
-	for (int32 PathIndex = 1; PathIndex < PropertyPathNames.Num(); ++PathIndex)
-	{
-		if (const FStructProperty* StructProperty = CastField<FStructProperty>(Property))
-		{
-			ContainerClass = StructProperty->Struct;
-			if (ContainerClass == nullptr)
-			{
-				return false;
-			}
+		// Apply the alpha change directly
+		Color.A = FMath::Clamp(Color.A + DeltaValue.Alpha, 0.0, 1.0);
 
-			Property = ContainerClass->FindPropertyByName(FName(*PropertyPathNames[PathIndex]));
-			if (Property == nullptr)
-			{
-				return false;
-			}
+		RemoteControlFunctionLibrary::SetPropertyAndFireEvents(TargetObject, Property, Container, &Color, bIsInteractive);
 
-			Container = StructProperty->ContainerPtrToValuePtr<void>(Container);
-		}
-		else
-		{
-			return false;
-		}
+		return true;
 	}
 
-	ensure(Container && Property);
+	return false;
+}
 
-	if (const FStructProperty* ColorProperty = CastField<FStructProperty>(Property))
+bool URemoteControlFunctionLibrary::ApplyColorGradingWheelDelta(UObject* TargetObject, const FString& PropertyName, const FColorGradingWheelColor& DeltaValue, const FColorGradingWheelColor& ReferenceColor, bool bIsInteractive)
+{
+	FProperty* Property = nullptr;
+	void* Container = nullptr;
+	if (!RemoteControlFunctionLibrary::GetNestedPropertyAndContainerFromPath(TargetObject, PropertyName, Property, Container))
 	{
-		FLinearColor Color;
-		ColorProperty->GetValue_InContainer(Container, &Color);
-
-		// Convert to HSV
-		Color = Color.LinearRGBToHSV();
-
-		FVector2D Position;
+		return false;
+	}
 
-		if (Color.B > UE_DOUBLE_KINDA_SMALL_NUMBER)
-		{
-			// Determine direction as a unit vector based on the calculated hue
-			const double HueRadians = FMath::DegreesToRadians(Color.R);
-			Position = FVector2D(FMath::Cos(HueRadians), FMath::Sin(HueRadians));
-			
-			// Multiply the unit vector by saturation to determine the current position in the color wheel
-			Position *= Color.G;
-		}
-		else
+	if (const FStructProperty* ColorProperty = CastField<FStructProperty>(Property))
+	{
+		if (ColorProperty->Struct != TBaseStructure<FVector4>::Get())
 		{
-			// Color's value is too low to determine the hue and saturation. Fall back to the reference color's position
-			Position = ReferenceColor.Position;
+			// This isn't an FVector4 representation of a color
+			return false;
 		}
 
-		// Apply the delta to the position, then convert back to hue and saturation
-		Position += FVector2D(DeltaValue.Position.X, DeltaValue.Position.Y);
-		Color.R = FMath::Fmod(FMath::RadiansToDegrees(FMath::Atan2(Position.Y, Position.X)) + 360.0, 360.0);
-		Color.G = FMath::Clamp(Position.Length(), 0.0, 1.0);
+		// Get the vector value and convert the color portion of it to HSV
+		FVector4 VectorValue;
+		ColorProperty->GetValue_InContainer(Container, &VectorValue);
 
-		// Apply the value and alpha changes directly
-		Color.B = FMath::Clamp(Color.B + DeltaValue.Value, 0.0, 1.0);
-		Color.A = FMath::Clamp(Color.A + DeltaValue.Alpha, 0.0, 1.0);
+		FLinearColor Color(VectorValue.X, VectorValue.Y, VectorValue.Z);
+		RemoteControlFunctionLibrary::ApplyWheelColorBaseDelta(Color, DeltaValue, ReferenceColor, 2.0);
 
-		Color = Color.HSVToLinearRGB();
+		// Apply the luminance change directly
+		const double NewLuminance = VectorValue.W + DeltaValue.Luminance;
+		VectorValue = FVector4(Color, NewLuminance);
 
-#if WITH_EDITOR
-		TargetObject->PreEditChange(Property);
-		TargetObject->Modify();
-#endif
-
-		ColorProperty->SetValue_InContainer(Container, &Color);
-
-#if WITH_EDITOR
-		FPropertyChangedEvent ChangeEvent(Property, bIsInteractive ? EPropertyChangeType::Interactive : EPropertyChangeType::ValueSet);
-		TargetObject->PostEditChangeProperty(ChangeEvent);
-#endif
+		RemoteControlFunctionLibrary::SetPropertyAndFireEvents(TargetObject, Property, Container, &VectorValue, bIsInteractive);
 
 		return true;
 	}
diff --git a/Engine/Plugins/VirtualProduction/RemoteControl/Source/RemoteControl/Public/RemoteControlFunctionLibrary.h b/Engine/Plugins/VirtualProduction/RemoteControl/Source/RemoteControl/Public/RemoteControlFunctionLibrary.h
--- a/Engine/Plugins/VirtualProduction/RemoteControl/Source/RemoteControl/Public/RemoteControlFunctionLibrary.h
+++ b/Engine/Plugins/VirtualProduction/RemoteControl/Source/RemoteControl/Public/RemoteControlFunctionLibrary.h
@@ -29,15 +29,15 @@ struct FRemoteControlOptionalExposeArgs
 	FString GroupName;
 };
 
-/** A color as represented by a position on a color wheel. */
+/** Base data for a color represented as its position on a color wheel. */
 USTRUCT(Blueprintable)
-struct FColorWheelColor
+struct FColorWheelColorBase
 {
 	GENERATED_BODY()
 
 	/**
-	 * The position on the unit circle. Magnitude should be in range [0, 1].
-	 */
+		* The position on the unit circle. Magnitude should be in range [0, 1].
+		*/
 	UPROPERTY()
 	FVector2D Position = FVector2D::ZeroVector;
 
@@ -46,6 +46,13 @@ struct FColorWheelColor
 	 */
 	UPROPERTY()
 	double Value = 0.0;
+};
+
+/** A color as represented by a position on a color wheel. */
+USTRUCT(Blueprintable)
+struct FColorWheelColor : public FColorWheelColorBase
+{
+	GENERATED_BODY()
 
 	/**
 	 * The color's alpha component.
@@ -54,6 +61,19 @@ struct FColorWheelColor
 	double Alpha = 0.0;
 };
 
+/** A color value as represented by a position on a color grading wheel. */
+USTRUCT(Blueprintable)
+struct FColorGradingWheelColor : public FColorWheelColorBase
+{
+	GENERATED_BODY()
+
+	/**
+	 * The color's luminance component, which is separate from the HSV color itself.
+	 */
+	UPROPERTY()
+	double Luminance = 1.0;
+};
+
 UCLASS()
 class URemoteControlFunctionLibrary : public UBlueprintFunctionLibrary
 {
@@ -103,4 +123,16 @@ public:
 	 */
 	UFUNCTION(BlueprintCallable, Category = "Remote Control Color")
 	static bool ApplyColorWheelDelta(UObject* TargetObject, const FString& PropertyName, const FColorWheelColor& DeltaValue, const FColorWheelColor& ReferenceColor, bool bIsInteractive);
+
+	/**
+	 * Add/subtract from the value of an FVector4 property interpreted as RGBV using a delta value based on color wheel coordinates.
+	 * @param TargetObject the object that contains the property to modify.
+	 * @param PropertyName the name of the property to modify.
+	 * @param DeltaValue the amount to change the color by.
+	 * @param ReferenceColor if the color's current position on the wheel is ambiguous as calculated from RGB values (e.g. black), use this reference color's position instead.
+	 * @param bIsInteractive if true, this is treated as an interactive change. If false, it will be treated as the final value set change.
+	 * @return Whether the operation was successful.
+	 */
+	UFUNCTION(BlueprintCallable, Category = "Remote Control Color")
+	static bool ApplyColorGradingWheelDelta(UObject* TargetObject, const FString& PropertyName, const FColorGradingWheelColor& DeltaValue, const FColorGradingWheelColor& ReferenceColor, bool bIsInteractive);
 };
\ No newline at end of file
