Auto bind functoinality:

* Allow drag-drop from Exposed Properties to Conftrollers
* Automatically create new Controller of correct type and apply Bind behaviour to it.
* Also allow drag-drop on to existing Controllers with auto bind upon drop.

diff --git a/Engine/Plugins/VirtualProduction/RemoteControl/Source/RemoteControlLogic/Private/Behaviour/Builtin/Bind/RCBehaviourBind.cpp b/Engine/Plugins/VirtualProduction/RemoteControl/Source/RemoteControlLogic/Private/Behaviour/Builtin/Bind/RCBehaviourBind.cpp
--- a/Engine/Plugins/VirtualProduction/RemoteControl/Source/RemoteControlLogic/Private/Behaviour/Builtin/Bind/RCBehaviourBind.cpp
+++ b/Engine/Plugins/VirtualProduction/RemoteControl/Source/RemoteControlLogic/Private/Behaviour/Builtin/Bind/RCBehaviourBind.cpp
@@ -7,6 +7,7 @@
 #include "Action/Bind/RCPropertyBindAction.h"
 #include "Behaviour/Builtin/Bind/RCBehaviourBindNode.h"
 #include "Controller/RCController.h"
+#include "IRemoteControlPropertyHandle.h"
 #include "RemoteControlField.h"
 #include "RCVirtualProperty.h"
 
@@ -57,6 +58,11 @@ URCAction* URCBehaviourBind::AddAction(const TSharedRef<const FRemoteControlFiel
 
 bool URCBehaviourBind::CanHaveActionForField(const TSharedPtr<FRemoteControlField> InRemoteControlField) const
 {
+	if (!InRemoteControlField.IsValid())
+	{
+		return false;
+	}
+
 	// Basic check (uniqueness)
 	if (ActionContainer->FindActionByFieldId(InRemoteControlField->GetId()))
 	{
@@ -64,98 +70,259 @@ bool URCBehaviourBind::CanHaveActionForField(const TSharedPtr<FRemoteControlFiel
 	}
 
 	// Advanced checks (by Controller type and Target type)
-	if(TSharedPtr<FRemoteControlProperty> RemoteControlEntityAsProperty = StaticCastSharedPtr<FRemoteControlProperty>(InRemoteControlField))
+	if (TSharedPtr<FRemoteControlProperty> RemoteControlEntityAsProperty = StaticCastSharedPtr<FRemoteControlProperty>(InRemoteControlField))
 	{
 		if (URCController* Controller = ControllerWeakPtr.Get())
 		{
-			const FProperty* ControllerAsProperty = Controller->GetProperty();
-			const FProperty* RemoteControlProperty = RemoteControlEntityAsProperty->GetProperty();
-			const FStructProperty* RemoteControlEntityAsStructProperty = CastField<FStructProperty>(RemoteControlProperty);
-
-			bool IsNumericStruct = RemoteControlEntityAsStructProperty
-				&& (RemoteControlEntityAsStructProperty->Struct == TBaseStructure<FVector>::Get() ||
-					RemoteControlEntityAsStructProperty->Struct == TBaseStructure<FRotator>::Get() ||
-					RemoteControlEntityAsStructProperty->Struct == TBaseStructure<FColor>::Get() ||
-					RemoteControlEntityAsStructProperty->Struct == TBaseStructure<FLinearColor>::Get());
-
-			// String Controller
-			if (ControllerAsProperty->IsA(FStrProperty::StaticClass()))
+			return URCBehaviourBind::CanHaveActionForField(Controller, RemoteControlEntityAsProperty.ToSharedRef(), bAllowNumericInputAsStrings);
+		}
+	}
+
+	return false;
+}
+
+static bool EvaluateBindCompatibility(const TMap<FFieldClass*, TArray<FFieldClass*>>& BindingsMap, const FFieldClass* ControllerPropertyClass, const FFieldClass* InRemoteControlPropertyClass)
+{
+	TArray< FFieldClass*> Keys;
+	BindingsMap.GenerateKeyArray(Keys);
+
+	for (const FFieldClass* PropertyType : Keys)
+	{
+		if (ControllerPropertyClass->IsChildOf(PropertyType))
+		{
+			const TArray<FFieldClass*>* SupportedBindings = BindingsMap.Find(PropertyType);
+
+			if (SupportedBindings)
 			{
-				if(RemoteControlProperty->IsA(FStrProperty::StaticClass()) ||
-					RemoteControlProperty->IsA(FTextProperty::StaticClass()) ||
-					RemoteControlProperty->IsA(FNameProperty::StaticClass()))
+				for (const FFieldClass* SupportedBinding : *SupportedBindings)
 				{
-					return true; // Implicit / direct conversions are possible between these types
-				}
-				else if (RemoteControlProperty->IsA(FNumericProperty::StaticClass()) ||
-					RemoteControlProperty->IsA(FBoolProperty::StaticClass()) ||
-					RemoteControlProperty->IsA(FByteProperty::StaticClass()))
-				{
-					// Explicit conversion to numbers, governed by user set flag
-					return bAllowNumericInputAsStrings;
+					if (InRemoteControlPropertyClass->IsChildOf(SupportedBinding))
+					{
+						return true;
+					}
 				}
 			}
-			// Numeric Controller (Float/Int)
-			else if (ControllerAsProperty->IsA(FNumericProperty::StaticClass()))
-			{
-				if (RemoteControlProperty->IsA(FNumericProperty::StaticClass()) ||
-					RemoteControlProperty->IsA(FBoolProperty::StaticClass()) || 
-					RemoteControlProperty->IsA(FByteProperty::StaticClass()))
-				{
-					return true; // Numeric to Numeric - Implicit / direct conversion
+		}
+	}
 
-					// Note: Numeric to Numeric struct conversion is not currently supported as it would ideally require Masking support (on the struct) for meaningful usage
-				}
-				else if (RemoteControlProperty->IsA(FStrProperty::StaticClass()) ||
-					RemoteControlProperty->IsA(FTextProperty::StaticClass()) ||
-					RemoteControlProperty->IsA(FNameProperty::StaticClass()))
-				{
-					// Explicit conversion to String/Name/Text, governed by user set flag
-					return bAllowNumericInputAsStrings;
-				}
-			}
-			// Boolean Controller
-			else if (ControllerAsProperty->IsA(FBoolProperty::StaticClass()))
+	return false;
+}
+
+bool URCBehaviourBind::CanHaveActionForField(URCController* Controller, TSharedRef<const FRemoteControlField> InRemoteControlField, const bool bInAllowNumericInputAsStrings)
+{
+	// Bind behaviour is only relevant for Exposed Properties
+	if (InRemoteControlField->FieldType != EExposedFieldType::Property)
+	{
+		return false;
+	}
+
+	TSharedRef<const FRemoteControlProperty> RemoteControlEntityAsProperty = StaticCastSharedRef<const FRemoteControlProperty>(InRemoteControlField);
+
+	const FProperty* ControllerAsProperty = Controller->GetProperty();
+	const FProperty* RemoteControlProperty = RemoteControlEntityAsProperty->GetProperty();
+	const FStructProperty* RemoteControlEntityAsStructProperty = CastField<FStructProperty>(RemoteControlProperty);
+
+	// Indirect Binding (related types)
+	//
+	const static TMap<FFieldClass*, TArray<FFieldClass*>> SupportedIndirectBindsMap =
+	{
+		/* Controller Type */                           /* Supported Remote Control Property Types */
+
+		{ FStrProperty::StaticClass(),      /* --> */   { FTextProperty::StaticClass(),     FNameProperty::StaticClass()}},
+		{ FNumericProperty::StaticClass(),  /* --> */   { FNumericProperty::StaticClass(),  FBoolProperty::StaticClass(),  FByteProperty::StaticClass() } },
+		{ FBoolProperty::StaticClass(),     /* --> */   { FFloatProperty::StaticClass(),    FIntProperty::StaticClass(),   FBoolProperty::StaticClass() } }
+	};
+
+	// Indirect Binding (via numeric conversion)
+	//
+	const static TMap<FFieldClass*, TArray<FFieldClass*>> SupportedNumericConversionsMap =
+	{
+		/* Controller Type */                           /* Supported Remote Control Property Types */
+
+		{ FStrProperty::StaticClass(),      /* --> */   { FNumericProperty::StaticClass(), FBoolProperty::StaticClass(),  FByteProperty::StaticClass()  } },
+		{ FNumericProperty::StaticClass(),  /* --> */   { FStrProperty::StaticClass(),     FTextProperty::StaticClass(),  FNameProperty::StaticClass() } }
+	};
+
+	// Indirect Binding (for Structs)
+	//
+	const static TMap<UScriptStruct*, TArray<UScriptStruct*>> SupportedStructConversions =
+	{
+		/* Struct Type */                                 /* Supported Struct Types */
+
+		{ TBaseStructure<FColor>::Get(),      /* --> */   { TBaseStructure<FLinearColor>::Get()  } }
+	};
+
+	const FFieldClass* ControllerPropertyClass = ControllerAsProperty->GetClass();
+	const FFieldClass* RemoteControlPropertyClass = RemoteControlProperty->GetClass();
+	const bool bIsStructController = ControllerAsProperty->IsA(FStructProperty::StaticClass());
+
+	// Direct Binding (same type)
+	if (ControllerPropertyClass == RemoteControlPropertyClass && !bIsStructController) // For Structs we need to check the inner component as well
+	{
+		return true;
+	}
+	// Indirect Binding (related types)
+	else if(EvaluateBindCompatibility(SupportedIndirectBindsMap, ControllerPropertyClass, RemoteControlPropertyClass))
+	{
+		return true;
+	}
+	// Indirect Binding (Numeric Conversion)
+	else if (bInAllowNumericInputAsStrings)
+	{
+		if (EvaluateBindCompatibility(SupportedNumericConversionsMap, ControllerPropertyClass, RemoteControlPropertyClass))
+		{
+			return true;
+		}
+	}
+
+	// Structs:
+	if (bIsStructController)
+	{
+		if (const FStructProperty* StructProperty = CastField<FStructProperty>(ControllerAsProperty))
+		{
+			if (const FStructProperty* RCFieldStructProperty = CastField<FStructProperty>(RemoteControlProperty))
 			{
-				if (RemoteControlProperty->IsA(FFloatProperty::StaticClass()) || // Note: FNumericPropery includes Enum, so we check explicitly for Float and Int here
-					RemoteControlProperty->IsA(FIntProperty::StaticClass()) ||
-					RemoteControlProperty->IsA(FBoolProperty::StaticClass()))
+				if (RCFieldStructProperty->Struct == StructProperty->Struct)
 				{
-					return true;
+					return true; // Binding via matching Struct
 				}
-				// Notes for explicit Bool conversions: 
-				// 1. Boolean to String/Text/Name is currently not deemed as a requirement, so hasn't been implemented here
-				// 2. Boolean to Struct is also not implemented. 
-			}
-			// Struct
-			else if (ControllerAsProperty->IsA(FStructProperty::StaticClass()))
-			{
-				if (const FStructProperty* StructProperty = CastField<FStructProperty>(ControllerAsProperty))
+				else
 				{
-					// FVector / FColor / FRotator
-					if (const FStructProperty* RCFieldStructProperty = CastField<FStructProperty>(RemoteControlProperty))
+					if (const TArray<UScriptStruct*>* SupportedConversions = SupportedStructConversions.Find(StructProperty->Struct))
 					{
-						// Generally - A Struct Controller can only drive properties of the _same_ Struct type...
-						if (RCFieldStructProperty->Struct == StructProperty->Struct)
+						if (SupportedConversions->Contains(RCFieldStructProperty->Struct))
 						{
-							return true;
+							return true; // Bind via explicit conversion
 						}
-						// ... with an exception for the following Struct conversions:
-						else
-						{
-							// FColor conversion:
-							if (StructProperty->Struct == TBaseStructure<FColor>::Get())
-							{
-								// To FLinearColor, is supported.
-								return RCFieldStructProperty->Struct == TBaseStructure<FLinearColor>::Get();
-							}
-						}
-
 					}
 				}
 			}
 		}
 	}
 
-	return false;
+	return false; // Unsupported types!
 }
+
+bool URCBehaviourBind::GetPropertyBagTypeFromFieldProperty(const FProperty* InProperty, EPropertyBagPropertyType& OutPropertyBagType, UObject*& OutStructObject)
+{
+	const static TMap<EPropertyBagPropertyType, TArray<FFieldClass*>> PropertyBagTypesMap =
+	{
+		/* Property Bag Type */                           /* Input - Remote Control Property Type */
+		{ EPropertyBagPropertyType::String,  /* --> */    { FStrProperty::StaticClass(),      FTextProperty::StaticClass(),     FNameProperty::StaticClass()}},
+		{ EPropertyBagPropertyType::Int32,   /* --> */    { FIntProperty::StaticClass(),      FInt64Property::StaticClass(),    FInt16Property::StaticClass(),   FUInt32Property::StaticClass(),
+														   FUInt64Property::StaticClass(),   FUInt16Property::StaticClass(),   FEnumProperty::StaticClass(),    FByteProperty::StaticClass() } },
+		{ EPropertyBagPropertyType::Float,   /* --> */    { FFloatProperty::StaticClass(),    FDoubleProperty::StaticClass() } },
+		{ EPropertyBagPropertyType::Bool,    /* --> */    { FBoolProperty::StaticClass() } },
+		{ EPropertyBagPropertyType::Struct,  /* --> */    { FStructProperty::StaticClass() } }
+	};
+
+	OutPropertyBagType = EPropertyBagPropertyType::None;
+
+	TArray<EPropertyBagPropertyType> Keys;
+	PropertyBagTypesMap.GenerateKeyArray(Keys);
+
+	for (const EPropertyBagPropertyType PropertyBagType : Keys)
+	{
+		if (PropertyBagTypesMap.Find(PropertyBagType)->Contains(InProperty->GetClass()))
+		{
+			OutPropertyBagType = PropertyBagType;
+		}
+	}
+
+	// Extract the inner Struct object if this is a Struct Property
+	if (const FStructProperty* StructProperty = CastField<FStructProperty>(InProperty))
+	{
+		OutStructObject = StructProperty->Struct;
+	}
+
+	return OutPropertyBagType != EPropertyBagPropertyType::None;
+}
+
+bool URCBehaviourBind::CopyPropertyValueToController(URCController* InController, TSharedRef<const FRemoteControlProperty> InRemoteControlProperty)
+{
+	FProperty* ControllerAsProperty = InController->GetProperty();
+	FProperty* RemoteControlProperty = InRemoteControlProperty->GetProperty();
+	if (!ensure(ControllerAsProperty && RemoteControlProperty))
+	{
+		return false;
+	}
+
+	TSharedPtr<IRemoteControlPropertyHandle> PropertyHandle = InRemoteControlProperty->GetPropertyHandle();
+	if (!ensure(PropertyHandle))
+	{
+		return false;
+	}
+
+	// String Controller
+	if (ControllerAsProperty->IsA(FStrProperty::StaticClass()))
+	{
+		FString StringValue;
+		if (RemoteControlProperty->IsA(FStrProperty::StaticClass()) || RemoteControlProperty->IsA(FNameProperty::StaticClass()))
+		{
+			PropertyHandle->GetValue(StringValue);
+		}
+		else if (RemoteControlProperty->IsA(FTextProperty::StaticClass()))
+		{
+			FText TextValue;
+			PropertyHandle->GetValue(TextValue);
+			StringValue = TextValue.ToString();
+		}
+
+		InController->SetValueString(StringValue);
+
+		return true;
+	}
+	// Int32 Controller
+	else if (ControllerAsProperty->IsA(FIntProperty::StaticClass()))
+	{
+		int32 IntValue;
+		PropertyHandle->GetValue(IntValue);
+
+		InController->SetValueDouble(IntValue);
+
+		return true;
+	}
+	// Float Controller
+	else if (ControllerAsProperty->IsA(FFloatProperty::StaticClass()))
+	{
+		float FloatValue;
+		PropertyHandle->GetValue(FloatValue);
+
+		InController->SetValueFloat(FloatValue);
+
+		return true;
+	}
+	// Bool Controller
+	else if (ControllerAsProperty->IsA(FBoolProperty::StaticClass()))
+	{
+		bool BoolValue;
+		PropertyHandle->GetValue(BoolValue);
+
+		InController->SetValueBool(BoolValue);
+
+		return true;
+	}
+	else if (FStructProperty* ControllerAsStructProperty = CastField<FStructProperty>(ControllerAsProperty))
+	{
+		if (ControllerAsStructProperty->Struct == TBaseStructure<FVector>::Get())
+		{
+			FVector VectorValue;
+			PropertyHandle->GetValue(VectorValue);
+			InController->SetValueVector(VectorValue);
+		}
+		else if (ControllerAsStructProperty->Struct == TBaseStructure<FRotator>::Get())
+		{
+			FRotator RotatorValue;
+			PropertyHandle->GetValue(RotatorValue);
+			InController->SetValueRotator(RotatorValue);
+		}
+		else if (ControllerAsStructProperty->Struct == TBaseStructure<FColor>::Get())
+		{
+			FColor ColorValue;
+			PropertyHandle->GetValue(ColorValue);
+			InController->SetValueColor(ColorValue);
+		}
+	}
+
+	return false;
+}
\ No newline at end of file
diff --git a/Engine/Plugins/VirtualProduction/RemoteControl/Source/RemoteControlLogic/Public/Behaviour/Builtin/Bind/RCBehaviourBind.h b/Engine/Plugins/VirtualProduction/RemoteControl/Source/RemoteControlLogic/Public/Behaviour/Builtin/Bind/RCBehaviourBind.h
--- a/Engine/Plugins/VirtualProduction/RemoteControl/Source/RemoteControlLogic/Public/Behaviour/Builtin/Bind/RCBehaviourBind.h
+++ b/Engine/Plugins/VirtualProduction/RemoteControl/Source/RemoteControlLogic/Public/Behaviour/Builtin/Bind/RCBehaviourBind.h
@@ -7,8 +7,10 @@
 
 #include "RCBehaviourBind.generated.h"
 
+enum class EPropertyBagPropertyType : uint8;
 class URCVirtualPropertySelfContainer;
 class URCAction;
+class URCController;
 
 /**
  * [Bind Behaviour]
@@ -35,9 +37,23 @@ public:
 
 	//~ End URCBehaviour interface
 
+	/** Contains behaviour-independent validation logic for determining whether a Remote Controller property can be bound to a given Controller*/
+	static bool CanHaveActionForField(URCController* InController, TSharedRef<const FRemoteControlField> InRemoteControlField, const bool bInAllowNumericInputAsStrings);
+
+	/** Transfers the value of a given Remote Control Property to a Controller.
+	* Used for Auto Bind workflow where matching Controller types for bind are automatically created from a given Remote Control Property */
+	static bool CopyPropertyValueToController(URCController* InController, TSharedRef<const FRemoteControlProperty> InRemoteControlProperty);
+
 	/** Add an action specifically for Bind Behaviour */
 	URCPropertyBindAction* AddPropertyBindAction(const TSharedRef<const FRemoteControlProperty> InRemoteControlProperty);
 
+	/** GetPropertyBagTypeFromFieldProperty
+	*
+	* Given a Remote Control Exposed Property type this function determines the appropriate EPropertyBagType and struct object required (if any) 
+	* In the Logic realm we use a single type like (eg: String / Int) to represent various related types (String/Name/Text, Int32, Int64, etc)
+	* For this reason explicit mapping conversion is required between a given FProperty type and the desired Controller (Property Bag) type */
+	static bool GetPropertyBagTypeFromFieldProperty(const FProperty* InProperty, EPropertyBagPropertyType& OutPropertyBagType, UObject*& OutStructObject);
+
 	/** Indicates whether we support binding of strings to numeric fields and vice versa.
 	* This flag determines the list of bindable properties the user sees in the "Add Action" menu*/
 	bool AreNumericInputsAllowedAsStrings() const { return bAllowNumericInputAsStrings; }
diff --git a/Engine/Plugins/VirtualProduction/RemoteControl/Source/RemoteControlUI/Private/UI/Behaviour/SRCBehaviourPanelList.cpp b/Engine/Plugins/VirtualProduction/RemoteControl/Source/RemoteControlUI/Private/UI/Behaviour/SRCBehaviourPanelList.cpp
--- a/Engine/Plugins/VirtualProduction/RemoteControl/Source/RemoteControlUI/Private/UI/Behaviour/SRCBehaviourPanelList.cpp
+++ b/Engine/Plugins/VirtualProduction/RemoteControl/Source/RemoteControlUI/Private/UI/Behaviour/SRCBehaviourPanelList.cpp
@@ -181,7 +181,7 @@ void SRCBehaviourPanelList::OnTreeSelectionChanged(TSharedPtr<FRCBehaviourModel>
 	}
 }
 
-void SRCBehaviourPanelList::OnBehaviourAdded(URCBehaviour* InBehaviour)
+void SRCBehaviourPanelList::OnBehaviourAdded(const URCBehaviour* InBehaviour)
 {
 	Reset();
 }
diff --git a/Engine/Plugins/VirtualProduction/RemoteControl/Source/RemoteControlUI/Private/UI/Behaviour/SRCBehaviourPanelList.h b/Engine/Plugins/VirtualProduction/RemoteControl/Source/RemoteControlUI/Private/UI/Behaviour/SRCBehaviourPanelList.h
--- a/Engine/Plugins/VirtualProduction/RemoteControl/Source/RemoteControlUI/Private/UI/Behaviour/SRCBehaviourPanelList.h
+++ b/Engine/Plugins/VirtualProduction/RemoteControl/Source/RemoteControlUI/Private/UI/Behaviour/SRCBehaviourPanelList.h
@@ -75,7 +75,7 @@ private:
 	void OnTreeSelectionChanged(TSharedPtr<FRCBehaviourModel> InItem , ESelectInfo::Type);
 
 	/** Responds to the selection of a newly created Behaviour. Resets UI state */
-	void OnBehaviourAdded(URCBehaviour* InBehaviour);
+	void OnBehaviourAdded(const URCBehaviour* InBehaviour);
 
 	/** Responds to the removal of all Behaviours. Rests UI state */
 	void OnEmptyBehaviours();
diff --git a/Engine/Plugins/VirtualProduction/RemoteControl/Source/RemoteControlUI/Private/UI/Controller/SRCControllerPanelList.cpp b/Engine/Plugins/VirtualProduction/RemoteControl/Source/RemoteControlUI/Private/UI/Controller/SRCControllerPanelList.cpp
--- a/Engine/Plugins/VirtualProduction/RemoteControl/Source/RemoteControlUI/Private/UI/Controller/SRCControllerPanelList.cpp
+++ b/Engine/Plugins/VirtualProduction/RemoteControl/Source/RemoteControlUI/Private/UI/Controller/SRCControllerPanelList.cpp
@@ -2,14 +2,19 @@
 
 #include "SRCControllerPanelList.h"
 
+#include "Behaviour/Builtin/Bind/RCBehaviourBind.h"
+#include "Behaviour/Builtin/Bind/RCBehaviourBindNode.h"
 #include "Commands/RemoteControlCommands.h"
+#include "Controller/RCController.h"
 #include "Controller/RCControllerContainer.h"
 #include "IDetailTreeNode.h"
 #include "Interfaces/IMainFrameModule.h"
 #include "IPropertyRowGenerator.h"
+#include "IRemoteControlModule.h"
 #include "RCControllerModel.h"
 #include "RCVirtualProperty.h"
 #include "RCVirtualPropertyContainer.h"
+#include "RemoteControlField.h"
 #include "RemoteControlPreset.h"
 #include "SDropTarget.h"
 #include "SlateOptMacros.h"
@@ -31,6 +36,13 @@
 
 namespace UE::RCControllerPanelList
 {
+	enum class EDragDropSupportedModes : uint8
+	{
+		ReorderOnly,
+		All,
+		None
+	};
+
 	namespace Columns
 	{
 		const FName Name = TEXT("Controller Name");
@@ -74,7 +86,7 @@ namespace UE::RCControllerPanelList
 						.Widget(DragDropBorderWidget)
 					];
 
-				return WrapWithDropTarget(DragHandleWidget);
+				return WrapWithDropTarget(DragHandleWidget, EDragDropSupportedModes::ReorderOnly /*Allow reorder operation only*/);
 			}
 
 			return SNullWidget::NullWidget;
@@ -82,16 +94,31 @@ namespace UE::RCControllerPanelList
 
 		TSharedPtr<SBorder> DragDropBorderWidget;
 
+
+		void OnDragEnter(FGeometry const& MyGeometry, FDragDropEvent const& DragDropEvent) override
+		{
+			bIsDragActive = true;
+		}
+
+
+		void OnDragLeave(FDragDropEvent const& DragDropEvent) override
+		{
+			bIsDragActive = false;
+		}
+
+	protected:
+
 	private:
 
-		TSharedRef<SWidget> WrapWithDropTarget(const TSharedRef<SWidget> InWidget)
+		TSharedRef<SWidget> WrapWithDropTarget(const TSharedRef<SWidget> InWidget, const EDragDropSupportedModes SupportedMode = EDragDropSupportedModes::All)
 		{
 			return SNew(SDropTarget)
+				.ValidColor(FStyleColors::AccentOrange)
 				.VerticalImage(FRemoteControlPanelStyle::Get()->GetBrush("RemoteControlPanel.VerticalDash"))
 				.HorizontalImage(FRemoteControlPanelStyle::Get()->GetBrush("RemoteControlPanel.HorizontalDash"))
 				.OnDropped_Lambda([this](const FGeometry& InGeometry, const FDragDropEvent& InDragDropEvent) { return SControllerItemListRow::OnControllerItemDragDrop(InDragDropEvent.GetOperation()); })
-				.OnAllowDrop(this, &SControllerItemListRow::OnAllowDrop)
-				.OnIsRecognized(this, &SControllerItemListRow::OnAllowDrop)
+				.OnAllowDrop(this, &SControllerItemListRow::OnAllowDrop, SupportedMode)
+				.OnIsRecognized(this, &SControllerItemListRow::OnAllowDrop, SupportedMode)
 				[
 					InWidget
 				];
@@ -99,6 +126,9 @@ namespace UE::RCControllerPanelList
 
 		FReply OnControllerItemDragDrop(TSharedPtr<FDragDropOperation> DragDropOperation)
 		{
+			bIsDragActive = false;
+			ControllerPanelList->bIsAnyControllerItemEligibleForDragDrop = false;
+
 			if (!DragDropOperation)
 			{
 				return FReply::Handled();
@@ -121,12 +151,36 @@ namespace UE::RCControllerPanelList
 					}
 				}
 			}
+			else if (DragDropOperation->IsOfType<FExposedEntityDragDrop>())
+			{
+				if (TSharedPtr<FExposedEntityDragDrop> DragDropOp = StaticCastSharedPtr<FExposedEntityDragDrop>(DragDropOperation))
+				{
+					// Fetch the Exposed Entity
+					const FGuid ExposedEntityId = DragDropOp->GetId();
+
+					if (URemoteControlPreset* Preset = ControllerPanelList->GetPreset())
+					{
+						if (TSharedPtr<const FRemoteControlProperty> RemoteControlProperty = Preset->GetExposedEntity<FRemoteControlProperty>(ExposedEntityId).Pin())
+						{
+							if (URCController* Controller = Cast<URCController>(ControllerItem->GetVirtualProperty()))
+							{
+								ControllerPanelList->CreateBindBehaviourAndAssignTo(Controller, RemoteControlProperty.ToSharedRef(), true);
+							}
+						}
+					}
+				}
+			}
 
 			return FReply::Handled();
 		}
 
-		bool OnAllowDrop(TSharedPtr<FDragDropOperation> DragDropOperation)
+		bool OnAllowDrop(TSharedPtr<FDragDropOperation> DragDropOperation, const EDragDropSupportedModes SupportedMode)
 		{
+			if (!ensure(ControllerPanelList))
+			{
+				return false;
+			}
+
 			if (DragDropOperation && ControllerItem)
 			{
 				// Dragging Controllers onto Controllers (Reordering)
@@ -137,6 +191,40 @@ namespace UE::RCControllerPanelList
 						return true;
 					}
 				}
+				else
+				{
+					if (SupportedMode == EDragDropSupportedModes::ReorderOnly)
+					{
+						return false; // This widget only supports Reordering operation as a drag-drop action
+					}
+
+					if (DragDropOperation->IsOfType<FExposedEntityDragDrop>())
+					{
+						if (TSharedPtr<FExposedEntityDragDrop> DragDropOp = StaticCastSharedPtr<FExposedEntityDragDrop>(DragDropOperation))
+						{
+							// Fetch the Exposed Entity
+							const FGuid ExposedEntityId = DragDropOp->GetId();
+
+							if (URemoteControlPreset* Preset = ControllerPanelList->GetPreset())
+							{
+								if (TSharedPtr<const FRemoteControlField> RemoteControlField = Preset->GetExposedEntity<FRemoteControlField>(ExposedEntityId).Pin())
+								{
+									if (URCController* Controller = Cast<URCController>(ControllerItem->GetVirtualProperty()))
+									{
+										const bool bAllowNumericInputAsStrings = true;
+
+										const bool bAllowDrop = URCBehaviourBind::CanHaveActionForField(Controller, RemoteControlField.ToSharedRef(), bAllowNumericInputAsStrings);
+
+										ControllerPanelList->bIsAnyControllerItemEligibleForDragDrop |= (bAllowDrop && bIsDragActive);
+
+										return bAllowDrop;
+									}
+								}
+							}
+						}
+
+					}
+				}
 			}
 
 			return false;
@@ -156,6 +244,7 @@ namespace UE::RCControllerPanelList
 private:
 		TSharedPtr<FRCControllerModel> ControllerItem;
 		TSharedPtr<SRCControllerPanelList> ControllerPanelList;
+		bool bIsDragActive = false;
 	};
 } 
 
@@ -195,7 +284,15 @@ void SRCControllerPanelList::Construct(const FArguments& InArgs, const TSharedRe
 
 	ChildSlot
 	[
-		ListView.ToSharedRef()
+		SNew(SDropTarget)
+		.VerticalImage(FRemoteControlPanelStyle::Get()->GetBrush("RemoteControlPanel.VerticalDash"))
+		.HorizontalImage(FRemoteControlPanelStyle::Get()->GetBrush("RemoteControlPanel.HorizontalDash"))
+		.OnDropped_Lambda([this](const FGeometry& InGeometry, const FDragDropEvent& InDragDropEvent) { return SRCControllerPanelList::OnControllerListViewDragDrop(InDragDropEvent.GetOperation()); })
+		.OnAllowDrop(this, &SRCControllerPanelList::OnAllowDrop)
+		.OnIsRecognized(this, &SRCControllerPanelList::OnAllowDrop)
+		[
+			ListView.ToSharedRef()
+		]
 	];
 
 	// Add delegates
@@ -241,7 +338,7 @@ void SRCControllerPanelList::Reset()
 	TSharedPtr<SRCControllerPanel> ControllerPanel = ControllerPanelWeakPtr.Pin();
 	URemoteControlPreset* Preset = ControllerPanel->GetPreset();
 	TSharedPtr<SRemoteControlPanel> RemoteControlPanel = ControllerPanel->GetRemoteControlPanel();
-
+	
 	check(Preset);
 
 	PropertyRowGenerator->SetStructure(Preset->GetControllerContainerStructOnScope());
@@ -260,7 +357,7 @@ void SRCControllerPanelList::Reset()
 		for (TSharedRef<IDetailTreeNode>& Child : Children)
 		{
 			FProperty* Property = Child->CreatePropertyHandle()->GetProperty();
-			check(Property);
+ 			check(Property);
 
 			if (URCVirtualPropertyBase* Controller = Preset->GetVirtualProperty(Property->GetFName()))
 			{
@@ -298,6 +395,22 @@ void SRCControllerPanelList::OnTreeSelectionChanged(TSharedPtr<FRCControllerMode
 	}
 }
 
+void SRCControllerPanelList::SelectController(URCController* InController)
+{
+	for (TSharedPtr<FRCControllerModel> ControllerItem : ControllerItems)
+	{
+		if (!ensure(ControllerItem))
+		{
+			continue;
+		}
+
+		if (ControllerItem->GetVirtualProperty() == InController)
+		{
+			ListView->SetSelection(ControllerItem);
+		}
+	}
+}
+
 void SRCControllerPanelList::OnControllerAdded(const FName& InNewPropertyName)
 {
 	Reset();
@@ -452,4 +565,179 @@ void SRCControllerPanelList::ReorderControllerItem(TSharedRef<FRCControllerModel
 	ListView->RequestListRefresh();
 }
 
-#undef LOCTEXT_NAMESPACE
+static TSharedPtr<FExposedEntityDragDrop> GetExposedEntityDragDrop(TSharedPtr<FDragDropOperation> DragDropOperation)
+{
+	if (DragDropOperation)
+	{
+		if (DragDropOperation->IsOfType<FExposedEntityDragDrop>())
+		{
+			return StaticCastSharedPtr<FExposedEntityDragDrop>(DragDropOperation);
+		}
+	}
+
+	return nullptr;
+}
+
+bool SRCControllerPanelList::OnAllowDrop(TSharedPtr<FDragDropOperation> DragDropOperation)
+{
+	if (IsListViewHovered())
+	{
+		bIsAnyControllerItemEligibleForDragDrop = false;
+	}
+	// Ensures that this drop target is visually disabled whenever the user is attempting to drop onto an existing Controller (rather than the ListView's empty space)
+	else if (bIsAnyControllerItemEligibleForDragDrop)
+	{
+		return false;
+	}
+
+	if (TSharedPtr<FExposedEntityDragDrop> DragDropOp = GetExposedEntityDragDrop(DragDropOperation))
+	{
+		// Fetch the Exposed Entity
+		const FGuid ExposedEntityId = DragDropOp->GetId();
+
+		if (URemoteControlPreset* Preset = GetPreset())
+		{
+			if (TSharedPtr<const FRemoteControlField> RemoteControlField = Preset->GetExposedEntity<FRemoteControlField>(ExposedEntityId).Pin())
+			{
+				return RemoteControlField->FieldType == EExposedFieldType::Property;
+			}
+		}
+	}
+
+	return false;
+}
+
+FReply SRCControllerPanelList::OnControllerListViewDragDrop(TSharedPtr<FDragDropOperation> DragDropOperation)
+{
+	if (TSharedPtr<FExposedEntityDragDrop> DragDropOp = GetExposedEntityDragDrop(DragDropOperation))
+	{
+		// Fetch the Exposed Entity
+		const FGuid ExposedEntityId = DragDropOp->GetId();
+
+		if (URemoteControlPreset* Preset = GetPreset())
+		{
+			if (TSharedPtr<const FRemoteControlProperty> RemoteControlProperty = Preset->GetExposedEntity<FRemoteControlProperty>(ExposedEntityId).Pin())
+			{
+				CreateAutoBindForProperty(RemoteControlProperty);
+			}
+		}
+	}
+	
+	return FReply::Handled();
+}
+
+void SRCControllerPanelList::CreateAutoBindForProperty(TSharedPtr<const FRemoteControlProperty> RemoteControlProperty)
+{
+	if (URemoteControlPreset* Preset = GetPreset())
+	{
+		// Derive the input data needed for creating a new Controller
+		FProperty* Property = RemoteControlProperty->GetProperty();
+		EPropertyBagPropertyType PropertyBagType = EPropertyBagPropertyType::None;
+		UObject* StructObject = nullptr;
+
+		// In the Logic realm we use a single type like (eg: String / Int) to represent various related types (String/Name/Text, Int32, Int64, etc)
+		// For this reason explicit mapping conversion is required between a given FProperty type and the desired Controller type
+		bool bSuccess = URCBehaviourBind::GetPropertyBagTypeFromFieldProperty(Property, PropertyBagType, StructObject);
+		if(bSuccess)
+		{
+			// Step 1. Create a Controller of matching type
+			URCController* NewController = Cast<URCController>(Preset->AddVirtualProperty(URCController::StaticClass(), PropertyBagType, StructObject));
+			NewController->DisplayIndex = Preset->GetNumVirtualProperties() - 1;
+
+			// Transfer property value from Exposed Property to the New Controller.
+			// The goal is to keep values synced for a Controller newly created via "Auto Bind"
+			URCBehaviourBind::CopyPropertyValueToController(NewController, RemoteControlProperty.ToSharedRef());
+
+			// Step 2. Refresh UI
+			Reset();
+
+			// Step 3. Create Bind Behaviour and Bind to the property
+			CreateBindBehaviourAndAssignTo(NewController, RemoteControlProperty.ToSharedRef(), false);
+		}
+	}
+}
+
+void SRCControllerPanelList::CreateBindBehaviourAndAssignTo(URCController* Controller, TSharedRef<const FRemoteControlProperty> InRemoteControlProperty, const bool bExecuteBind)
+{
+	const URCBehaviourBind* BindBehaviour = nullptr;
+
+	bool bRequiresNumericConversion = false;
+	if (!URCBehaviourBind::CanHaveActionForField(Controller, InRemoteControlProperty, false))
+	{
+		if (URCBehaviourBind::CanHaveActionForField(Controller, InRemoteControlProperty, true))
+		{
+			bRequiresNumericConversion = true;
+		}
+		else
+		{
+			ensureAlwaysMsgf(false, TEXT("Incompatible property provided for Auto Bind!"));
+			return;
+		}
+	}
+
+	for (const URCBehaviour* Behaviour : Controller->Behaviours)
+	{
+		if (Behaviour && Behaviour->IsA(URCBehaviourBind::StaticClass()))
+		{
+			BindBehaviour = Cast<URCBehaviourBind>(Behaviour);
+
+			// In case numeric conversion is required we might have multiple Bind behaviours with different settings,
+			// so we do not break in case there is at least one Bind behaviour with a matching clause. If not, we will create a new Bind behaviour with the requried setting (see below)
+			if (!bRequiresNumericConversion || BindBehaviour->AreNumericInputsAllowedAsStrings())
+			{
+				break;
+			}
+		}
+	}
+
+	if (BindBehaviour)
+	{
+		if (bRequiresNumericConversion && !BindBehaviour->AreNumericInputsAllowedAsStrings())
+		{
+			// If the requested Bind operation requires numeric conversion but the existing Bind behaviour doesn't support this, then we prefer creating a new Bind behaviour to facilitate this operation.
+			// This allows the the user to successfully perform the Auto Bind as desired without disrupting the setting on the existing Bind behaviour
+			BindBehaviour = nullptr;
+		}
+		
+	}
+
+	if (!BindBehaviour)
+	{
+		URCBehaviourBind* NewBindBehaviour = Cast<URCBehaviourBind>(Controller->AddBehaviour(URCBehaviourBindNode::StaticClass()));
+
+		// If this is a new Bind Behaviour and we are attempting to link unrelated (but compatible types), via numeric conversion, then we apply the bAllowNumericInputAsStrings flag
+		NewBindBehaviour->SetAllowNumericInputAsStrings(bRequiresNumericConversion);
+
+		BindBehaviour = NewBindBehaviour;
+
+		// Broadcast new behaviour
+		if (const TSharedPtr<SRCControllerPanel> ControllerPanel = ControllerPanelWeakPtr.Pin())
+		{
+			if (const TSharedPtr<SRemoteControlPanel> RemoteControlPanel = ControllerPanel->GetRemoteControlPanel())
+			{
+				RemoteControlPanel->OnBehaviourAdded.Broadcast(BindBehaviour);
+			}
+		}
+	}
+
+	if (ensure(BindBehaviour))
+	{
+		URCBehaviourBind* BindBehaviourMutable = const_cast<URCBehaviourBind*>(BindBehaviour);
+		URCAction* BindAction = BindBehaviourMutable->AddPropertyBindAction(InRemoteControlProperty);
+
+		if (bExecuteBind)
+		{
+			BindAction->Execute();
+		}
+	}
+
+	// Update the UI selection to provide the user visual feedback indicating that a Bind behaviour has been created/updated
+	SelectController(Controller);
+}
+
+bool SRCControllerPanelList::IsListViewHovered()
+{
+	return ListView->IsDirectlyHovered();
+}
+
+#undef LOCTEXT_NAMESPACE
\ No newline at end of file
diff --git a/Engine/Plugins/VirtualProduction/RemoteControl/Source/RemoteControlUI/Private/UI/Controller/SRCControllerPanelList.h b/Engine/Plugins/VirtualProduction/RemoteControl/Source/RemoteControlUI/Private/UI/Controller/SRCControllerPanelList.h
--- a/Engine/Plugins/VirtualProduction/RemoteControl/Source/RemoteControlUI/Private/UI/Controller/SRCControllerPanelList.h
+++ b/Engine/Plugins/VirtualProduction/RemoteControl/Source/RemoteControlUI/Private/UI/Controller/SRCControllerPanelList.h
@@ -6,12 +6,14 @@
 #include "UI/RemoteControlPanelStyle.h"
 #include "Widgets/Layout/SBorder.h"
 
+class FDragDropOperation;
 struct FRCPanelStyle;
 class FRCControllerModel;
 class FRCLogicModeBase;
 class IPropertyRowGenerator;
 class ITableRow;
 class ITableBase;
+class SDropTarget;
 class SRCControllerPanel;
 class SRemoteControlPanel;
 class STableViewBase;
@@ -115,6 +117,29 @@ public:
 		Reset();
 	}
 
+	/** Drag-Drop validation delegate for the Controllers Panel List */
+	bool OnAllowDrop(TSharedPtr<FDragDropOperation> DragDropOperation);
+
+	/** Drag-Drop action delegate for the Controllers Panel List*/
+	FReply OnControllerListViewDragDrop(TSharedPtr<FDragDropOperation> DragDropOperation);
+
+	/** Fetches the Remote Control preset associated with the parent panel */
+	virtual URemoteControlPreset* GetPreset() override;
+
+	/** Creates a Bind Behaviour for the given Controller and binds the given remote control property to it*/
+	void CreateBindBehaviourAndAssignTo(URCController* Controller, TSharedRef<const FRemoteControlProperty> InRemoteControlProperty, const bool bExecuteBind);
+
+	/** Whether the user's cursor is directly hovered over the List View*/
+	bool IsListViewHovered();
+
+	/** Flag that facilitates usage of two mutually exclusive drag-drop zones within a single panel 
+	* 
+	* The first drag-drop zone is empty panel space for "Bind To New Controller"
+	* The second drag-drop zone is the Controller name widget for "Bind To Existing Controller"
+	* 
+	* This flag is set if any Controller has active drag-drop focus, in which case we disable the first drag-drop zone. This is purely for visual clarity*/
+	bool bIsAnyControllerItemEligibleForDragDrop = false;
+
 private:
 
 	/** OnGenerateRow delegate for the Actions List View */
@@ -123,6 +148,9 @@ private:
 	/** OnSelectionChanged delegate for Actions List View */
 	void OnTreeSelectionChanged(TSharedPtr<FRCControllerModel> InItem , ESelectInfo::Type);
 
+	/** Selects the Controller UI item corresponding to a given Controller UObject */
+	void SelectController(URCController* InController);
+
 	/** Responds to the selection of a newly created Controller. Resets UI state */
 	void OnControllerAdded(const FName& InNewPropertyName);
 	
@@ -136,6 +164,9 @@ private:
 	*/
 	void OnFinishedChangingProperties(const FPropertyChangedEvent& PropertyChangedEvent);
 
+	/** Creates a new Controller for the given Remote Control Property and also binds to it */
+	void CreateAutoBindForProperty(TSharedPtr<const FRemoteControlProperty> RemoteControlProperty);
+
 	/** The row generator used to represent each Controller as a row, when used with SListView */
 	TSharedPtr<IPropertyRowGenerator> PropertyRowGenerator;
 
@@ -157,9 +188,6 @@ private:
 	/** Handles broadcasting of a successful remove item operation.*/
 	virtual void BroadcastOnItemRemoved() override;
 
-	/** Fetches the Remote Control preset associated with the parent panel */
-	virtual URemoteControlPreset* GetPreset() override;
-
 	/** Removes the given Controller UI model item from the list of UI models*/
 	virtual int32 RemoveModel(const TSharedPtr<FRCLogicModeBase> InModel) override;
 
diff --git a/Engine/Plugins/VirtualProduction/RemoteControl/Source/RemoteControlUI/Private/UI/SRemoteControlPanel.h b/Engine/Plugins/VirtualProduction/RemoteControl/Source/RemoteControlUI/Private/UI/SRemoteControlPanel.h
--- a/Engine/Plugins/VirtualProduction/RemoteControl/Source/RemoteControlUI/Private/UI/SRemoteControlPanel.h
+++ b/Engine/Plugins/VirtualProduction/RemoteControl/Source/RemoteControlUI/Private/UI/SRemoteControlPanel.h
@@ -66,7 +66,7 @@ public:
 	// Remote Control Logic Delegates
 	DECLARE_MULTICAST_DELEGATE_OneParam(FOnControllerAdded, const FName& /* InPropertyName */);
 	DECLARE_MULTICAST_DELEGATE_OneParam(FOnControllerSelectionChanged, TSharedPtr<FRCControllerModel> /* InControllerItem */);
-	DECLARE_MULTICAST_DELEGATE_OneParam(FOnBehaviourAdded, URCBehaviour* /* InBehaviour */);
+	DECLARE_MULTICAST_DELEGATE_OneParam(FOnBehaviourAdded, const URCBehaviour* /* InBehaviour */);
 	DECLARE_MULTICAST_DELEGATE_OneParam(FOnBehaviourSelectionChanged, TSharedPtr<FRCBehaviourModel> /* InBehaviourItem */);
 	DECLARE_MULTICAST_DELEGATE_OneParam(FOnActionAdded, URCAction* /* InAction */);
 
