Fix static analysis error in Horde.
Issue 211995

diff --git a/Engine/Plugins/VirtualProduction/RemoteControl/Source/RemoteControlUI/Private/UI/Action/SRCActionPanel.cpp b/Engine/Plugins/VirtualProduction/RemoteControl/Source/RemoteControlUI/Private/UI/Action/SRCActionPanel.cpp
--- a/Engine/Plugins/VirtualProduction/RemoteControl/Source/RemoteControlUI/Private/UI/Action/SRCActionPanel.cpp
+++ b/Engine/Plugins/VirtualProduction/RemoteControl/Source/RemoteControlUI/Private/UI/Action/SRCActionPanel.cpp
@@ -522,31 +522,25 @@ void SRCActionPanel::PasteItemFromClipboard()
 bool SRCActionPanel::CanPasteClipboardItem(UObject* InLogicClipboardItem)
 {
 	URCAction* LogicClipboardAction = Cast<URCAction>(InLogicClipboardItem);
-
-	URCBehaviour* BehaviourSource = nullptr;
-	URCBehaviour* BehaviourTarget = nullptr;
-
-	if (LogicClipboardAction)
+	if (!LogicClipboardAction)
 	{
-		BehaviourSource = LogicClipboardAction->GetParentBehaviour();
-		if (!BehaviourSource)
-		{
-			return false;
-		}
+		return false;
 	}
 
-	if (TSharedPtr<FRCBehaviourModel> BehaviourItemTarget = SelectedBehaviourItemWeakPtr.Pin())
+	if (URCBehaviour* BehaviourSource = LogicClipboardAction->GetParentBehaviour())
 	{
-		BehaviourTarget = BehaviourItemTarget->GetBehaviour();
-		if (!BehaviourTarget)
+		if (TSharedPtr<FRCBehaviourModel> BehaviourItemTarget = SelectedBehaviourItemWeakPtr.Pin())
 		{
-			return false;
+			if (URCBehaviour* BehaviourTarget = BehaviourItemTarget->GetBehaviour())
+			{
+				// Copy-paste is allowed between compatible Behaviour types only
+				//
+				return BehaviourSource->GetClass() == BehaviourTarget->GetClass();
+			}
 		}
 	}
 
-	// Copy-paste is allowed between compatible Behaviour types only
-	//
-	return BehaviourSource->GetClass() == BehaviourTarget->GetClass();
+	return false;
 }
 
 FText SRCActionPanel::GetPasteItemMenuEntrySuffix()
