Controller panel improvements:
* Added color code as requested in UE-1629507
* Added tooltip for all color codes
* Fixed Controller type display name bug for Structs

Misc:
As discussed in prior meetings:
* Removed IsEqual behaviour
* Renamed SetValue behaviour to OnValueChanged

diff --git a/Engine/Plugins/VirtualProduction/RemoteControl/Source/RemoteControlLogic/Private/Behaviour/Builtin/RCBehaviourIsEqualNode.cpp b/Engine/Plugins/VirtualProduction/RemoteControl/Source/RemoteControlLogic/Private/Behaviour/Builtin/RCBehaviourIsEqualNode.cpp
deleted file mode 100644
diff --git a/Engine/Plugins/VirtualProduction/RemoteControl/Source/RemoteControlLogic/Private/Behaviour/Builtin/RCBehaviourOnEmptyNode.cpp b/Engine/Plugins/VirtualProduction/RemoteControl/Source/RemoteControlLogic/Private/Behaviour/Builtin/RCBehaviourOnEmptyNode.cpp
deleted file mode 100644
diff --git a/Engine/Plugins/VirtualProduction/RemoteControl/Source/RemoteControlLogic/Private/Behaviour/Builtin/RCBehaviourOnValueChangedNode.cpp b/Engine/Plugins/VirtualProduction/RemoteControl/Source/RemoteControlLogic/Private/Behaviour/Builtin/RCBehaviourOnValueChangedNode.cpp
new file mode 100644
--- /dev/null
+++ b/Engine/Plugins/VirtualProduction/RemoteControl/Source/RemoteControlLogic/Private/Behaviour/Builtin/RCBehaviourOnValueChangedNode.cpp
@@ -0,0 +1,37 @@
+﻿// Copyright Epic Games, Inc. All Rights Reserved.
+
+#include "Behaviour/Builtin/RCBehaviourOnValueChangedNode.h"
+
+#include "PropertyBag.h"
+
+URCBehaviourOnValueChangedNode::URCBehaviourOnValueChangedNode()
+{
+	DisplayName = NSLOCTEXT("Remote Control Behaviour", "Behavior Name - On Value Changed", "On Value Changed");
+	BehaviorDescription = NSLOCTEXT("Remote Control Behaviour", "Behavior Desc - On Value Changed", "Triggers an event when the associated property is modified");
+}
+
+
+bool URCBehaviourOnValueChangedNode::Execute_Implementation(URCBehaviour* InBehaviour) const
+{
+	return true;
+}
+
+bool URCBehaviourOnValueChangedNode::IsSupported_Implementation(URCBehaviour* InBehaviour) const
+{
+	static TArray<TPair<EPropertyBagPropertyType, UObject*>> SupportedPropertyBagTypes =
+	{
+		{ EPropertyBagPropertyType::Bool, nullptr },
+		{ EPropertyBagPropertyType::Byte, nullptr },
+		{ EPropertyBagPropertyType::Int32, nullptr },
+		{ EPropertyBagPropertyType::Float, nullptr },
+		{ EPropertyBagPropertyType::Double, nullptr },
+		{ EPropertyBagPropertyType::Name, nullptr },
+		{ EPropertyBagPropertyType::String, nullptr },
+		{ EPropertyBagPropertyType::Text, nullptr },
+		{ EPropertyBagPropertyType::Struct, TBaseStructure<FVector>::Get() },
+		{ EPropertyBagPropertyType::Struct, TBaseStructure<FColor>::Get() },
+		{ EPropertyBagPropertyType::Struct, TBaseStructure<FRotator>::Get() }
+	};
+	
+	return SupportedPropertyBagTypes.ContainsByPredicate(GetIsSupportedCallback(InBehaviour));
+}
diff --git a/Engine/Plugins/VirtualProduction/RemoteControl/Source/RemoteControlLogic/Private/Behaviour/Builtin/RCBehaviourSetValueNode.cpp b/Engine/Plugins/VirtualProduction/RemoteControl/Source/RemoteControlLogic/Private/Behaviour/Builtin/RCBehaviourSetValueNode.cpp
deleted file mode 100644
diff --git a/Engine/Plugins/VirtualProduction/RemoteControl/Source/RemoteControlLogic/Private/Behaviour/RCIsEqualBehaviour.cpp b/Engine/Plugins/VirtualProduction/RemoteControl/Source/RemoteControlLogic/Private/Behaviour/RCIsEqualBehaviour.cpp
deleted file mode 100644
diff --git a/Engine/Plugins/VirtualProduction/RemoteControl/Source/RemoteControlLogic/Private/Tests/RemoteControlLogicTests.cpp b/Engine/Plugins/VirtualProduction/RemoteControl/Source/RemoteControlLogic/Private/Tests/RemoteControlLogicTests.cpp
--- a/Engine/Plugins/VirtualProduction/RemoteControl/Source/RemoteControlLogic/Private/Tests/RemoteControlLogicTests.cpp
+++ b/Engine/Plugins/VirtualProduction/RemoteControl/Source/RemoteControlLogic/Private/Tests/RemoteControlLogicTests.cpp
@@ -11,10 +11,11 @@
 #include "Action/RCFunctionAction.h"
 #include "Action/RCPropertyAction.h"
 #include "Behaviour/RCBehaviour.h"
-#include "Behaviour/RCIsEqualBehaviour.h"
-#include "Behaviour/Builtin/RCBehaviourIsEqualNode.h"
-#include "Behaviour/Builtin/RCBehaviourOnEmptyNode.h"
-#include "Behaviour/Builtin/RCBehaviourSetValueNode.h"
+#include "Behaviour/Builtin/Bind/RCBehaviourBind.h"
+#include "Behaviour/Builtin/Bind/RCBehaviourBindNode.h"
+#include "Behaviour/Builtin/Conditional/RCBehaviourConditional.h"
+#include "Behaviour/Builtin/Conditional/RCBehaviourConditionalNode.h"
+#include "Behaviour/Builtin/RCBehaviourOnValueChangedNode.h"
 #include "Controller/RCController.h"
 #include "Misc/AutomationTest.h"
 #include "UObject/StrongObjectPtr.h"
@@ -70,18 +71,17 @@ bool FRemoteControlLogicTest::RunTest(const FString& Parameters)
 	TestEqual(TEXT("Vectors should be the same"), VectorValue, OutVectorValue);
 	
 	// 4. Add Behaviour To Properties
-	// 4.1 Add Is Equal Behaviour
-	URCBehaviour* FloatControllerBehaviour = FloatController->AddBehaviour(URCBehaviourIsEqualNode::StaticClass());
-	const URCIsEqualBehaviour* IsEqualBehaviour = Cast<URCIsEqualBehaviour>(FloatControllerBehaviour);
+	// 4.1 Add Conditional Behaviour
+	URCBehaviour* FloatControllerBehaviour = FloatController->AddBehaviour(URCBehaviourConditionalNode::StaticClass());
+	const URCBehaviourConditional* ConditionalBehaviour = Cast<URCBehaviourConditional>(FloatControllerBehaviour);
 	float IsEqualBehaviourFloatValue = 0.f;
  
-	// 4.1 Get Is Equal behaviour Value
-	TestTrue(TEXT("Should Get Float"), IsEqualBehaviour->PropertySelfContainer->GetValueFloat(IsEqualBehaviourFloatValue));
-	TestEqual(TEXT("float controller value should be the same as a is equal behaviour"), IsEqualBehaviourFloatValue, FloatValue);
+	// @todo: Add Conditional Behaviour Unit Tests here.
+	// ...
  
 	// 4.2 Add Other behaviours
-	const URCBehaviour* IntControllerBehaviour = FloatController->AddBehaviour(URCBehaviourSetValueNode::StaticClass());
-	const URCBehaviour* StrControllerBehaviour = StrController->AddBehaviour(URCBehaviourOnEmptyNode::StaticClass());
+	const URCBehaviour* IntControllerBehaviour = FloatController->AddBehaviour(URCBehaviourOnValueChangedNode::StaticClass());
+	const URCBehaviour* StrControllerBehaviour = StrController->AddBehaviour(URCBehaviourBindNode::StaticClass());
  
 	// 5 Add actions
 	// 5.1 Add Float Controller Actions
diff --git a/Engine/Plugins/VirtualProduction/RemoteControl/Source/RemoteControlLogic/Public/Behaviour/Builtin/RCBehaviourIsEqualNode.h b/Engine/Plugins/VirtualProduction/RemoteControl/Source/RemoteControlLogic/Public/Behaviour/Builtin/RCBehaviourIsEqualNode.h
deleted file mode 100644
diff --git a/Engine/Plugins/VirtualProduction/RemoteControl/Source/RemoteControlLogic/Public/Behaviour/Builtin/RCBehaviourOnEmptyNode.h b/Engine/Plugins/VirtualProduction/RemoteControl/Source/RemoteControlLogic/Public/Behaviour/Builtin/RCBehaviourOnEmptyNode.h
deleted file mode 100644
diff --git a/Engine/Plugins/VirtualProduction/RemoteControl/Source/RemoteControlLogic/Public/Behaviour/Builtin/RCBehaviourOnValueChangedNode.h b/Engine/Plugins/VirtualProduction/RemoteControl/Source/RemoteControlLogic/Public/Behaviour/Builtin/RCBehaviourOnValueChangedNode.h
new file mode 100644
--- /dev/null
+++ b/Engine/Plugins/VirtualProduction/RemoteControl/Source/RemoteControlLogic/Public/Behaviour/Builtin/RCBehaviourOnValueChangedNode.h
@@ -0,0 +1,29 @@
+﻿// Copyright Epic Games, Inc. All Rights Reserved.
+
+#pragma once
+
+#include "Behaviour/RCBehaviourNode.h"
+#include "RCBehaviourOnValueChangedNode.generated.h"
+
+/**
+ * On Value Changed Behaviour Node
+ *
+ * This behaviour always returns true and thus executes all actions any time the Controller value is modified
+ */
+UCLASS()
+class REMOTECONTROLLOGIC_API URCBehaviourOnValueChangedNode : public URCBehaviourNode
+{
+	GENERATED_BODY()
+
+public:
+
+	URCBehaviourOnValueChangedNode();
+
+	//~ Begin URCBehaviourNode interface
+	UFUNCTION(BlueprintNativeEvent, BlueprintCallable, Category = "Remote Control Behaviour")
+	bool Execute(URCBehaviour* InBehaviour) const override;
+
+	UFUNCTION(BlueprintNativeEvent, BlueprintCallable, Category = "Remote Control Behaviour")
+	bool IsSupported(URCBehaviour* InBehaviour) const override;
+	//~ End URCBehaviourNode interface
+};
diff --git a/Engine/Plugins/VirtualProduction/RemoteControl/Source/RemoteControlLogic/Public/Behaviour/Builtin/RCBehaviourSetValueNode.h b/Engine/Plugins/VirtualProduction/RemoteControl/Source/RemoteControlLogic/Public/Behaviour/Builtin/RCBehaviourSetValueNode.h
deleted file mode 100644
diff --git a/Engine/Plugins/VirtualProduction/RemoteControl/Source/RemoteControlLogic/Public/Behaviour/RCIsEqualBehaviour.h b/Engine/Plugins/VirtualProduction/RemoteControl/Source/RemoteControlLogic/Public/Behaviour/RCIsEqualBehaviour.h
deleted file mode 100644
diff --git a/Engine/Plugins/VirtualProduction/RemoteControl/Source/RemoteControlUI/Private/UI/Behaviour/Builtin/RCBehaviourIsEqualModel.cpp b/Engine/Plugins/VirtualProduction/RemoteControl/Source/RemoteControlUI/Private/UI/Behaviour/Builtin/RCBehaviourIsEqualModel.cpp
deleted file mode 100644
diff --git a/Engine/Plugins/VirtualProduction/RemoteControl/Source/RemoteControlUI/Private/UI/Behaviour/Builtin/RCBehaviourIsEqualModel.h b/Engine/Plugins/VirtualProduction/RemoteControl/Source/RemoteControlUI/Private/UI/Behaviour/Builtin/RCBehaviourIsEqualModel.h
deleted file mode 100644
diff --git a/Engine/Plugins/VirtualProduction/RemoteControl/Source/RemoteControlUI/Private/UI/Behaviour/Builtin/SRCBehaviourIsEqual.cpp b/Engine/Plugins/VirtualProduction/RemoteControl/Source/RemoteControlUI/Private/UI/Behaviour/Builtin/SRCBehaviourIsEqual.cpp
deleted file mode 100644
diff --git a/Engine/Plugins/VirtualProduction/RemoteControl/Source/RemoteControlUI/Private/UI/Behaviour/Builtin/SRCBehaviourIsEqual.h b/Engine/Plugins/VirtualProduction/RemoteControl/Source/RemoteControlUI/Private/UI/Behaviour/Builtin/SRCBehaviourIsEqual.h
deleted file mode 100644
diff --git a/Engine/Plugins/VirtualProduction/RemoteControl/Source/RemoteControlUI/Private/UI/Behaviour/RCBehaviourModel.h b/Engine/Plugins/VirtualProduction/RemoteControl/Source/RemoteControlUI/Private/UI/Behaviour/RCBehaviourModel.h
--- a/Engine/Plugins/VirtualProduction/RemoteControl/Source/RemoteControlUI/Private/UI/Behaviour/RCBehaviourModel.h
+++ b/Engine/Plugins/VirtualProduction/RemoteControl/Source/RemoteControlUI/Private/UI/Behaviour/RCBehaviourModel.h
@@ -1,7 +1,7 @@
 ﻿// Copyright Epic Games, Inc. All Rights Reserved.
 
 #pragma once
-#include "Behaviour/RCIsEqualBehaviour.h"
+#include "Behaviour/RCBehaviour.h"
 #include "UI/BaseLogicUI/RCLogicModeBase.h"
 
 struct FRCPanelStyle;
diff --git a/Engine/Plugins/VirtualProduction/RemoteControl/Source/RemoteControlUI/Private/UI/Behaviour/SRCBehaviourPanelList.cpp b/Engine/Plugins/VirtualProduction/RemoteControl/Source/RemoteControlUI/Private/UI/Behaviour/SRCBehaviourPanelList.cpp
--- a/Engine/Plugins/VirtualProduction/RemoteControl/Source/RemoteControlUI/Private/UI/Behaviour/SRCBehaviourPanelList.cpp
+++ b/Engine/Plugins/VirtualProduction/RemoteControl/Source/RemoteControlUI/Private/UI/Behaviour/SRCBehaviourPanelList.cpp
@@ -13,7 +13,6 @@
 #include "Styling/RemoteControlStyles.h"
 #include "UI/Behaviour/Builtin/Bind/RCBehaviourBindModel.h"
 #include "UI/Behaviour/Builtin/Conditional/RCBehaviourConditionalModel.h"
-#include "UI/Behaviour/Builtin/RCBehaviourIsEqualModel.h"
 #include "UI/Behaviour/Builtin/Path/RCBehaviourSetAssetByPathModel.h"
 #include "UI/Controller/RCControllerModel.h"
 #include "UI/Behaviour/Builtin/RangeMap/RCBehaviourRangeMapModel.h"
@@ -107,11 +106,7 @@ void SRCBehaviourPanelList::AddBehaviourToList(URCBehaviour* InBehaviour)
 
 	const TSharedPtr<SRemoteControlPanel> RemoteControlPanel = BehaviourPanelWeakPtr.Pin()->GetRemoteControlPanel();
 
-	if (URCIsEqualBehaviour* IsEqualBehaviour = Cast<URCIsEqualBehaviour>(InBehaviour))
-	{
-		BehaviourItems.Add(MakeShared<FRCIsEqualBehaviourModel>(IsEqualBehaviour));
-	}
-	else if (URCBehaviourConditional* ConditionalBehaviour = Cast<URCBehaviourConditional>(InBehaviour))
+	if (URCBehaviourConditional* ConditionalBehaviour = Cast<URCBehaviourConditional>(InBehaviour))
 	{
 		BehaviourItems.Add(MakeShared<FRCBehaviourConditionalModel>(ConditionalBehaviour));
 	}
diff --git a/Engine/Plugins/VirtualProduction/RemoteControl/Source/RemoteControlUI/Private/UI/Controller/SRCControllerPanel.cpp b/Engine/Plugins/VirtualProduction/RemoteControl/Source/RemoteControlUI/Private/UI/Controller/SRCControllerPanel.cpp
--- a/Engine/Plugins/VirtualProduction/RemoteControl/Source/RemoteControlUI/Private/UI/Controller/SRCControllerPanel.cpp
+++ b/Engine/Plugins/VirtualProduction/RemoteControl/Source/RemoteControlUI/Private/UI/Controller/SRCControllerPanel.cpp
@@ -300,9 +300,9 @@ TSharedRef<SWidget> SRCControllerPanel::GetControllerMenuContentWidget() const
 		FInstancedPropertyBag Bag;
 		Bag.AddProperty(DefaultName, Pair.Key, Pair.Value);
 		const FPropertyBagPropertyDesc* BagPropertyDesc = Bag.FindPropertyDescByName(DefaultName);
-		if (ensure(BagPropertyDesc))
+		if (!ensure(BagPropertyDesc))
 		{
-			TypeColor = UE::RCUIHelpers::GetFieldClassTypeColor(BagPropertyDesc->CachedProperty);
+			continue;
 		}
 
 		// Variable Color Bar
@@ -313,15 +313,7 @@ TSharedRef<SWidget> SRCControllerPanel::GetControllerMenuContentWidget() const
 			.HAlign(HAlign_Left)
 			.VAlign(VAlign_Fill)
 			[
-				SNew(SBox)
-				.HeightOverride(5.f)
-				[
-					SNew(SBorder)
-					.Visibility(EVisibility::HitTestInvisible)
-					.BorderImage(FAppStyle::Get().GetBrush("NumericEntrySpinBox.NarrowDecorator"))
-					.BorderBackgroundColor(TypeColor)
-					.Padding(FMargin(5.0f, 0.0f, 0.0f, 0.f))
-				]
+				UE::RCUIHelpers::GetTypeColorWidget(BagPropertyDesc->CachedProperty)
 			]
 		+ SHorizontalBox::Slot()
 			.AutoWidth()
diff --git a/Engine/Plugins/VirtualProduction/RemoteControl/Source/RemoteControlUI/Private/UI/Controller/SRCControllerPanelList.cpp b/Engine/Plugins/VirtualProduction/RemoteControl/Source/RemoteControlUI/Private/UI/Controller/SRCControllerPanelList.cpp
--- a/Engine/Plugins/VirtualProduction/RemoteControl/Source/RemoteControlUI/Private/UI/Controller/SRCControllerPanelList.cpp
+++ b/Engine/Plugins/VirtualProduction/RemoteControl/Source/RemoteControlUI/Private/UI/Controller/SRCControllerPanelList.cpp
@@ -22,6 +22,7 @@
 #include "Styling/RemoteControlStyles.h"
 #include "UI/Action/SRCActionPanelList.h"
 #include "UI/BaseLogicUI/RCLogicModeBase.h"
+#include "UI/RCUIHelpers.h"
 #include "UI/RemoteControlPanelStyle.h"
 #include "UI/SRCPanelExposedEntity.h"
 #include "UI/SRCPanelDragHandle.h"
@@ -45,6 +46,7 @@ namespace UE::RCControllerPanelList
 
 	namespace Columns
 	{
+		const FName TypeColor = TEXT("TypeColor");
 		const FName Name = TEXT("Controller Name");
 		const FName Value = TEXT("Controller Value");
 		const FName DragHandle = TEXT("Drag Handle");
@@ -65,7 +67,15 @@ namespace UE::RCControllerPanelList
 			if (!ensure(ControllerItem.IsValid()))
 				return SNullWidget::NullWidget;
 
-			if (ColumnName == UE::RCControllerPanelList::Columns::Name)
+
+			if (ColumnName == UE::RCControllerPanelList::Columns::TypeColor)
+			{
+				if (URCController* Controller = Cast< URCController>(ControllerItem->GetVirtualProperty()))
+				{
+					return UE::RCUIHelpers::GetTypeColorWidget(Controller->GetProperty());
+				}
+			}
+			else if (ColumnName == UE::RCControllerPanelList::Columns::Name)
 			{
 				return WrapWithDropTarget(ControllerItem->GetNameWidget());
 			}
@@ -266,9 +276,14 @@ void SRCControllerPanelList::Construct(const FArguments& InArgs, const TSharedRe
 			SNew(SHeaderRow)
 			.Style(&RCPanelStyle->HeaderRowStyle)
 
+			+ SHeaderRow::Column(UE::RCControllerPanelList::Columns::TypeColor)
+			.DefaultLabel(LOCTEXT("ControllerNameColumnName", ""))
+			.FillWidth(0.03f)
+			.HeaderContentPadding(RCPanelStyle->HeaderRowPadding)
+
 			+ SHeaderRow::Column(UE::RCControllerPanelList::Columns::DragHandle)
 			.DefaultLabel(FText::GetEmpty())
-			.FillWidth(0.1f)
+			.FillWidth(0.07f)
 			.HeaderContentPadding(RCPanelStyle->HeaderRowPadding)
 
 			+ SHeaderRow::Column(UE::RCControllerPanelList::Columns::Name)
diff --git a/Engine/Plugins/VirtualProduction/RemoteControl/Source/RemoteControlUI/Private/UI/RCUIHelpers.cpp b/Engine/Plugins/VirtualProduction/RemoteControl/Source/RemoteControlUI/Private/UI/RCUIHelpers.cpp
--- a/Engine/Plugins/VirtualProduction/RemoteControl/Source/RemoteControlUI/Private/UI/RCUIHelpers.cpp
+++ b/Engine/Plugins/VirtualProduction/RemoteControl/Source/RemoteControlUI/Private/UI/RCUIHelpers.cpp
@@ -46,27 +46,38 @@ FName UE::RCUIHelpers::GetFieldClassDisplayName(const FProperty* InProperty)
 			{
 				FieldClassDisplayName = *SubCategoryObject->GetName();
 			}
-			
-			FieldClassDisplayName  = PinType.PinCategory;
+			else
+			{
+				FieldClassDisplayName = PinType.PinCategory;
+			}
 		}
 	}
 
 	if (!FieldClassDisplayName.IsNone())
 	{
-		FString FieldClassNameStr = FieldClassDisplayName.ToString();
+		FString ResultString = FieldClassDisplayName.ToString();
+
+		if (ResultString.StartsWith("Bool"))
+		{
+			ResultString = "Boolean";
+		}
 
-		if (FieldClassNameStr.StartsWith("Bool"))
+		if (ResultString.StartsWith("Int"))
 		{
-			FieldClassDisplayName = "Boolean";
+			ResultString = *ResultString.Replace(TEXT("Int"), TEXT("Integer"));
 		}
 
-		if (FieldClassNameStr.StartsWith("Int"))
+		if (ResultString == "real")
 		{
-			FieldClassDisplayName = *FieldClassNameStr.Replace(TEXT("Int"), TEXT("Integer"));
+			ResultString = "Float";
 		}
+
+		ResultString[0] = FChar::ToUpper(ResultString[0]);
+
+		return *ResultString;
 	}
 
-	return FieldClassDisplayName;
+	return NAME_None;
 }
 
 TSharedPtr<IDetailTreeNode> UE::RCUIHelpers::GetDetailTreeNodeForVirtualProperty(TObjectPtr<URCVirtualPropertySelfContainer> InVirtualPropertySelfContainer, 
@@ -175,3 +186,30 @@ bool UE::RCUIHelpers::FindFocusableWidgetAndSetKeyboardFocus(const TSharedRef< S
 
 	return false; // No focusable widget found in the entire widget hierarchy
 }
+
+TSharedRef<SWidget> UE::RCUIHelpers::GetTypeColorWidget(const FProperty* InProperty)
+{
+	if (!ensure(InProperty))
+	{
+		return SNullWidget::NullWidget;
+	}
+
+	const FLinearColor TypeColor = UE::RCUIHelpers::GetFieldClassTypeColor(InProperty);
+	const FText TooltipText = FText::FromName(UE::RCUIHelpers::GetFieldClassDisplayName(InProperty));
+
+	TSharedRef<SWidget> TypeColorWidget =
+			SNew(SBox)
+			.HeightOverride(5.f)
+			.HAlign(HAlign_Left)
+			.Padding(FMargin(5.f, 0.f))
+			.ToolTipText(TooltipText)
+			[
+				SNew(SBorder)
+				.Visibility(EVisibility::Visible)
+				.BorderImage(FAppStyle::Get().GetBrush("NumericEntrySpinBox.NarrowDecorator"))
+				.BorderBackgroundColor(TypeColor)
+				.Padding(FMargin(5.0f, 0.0f, 0.0f, 0.f))
+			];
+
+	return TypeColorWidget;
+}
diff --git a/Engine/Plugins/VirtualProduction/RemoteControl/Source/RemoteControlUI/Private/UI/RCUIHelpers.h b/Engine/Plugins/VirtualProduction/RemoteControl/Source/RemoteControlUI/Private/UI/RCUIHelpers.h
--- a/Engine/Plugins/VirtualProduction/RemoteControl/Source/RemoteControlUI/Private/UI/RCUIHelpers.h
+++ b/Engine/Plugins/VirtualProduction/RemoteControl/Source/RemoteControlUI/Private/UI/RCUIHelpers.h
@@ -50,4 +50,12 @@ namespace UE::RCUIHelpers
 	* If the input widget is itself focusable, then we simply set the focus on it and return
 	*/
 	bool FindFocusableWidgetAndSetKeyboardFocus(const TSharedRef< SWidget> InWidget);
+
+	/**
+	* GetTypeColorWidget
+	*
+	* Generates a Type Color widget which provides a color coding guide for the user given a desired color
+	* Used to provide a consistent look for the type color widget across various Remote Control Logic panels
+	*/
+	TSharedRef<SWidget> GetTypeColorWidget(const FProperty* InProperty);
 }
\ No newline at end of file
