[RemoteControl] Fix crash when client disconnects mid-transaction

diff --git a/Engine/Plugins/VirtualProduction/RemoteControl/Source/WebRemoteControl/Private/WebSocketMessageHandler.cpp b/Engine/Plugins/VirtualProduction/RemoteControl/Source/WebRemoteControl/Private/WebSocketMessageHandler.cpp
--- a/Engine/Plugins/VirtualProduction/RemoteControl/Source/WebRemoteControl/Private/WebSocketMessageHandler.cpp
+++ b/Engine/Plugins/VirtualProduction/RemoteControl/Source/WebRemoteControl/Private/WebSocketMessageHandler.cpp
@@ -1099,19 +1099,25 @@ void FWebSocketMessageHandler::OnConnectionClosedCallback(FGuid ClientId)
 		}
 	}
 
-	// Remove the client as a listener for any active transactions
+	// Remove the client as a listener for any active transactions and end them
+	// Do this in two separate loops since ending the transaction may modify the map
+	TSet<FGuid> TransactionGuidsToEnd;
 	for (auto Iter = ClientsByTransactionGuid.CreateIterator(); Iter; ++Iter)
 	{
-		IRemoteControlModule::Get().EndManualEditorTransaction(Iter.Key());
-
 		TMap<FGuid, FDateTime>& Clients = Iter.Value();
-		Clients.Remove(ClientId);
-		if (Clients.IsEmpty())
+		if (Clients.Remove(ClientId))
 		{
-			Iter.RemoveCurrent();
+			// Note that we always end the transaction even if other clients are listening, as it will just decrement
+			// the transaction's internal counter rather than actually ending it unless all transactions have been ended.
+			TransactionGuidsToEnd.Add(Iter.Key());
 		}
 	}
 
+	for (const FGuid& TransactionGuid : TransactionGuidsToEnd)
+	{
+		IRemoteControlModule::Get().EndManualEditorTransaction(TransactionGuid);
+	}
+
 	TransactionIdsByClientId.Remove(ClientId);
 
 	/** Remove this client's config. */
