Remote Control Preset Panel - FIX Focus Issue when Docked

Fixes issue whereas the Actor of a selected Property is deleted instead of the Property itself

diff --git a/Engine/Plugins/VirtualProduction/RemoteControl/Source/RemoteControlUI/Private/UI/SRemoteControlPanel.cpp b/Engine/Plugins/VirtualProduction/RemoteControl/Source/RemoteControlUI/Private/UI/SRemoteControlPanel.cpp
--- a/Engine/Plugins/VirtualProduction/RemoteControl/Source/RemoteControlUI/Private/UI/SRemoteControlPanel.cpp
+++ b/Engine/Plugins/VirtualProduction/RemoteControl/Source/RemoteControlUI/Private/UI/SRemoteControlPanel.cpp
@@ -2,8 +2,6 @@
 
 #include "SRemoteControlPanel.h"
 
-#include "IRemoteControlUIModule.h"
-#include "RemoteControlUIModule.h"
 #include "Action/SRCActionPanel.h"
 #include "ActorEditorUtils.h"
 #include "Behaviour/SRCBehaviourPanel.h"
@@ -14,24 +12,26 @@
 #include "Editor.h"
 #include "Editor/EditorPerformanceSettings.h"
 #include "EditorFontGlyphs.h"
-#include "EngineUtils.h"
 #include "Engine/Selection.h"
+#include "EngineUtils.h"
 #include "FileHelpers.h"
-#include "Layout/Visibility.h"
-#include "Input/Reply.h"
-#include "Interfaces/IMainFrameModule.h"
+#include "Framework/Application/IInputProcessor.h"
+#include "Framework/MultiBox/MultiBoxBuilder.h"
+#include "GameFramework/Actor.h"
 #include "IRemoteControlModule.h"
 #include "IRemoteControlProtocolWidgetsModule.h"
+#include "IRemoteControlUIModule.h"
 #include "ISettingsModule.h"
 #include "IStructureDetailsView.h"
+#include "Input/Reply.h"
+#include "Interfaces/IMainFrameModule.h"
 #include "Kismet/BlueprintFunctionLibrary.h"
-#include "Framework/MultiBox/MultiBoxBuilder.h"
-#include "GameFramework/Actor.h"
+#include "Layout/Visibility.h"
 #include "Materials/Material.h"
 #include "Modules/ModuleManager.h"
 #include "PropertyCustomizationHelpers.h"
-#include "PropertyHandle.h"
 #include "PropertyEditorModule.h"
+#include "PropertyHandle.h"
 #include "RCPanelWidgetRegistry.h"
 #include "RemoteControlActor.h"
 #include "RemoteControlEntity.h"
@@ -40,30 +40,31 @@
 #include "RemoteControlPanelStyle.h"
 #include "RemoteControlPreset.h"
 #include "RemoteControlSettings.h"
-#include "SceneOutlinerModule.h"
-#include "SceneOutlinerPublicTypes.h"
-#include "ScopedTransaction.h"
+#include "RemoteControlUIModule.h"
 #include "SClassViewer.h"
 #include "SRCLogger.h"
 #include "SRCModeSwitcher.h"
-#include "SRCPanelExposedEntitiesList.h"
-#include "SRCPanelFunctionPicker.h"
 #include "SRCPanelExposedActor.h"
+#include "SRCPanelExposedEntitiesList.h"
 #include "SRCPanelExposedField.h"
+#include "SRCPanelFunctionPicker.h"
 #include "SRCPanelTreeNode.h"
+#include "SWarningOrErrorBox.h"
+#include "SceneOutlinerModule.h"
+#include "SceneOutlinerPublicTypes.h"
+#include "ScopedTransaction.h"
 #include "Styling/RemoteControlStyles.h"
 #include "Styling/ToolBarStyle.h"
 #include "Subsystems/Subsystem.h"
-#include "SWarningOrErrorBox.h"
 #include "Templates/SharedPointer.h"
 #include "Templates/SubclassOf.h"
 #include "Templates/UnrealTypeTraits.h"
-#include "Toolkits/IToolkitHost.h"
 #include "ToolMenus.h"
+#include "Toolkits/IToolkitHost.h"
 #include "UI/BaseLogicUI/RCLogicModeBase.h"
 #include "UI/BaseLogicUI/SRCLogicPanelListBase.h"
-#include "UI/Filters/SRCPanelFilter.h"
 #include "UI/Drawers/SRCPanelDrawer.h"
+#include "UI/Filters/SRCPanelFilter.h"
 #include "UI/Panels/SRCDockPanel.h"
 #include "Widgets/DeclarativeSyntaxSupport.h"
 #include "Widgets/Docking/SDockTab.h"
@@ -72,8 +73,8 @@
 #include "Widgets/Input/SCheckBox.h"
 #include "Widgets/Layout/SBorder.h"
 #include "Widgets/Layout/SBox.h"
-#include "Widgets/Layout/SSplitter.h"
 #include "Widgets/Layout/SSeparator.h"
+#include "Widgets/Layout/SSplitter.h"
 #include "Widgets/Layout/SWidgetSwitcher.h"
 #include "Widgets/SBoxPanel.h"
 #include "Widgets/Text/STextBlock.h"
@@ -119,6 +120,48 @@ namespace RemoteControlPanelUtils
 	}
 }
 
+class FRemoteControlPanelInputProcessor : public IInputProcessor
+{
+public:
+	FRemoteControlPanelInputProcessor(TSharedPtr<SRemoteControlPanel> InOwner)
+	{
+		Owner = InOwner;
+	}
+
+	void SetOwner(TSharedPtr<SRemoteControlPanel> InOwner)
+	{
+		Owner = InOwner;
+	}
+
+	virtual ~FRemoteControlPanelInputProcessor() = default;
+
+	virtual void Tick(const float DeltaTime, FSlateApplication& SlateApp, TSharedRef<ICursor> Cursor) override
+	{
+	}
+
+	virtual bool HandleKeyDownEvent(FSlateApplication& SlateApp, const FKeyEvent& InKeyEvent) override
+	{
+		TSharedPtr<SRemoteControlPanel> PinnedOwner = Owner.Pin();
+		if (!PinnedOwner)
+		{
+			return false;
+		}
+		
+		if (PinnedOwner && InKeyEvent.GetKey() == EKeys::Delete && SlateApp.HasFocusedDescendants(PinnedOwner.ToSharedRef()))
+		{
+			PinnedOwner->DeleteEntity();
+			return true;
+		}
+		
+		return false;
+	}
+
+	virtual const TCHAR* GetDebugName() const override { return TEXT("RemoteControlInputInterceptor"); }
+
+private:
+	TWeakPtr<SRemoteControlPanel> Owner;
+};
+
 /**
  * UI representation of a auto resizing button.
  */
@@ -701,6 +744,10 @@ void SRemoteControlPanel::Construct(const FArguments& InArgs, URemoteControlPres
 			]
 		];
 
+
+	InputProcessor = MakeShared<FRemoteControlPanelInputProcessor>(SharedThis(this));
+	FSlateApplication::Get().RegisterInputPreProcessor(InputProcessor);
+
 	RegisterEvents();
 	RegisterPanels();
 	CacheLevelClasses();
@@ -716,6 +763,8 @@ SRemoteControlPanel::~SRemoteControlPanel()
 	UnregisterPanels();
 	UnregisterEvents();
 
+	FSlateApplication::Get().UnregisterInputPreProcessor(InputProcessor);
+
 	// Clear the log
 	FRemoteControlLogger::Get().ClearLog();
 
@@ -1598,7 +1647,7 @@ TSharedRef<SWidget> SRemoteControlPanel::CreateEntityDetailsView()
 		];
 
 	// Details Panel Icon
-	TSharedRef<SWidget> DetailsIcon = SNew(SImage)
+	const TSharedRef<SWidget> DetailsIcon = SNew(SImage)
 		.Image(FAppStyle::Get().GetBrush("LevelEditor.Tabs.Details"))
 		.ColorAndOpacity(FSlateColor::UseForeground());
 
@@ -1636,14 +1685,14 @@ void SRemoteControlPanel::UpdateEntityDetailsView(const TSharedPtr<SRCPanelTreeN
 	TSharedPtr<FStructOnScope> SelectedEntityPtr;
 
 	SelectedEntity = SelectedNode;
-
-	if (SelectedNode)
+	
+	if (SelectedEntity)
 	{
-		if (SelectedNode->GetRCType() != SRCPanelTreeNode::Group &&
-			SelectedNode->GetRCType() != SRCPanelTreeNode::FieldChild) // Field Child does not contain entity ID, that is why it should not be processed
+		if (SelectedEntity->GetRCType() != SRCPanelTreeNode::Group &&
+			SelectedEntity->GetRCType() != SRCPanelTreeNode::FieldChild) // Field Child does not contain entity ID, that is why it should not be processed
 		{
-			TSharedPtr<FRemoteControlEntity> Entity = Preset->GetExposedEntity<FRemoteControlEntity>(SelectedNode->GetRCId()).Pin();
-			SelectedEntityPtr = RemoteControlPanelUtils::GetEntityOnScope(Entity, Preset->GetExposedEntityType(SelectedNode->GetRCId()));
+			const TSharedPtr<FRemoteControlEntity> Entity = Preset->GetExposedEntity<FRemoteControlEntity>(SelectedEntity->GetRCId()).Pin();
+			SelectedEntityPtr = RemoteControlPanelUtils::GetEntityOnScope(Entity, Preset->GetExposedEntityType(SelectedEntity->GetRCId()));
 		}
 	}
 
@@ -1664,9 +1713,9 @@ void SRemoteControlPanel::UpdateEntityDetailsView(const TSharedPtr<SRCPanelTreeN
 	}
 
 	static const FName ProtocolWidgetsModuleName = "RemoteControlProtocolWidgets";	
-	if(SelectedNode.IsValid() && FModuleManager::Get().IsModuleLoaded(ProtocolWidgetsModuleName) && ensure(Preset.IsValid()))
+	if(SelectedEntity.IsValid() && FModuleManager::Get().IsModuleLoaded(ProtocolWidgetsModuleName) && ensure(Preset.IsValid()))
 	{
-		if (TSharedPtr<FRemoteControlEntity> RCEntity = Preset->GetExposedEntity(SelectedNode->GetRCId()).Pin())
+		if (const TSharedPtr<FRemoteControlEntity> RCEntity = Preset->GetExposedEntity(SelectedEntity->GetRCId()).Pin())
 		{
 			if(RCEntity->IsBound())
 			{
@@ -2421,4 +2470,10 @@ void SRemoteControlPanel::SaveSettings()
 	}
 }
 
+void SRemoteControlPanel::DeleteEntity()
+{
+	DeleteEntity_Execute();
+}
+
+
 #undef LOCTEXT_NAMESPACE /*RemoteControlPanel*/
diff --git a/Engine/Plugins/VirtualProduction/RemoteControl/Source/RemoteControlUI/Private/UI/SRemoteControlPanel.h b/Engine/Plugins/VirtualProduction/RemoteControl/Source/RemoteControlUI/Private/UI/SRemoteControlPanel.h
--- a/Engine/Plugins/VirtualProduction/RemoteControl/Source/RemoteControlUI/Private/UI/SRemoteControlPanel.h
+++ b/Engine/Plugins/VirtualProduction/RemoteControl/Source/RemoteControlUI/Private/UI/SRemoteControlPanel.h
@@ -143,6 +143,9 @@ public:
 	void AddToolbarWidget(TSharedRef<SWidget> Widget);
 	void RemoveAllToolbarWidgets();
 
+	/** Public Workaround Delete for Focus/Delete Key Issues when Docked */
+	void DeleteEntity();
+
 	/** Retrieves the Logic Action panel. */
 	TSharedPtr<SRCActionPanel> GetLogicActionPanel()
 	{
@@ -412,6 +415,8 @@ private:
 	const FRCPanelStyle* RCPanelStyle;
 	/** Stores the active panel that is drawn. */
 	ERCPanels ActivePanel;
+	/** Input Preprocessor which catches the Delete Key when Docked. */
+	TSharedPtr<IInputProcessor> InputProcessor;
 
 	// ~ Remote Control Logic Panels ~
 
