Remote Control Logic - Fix F2 Rename Button bind not working as intended

Fixes an issue where as the F2 Key does not force into Rename

diff --git a/Engine/Plugins/VirtualProduction/RemoteControl/Source/RemoteControlUI/Private/UI/SRemoteControlPanel.cpp b/Engine/Plugins/VirtualProduction/RemoteControl/Source/RemoteControlUI/Private/UI/SRemoteControlPanel.cpp
--- a/Engine/Plugins/VirtualProduction/RemoteControl/Source/RemoteControlUI/Private/UI/SRemoteControlPanel.cpp
+++ b/Engine/Plugins/VirtualProduction/RemoteControl/Source/RemoteControlUI/Private/UI/SRemoteControlPanel.cpp
@@ -141,16 +141,24 @@ public:
 
 	virtual bool HandleKeyDownEvent(FSlateApplication& SlateApp, const FKeyEvent& InKeyEvent) override
 	{
-		TSharedPtr<SRemoteControlPanel> PinnedOwner = Owner.Pin();
-		if (!PinnedOwner)
+		if (const TSharedPtr<SRemoteControlPanel> PinnedOwner = Owner.Pin())
 		{
-			return false;
-		}
-		
-		if (PinnedOwner && InKeyEvent.GetKey() == EKeys::Delete && SlateApp.HasFocusedDescendants(PinnedOwner.ToSharedRef()))
-		{
-			PinnedOwner->DeleteEntity();
-			return true;
+			if (!SlateApp.HasFocusedDescendants(PinnedOwner.ToSharedRef()))
+			{
+				return false;
+			}
+			
+			if (InKeyEvent.GetKey() == EKeys::Delete)
+			{
+				PinnedOwner->DeleteEntity();
+				return true;
+			}
+
+			if (InKeyEvent.GetKey() == EKeys::F2)
+			{
+				PinnedOwner->RenameEntity();
+				return true;
+			}
 		}
 		
 		return false;
@@ -2472,7 +2480,18 @@ void SRemoteControlPanel::SaveSettings()
 
 void SRemoteControlPanel::DeleteEntity()
 {
-	DeleteEntity_Execute();
+	if (CanDeleteEntity())
+	{
+		DeleteEntity_Execute();
+	}
+}
+
+void SRemoteControlPanel::RenameEntity()
+{
+	if (CanRenameEntity())
+	{
+		RenameEntity_Execute();
+	}
 }
 
 
diff --git a/Engine/Plugins/VirtualProduction/RemoteControl/Source/RemoteControlUI/Private/UI/SRemoteControlPanel.h b/Engine/Plugins/VirtualProduction/RemoteControl/Source/RemoteControlUI/Private/UI/SRemoteControlPanel.h
--- a/Engine/Plugins/VirtualProduction/RemoteControl/Source/RemoteControlUI/Private/UI/SRemoteControlPanel.h
+++ b/Engine/Plugins/VirtualProduction/RemoteControl/Source/RemoteControlUI/Private/UI/SRemoteControlPanel.h
@@ -143,8 +143,9 @@ public:
 	void AddToolbarWidget(TSharedRef<SWidget> Widget);
 	void RemoveAllToolbarWidgets();
 
-	/** Public Workaround Delete for Focus/Delete Key Issues when Docked */
+	/** Public Workaround Delete for Key Handling Issues when Docked */
 	void DeleteEntity();
+	void RenameEntity();
 
 	/** Retrieves the Logic Action panel. */
 	TSharedPtr<SRCActionPanel> GetLogicActionPanel()
