Enable/Disable Behaviour

diff --git a/Engine/Plugins/VirtualProduction/RemoteControl/Source/RemoteControlUI/Private/UI/Action/SRCActionPanel.cpp b/Engine/Plugins/VirtualProduction/RemoteControl/Source/RemoteControlUI/Private/UI/Action/SRCActionPanel.cpp
--- a/Engine/Plugins/VirtualProduction/RemoteControl/Source/RemoteControlUI/Private/UI/Action/SRCActionPanel.cpp
+++ b/Engine/Plugins/VirtualProduction/RemoteControl/Source/RemoteControlUI/Private/UI/Action/SRCActionPanel.cpp
@@ -170,52 +170,34 @@ void SRCActionPanel::UpdateWrappedWidget(TSharedPtr<FRCBehaviourModel> InBehavio
 		ActionDockPanel->AddHeaderToolbarItem(EToolbar::Right, AddAllActionsButton);
 		ActionDockPanel->AddHeaderToolbarItem(EToolbar::Right, DeleteSelectedActionButton);
 
-		// Duplicate Dock Panel
-		TSharedPtr<SRCMinorPanel> DuplicateDockPanel = SNew(SRCMinorPanel)
-			.HeaderLabel(LOCTEXT("DuplicateLabel", "Duplicate"))
+		// Header Dock Panel
+		TSharedPtr<SRCMinorPanel> HeaderDockPanel = SNew(SRCMinorPanel)
+			.HeaderLabel(LOCTEXT("HeaderDockLabel", ""))
 			[
-				SNew(SRCBehaviourDetails, SharedThis(this), InBehaviourItem.ToSharedRef())
-			];
-
-		// Duplicate Behaviour Button
-		const TSharedRef<SWidget> DuplicateBehaviourButton = SNew(SButton)
-			.AddMetaData<FTagMetaData>(FTagMetaData(TEXT("Duplicate Behavior")))
-			.HAlign(HAlign_Center)
-			.VAlign(VAlign_Center)
-			.ButtonStyle(&RCPanelStyle->FlatButtonStyle)
-			.ForegroundColor(FSlateColor::UseForeground())
-			.ContentPadding(FMargin(4.f, 2.f))
-			.OnClicked(this, &SRCActionPanel::OnClickOverrideBlueprintButton)
-			[
-				SNew(SBox)
-				.WidthOverride(RCPanelStyle->IconSize.X)
-				.HeightOverride(RCPanelStyle->IconSize.Y)
-				[
-					SNew(SImage)
-					.ColorAndOpacity(FSlateColor::UseForeground())
-					.Image(FAppStyle::GetBrush("Icons.PlusCircle"))
-				]
+				SAssignNew(BehaviourDetailsWidget, SRCBehaviourDetails, SharedThis(this), InBehaviourItem.ToSharedRef())
 			];
 		
 		// Toggle Behaviour Button
-		// TODO : Replace this with the actual algorithm to disable or enable behaviour.
-		const bool bIsChecked = FMath::RandBool();
-		const TSharedRef<SWidget> ToggleBehaviourButton = SNew(SCheckBox)
-			.AddMetaData<FTagMetaData>(FTagMetaData(TEXT("Duplicate Behavior")))
+		const bool bIsChecked = InBehaviourItem->IsBehaviourEnabled();
+		SAssignNew(ToggleBehaviourButton, SCheckBox)
+			.AddMetaData<FTagMetaData>(FTagMetaData(TEXT("Toggle Behavior")))
+			.ToolTipText(LOCTEXT("EditModeTooltip", "Enable/Disable this Behaviour.\nWhen a behaviour is disabled its Actions will not be processed when the Controller value changes"))
 			.HAlign(HAlign_Center)
 			.Style(&RCPanelStyle->ToggleButtonStyle)
 			.ForegroundColor(FSlateColor::UseForeground())
-			.IsChecked_Lambda([bIsChecked]() { return bIsChecked ? ECheckBoxState::Checked : ECheckBoxState::Unchecked; });
+			.IsChecked(bIsChecked)
+			.OnCheckStateChanged(this, &SRCActionPanel::OnToggleEnableBehaviour);
 
-		DuplicateDockPanel->AddHeaderToolbarItem(EToolbar::Left, DuplicateBehaviourButton);
-		DuplicateDockPanel->AddHeaderToolbarItem(EToolbar::Right, ToggleBehaviourButton);
+		RefreshIsBehaviourEnabled(bIsChecked);
+
+		HeaderDockPanel->AddHeaderToolbarItem(EToolbar::Left, ToggleBehaviourButton.ToSharedRef());
 
 		TSharedRef<SRCMajorPanel> ActionsPanel = SNew(SRCMajorPanel)
 			.EnableFooter(false)
 			.EnableHeader(false)
 			.ChildOrientation(Orient_Vertical);
 
-		ActionsPanel->AddPanel(DuplicateDockPanel.ToSharedRef(), 0.5f);
+		ActionsPanel->AddPanel(HeaderDockPanel.ToSharedRef(), 0.5f);
 		ActionsPanel->AddPanel(ActionDockPanel.ToSharedRef(), 0.5f);
 
 		WrappedBoxWidget->SetContent(ActionsPanel);
@@ -241,6 +223,31 @@ FReply SRCActionPanel::OnClickOverrideBlueprintButton()
 	return FReply::Handled();
 }
 
+void SRCActionPanel::OnToggleEnableBehaviour(ECheckBoxState State)
+{
+	SetIsBehaviourEnabled(State == ECheckBoxState::Checked);
+}
+
+void SRCActionPanel::SetIsBehaviourEnabled(const bool bIsEnabled)
+{
+	if (TSharedPtr<FRCBehaviourModel> BehaviourItem = GetSelectedBehaviourItem())
+	{
+		BehaviourItem->SetIsBehaviourEnabled(bIsEnabled);
+
+		RefreshIsBehaviourEnabled(bIsEnabled);
+	}
+}
+
+void SRCActionPanel::RefreshIsBehaviourEnabled(const bool bIsEnabled)
+{
+	if (ToggleBehaviourButton && BehaviourDetailsWidget && ActionPanelList)
+	{
+		ToggleBehaviourButton->SetIsChecked(bIsEnabled);
+		BehaviourDetailsWidget->SetEnabled(bIsEnabled);
+		ActionPanelList->SetEnabled(bIsEnabled);
+	}
+}
+
 TSharedRef<SWidget> SRCActionPanel::GetActionMenuContentWidget()
 {
 	if (AddNewActionMenuWidget && !bAddActionMenuNeedsRefresh)
diff --git a/Engine/Plugins/VirtualProduction/RemoteControl/Source/RemoteControlUI/Private/UI/Action/SRCActionPanel.h b/Engine/Plugins/VirtualProduction/RemoteControl/Source/RemoteControlUI/Private/UI/Action/SRCActionPanel.h
--- a/Engine/Plugins/VirtualProduction/RemoteControl/Source/RemoteControlUI/Private/UI/Action/SRCActionPanel.h
+++ b/Engine/Plugins/VirtualProduction/RemoteControl/Source/RemoteControlUI/Private/UI/Action/SRCActionPanel.h
@@ -4,14 +4,19 @@
 
 #include "UI/BaseLogicUI/SRCLogicPanelBase.h"
 
+enum class ECheckBoxState : uint8;
+
 struct FRCPanelStyle;
 class FRCActionModel;
 struct FRemoteControlField;
 class FRCBehaviourModel;
 class SBox;
+class SCheckBox;
+class SRCBehaviourDetails;
 class SRCLogicPanelListBase;
 class SRemoteControlPanel;
 class URCAction;
+class URCBehaviour;
 
 /*
 * ~ SRCActionPanel ~
@@ -69,6 +74,15 @@ public:
 		bAddActionMenuNeedsRefresh = true;
 	}
 
+	/** Fetches the currently selected behaviour item*/
+	TSharedPtr<FRCBehaviourModel> GetSelectedBehaviourItem()
+	{
+		return SelectedBehaviourItemWeakPtr.Pin();
+	}
+
+	/** Set the Enabled state of our parent Behaviour */
+	void SetIsBehaviourEnabled(const bool bIsEnabled);
+
 protected:
 	/** Warns user before deleting a selected panel item. */
 	virtual FReply RequestDeleteSelectedItem() override;
@@ -92,6 +106,12 @@ private:
 	/** Handles click event for Open Behaviour Blueprint button*/
 	FReply OnClickOverrideBlueprintButton();
 
+	/** Handles click event for toggling Enable/Disable behaviour*/
+	void OnToggleEnableBehaviour(ECheckBoxState State);
+
+	/** Refreshes the UI widgets enabled state depending on whether the parent behaviour is currently enabled */
+	void RefreshIsBehaviourEnabled(const bool bIsEnabled);
+
 	/**
 	 * Builds a menu containing the list of all possible Actions
 	 * These are derived from the list of Exposed entities of the Remote Control Preset associated with us.
@@ -128,6 +148,12 @@ private:
 	/** Widget representing List of Actions */
 	TSharedPtr<SRCLogicPanelListBase> ActionPanelList;
 
+	/** Behaviour specific details widget*/
+	TSharedPtr<SRCBehaviourDetails> BehaviourDetailsWidget;
+
+	/** Toggle Behaviour button*/
+	TSharedPtr<SCheckBox> ToggleBehaviourButton;
+
 	/** Helper widget for behavior details. */
 	static TSharedPtr<SBox> NoneSelectedWidget;
 
diff --git a/Engine/Plugins/VirtualProduction/RemoteControl/Source/RemoteControlUI/Private/UI/BaseLogicUI/SRCLogicPanelListBase.h b/Engine/Plugins/VirtualProduction/RemoteControl/Source/RemoteControlUI/Private/UI/BaseLogicUI/SRCLogicPanelListBase.h
--- a/Engine/Plugins/VirtualProduction/RemoteControl/Source/RemoteControlUI/Private/UI/BaseLogicUI/SRCLogicPanelListBase.h
+++ b/Engine/Plugins/VirtualProduction/RemoteControl/Source/RemoteControlUI/Private/UI/BaseLogicUI/SRCLogicPanelListBase.h
@@ -60,6 +60,11 @@ public:
 	/** Allows Logic panels to add special functionality to the Context Menu based on context */
 	virtual void AddSpecialContextMenuOptions(FMenuBuilder& MenuBuilder) {};
 
+	TSharedPtr<SRemoteControlPanel> GetRemoteControlPanel()
+	{
+		return RemoteControlPanelWeakPtr.Pin();
+	}
+
 private:
 
 	/** Refreshes the list from the latest state of the model*/
diff --git a/Engine/Plugins/VirtualProduction/RemoteControl/Source/RemoteControlUI/Private/UI/Behaviour/RCBehaviourModel.cpp b/Engine/Plugins/VirtualProduction/RemoteControl/Source/RemoteControlUI/Private/UI/Behaviour/RCBehaviourModel.cpp
--- a/Engine/Plugins/VirtualProduction/RemoteControl/Source/RemoteControlUI/Private/UI/Behaviour/RCBehaviourModel.cpp
+++ b/Engine/Plugins/VirtualProduction/RemoteControl/Source/RemoteControlUI/Private/UI/Behaviour/RCBehaviourModel.cpp
@@ -16,18 +16,23 @@
 
 #define LOCTEXT_NAMESPACE "FRCBehaviourModel"
 
-FRCBehaviourModel::FRCBehaviourModel(URCBehaviour* InBehaviour,
-	const TSharedPtr<SRemoteControlPanel> InRemoteControlPanel /*= nullptr*/)
-	: FRCLogicModeBase(InRemoteControlPanel),
-	BehaviourWeakPtr(InBehaviour)
+FRCBehaviourModel::FRCBehaviourModel(URCBehaviour* InBehaviour
+	, const TSharedPtr<SRemoteControlPanel> InRemoteControlPanel /*= nullptr*/)
+	: FRCLogicModeBase(InRemoteControlPanel)
+	, BehaviourWeakPtr(InBehaviour)
 {
 	RCPanelStyle = &FRemoteControlPanelStyle::Get()->GetWidgetStyle<FRCPanelStyle>("RemoteControlPanel.BehaviourPanel");
 
-	const FText BehaviorDisplayName = BehaviourWeakPtr->GetDisplayName().ToUpper();
+	if (BehaviourWeakPtr.IsValid())
+	{
+		const FText BehaviorDisplayName = BehaviourWeakPtr->GetDisplayName().ToUpper();
+		
+		SAssignNew(BehaviourTitleText, STextBlock)
+			.Text(BehaviorDisplayName)
+			.TextStyle(&RCPanelStyle->HeaderTextStyle);
 
-	SAssignNew(BehaviourTitleText, STextBlock)
-		.Text(BehaviorDisplayName)
-		.TextStyle(&RCPanelStyle->HeaderTextStyle);
+		RefreshIsBehaviourEnabled(BehaviourWeakPtr->bIsEnabled);
+	}
 }
 
 URCAction* FRCBehaviourModel::AddAction(const TSharedRef<const FRemoteControlField> InRemoteControlField)
@@ -99,10 +104,15 @@ void FRCBehaviourModel::SetIsBehaviourEnabled(const bool bIsEnabled)
 	{
 		Behaviour->bIsEnabled = bIsEnabled;
 
-		BehaviourTitleText->SetEnabled(bIsEnabled);
+		RefreshIsBehaviourEnabled(bIsEnabled);
 	}
 }
 
+void FRCBehaviourModel::RefreshIsBehaviourEnabled(const bool bIsEnabled)
+{
+	BehaviourTitleText->SetEnabled(bIsEnabled);
+}
+
 TSharedPtr<SRCLogicPanelListBase> FRCBehaviourModel::GetActionsListWidget(TSharedRef<SRCActionPanel> InActionPanel)
 {
 	// Returns the default Actions List; child classes can override as required
diff --git a/Engine/Plugins/VirtualProduction/RemoteControl/Source/RemoteControlUI/Private/UI/Behaviour/RCBehaviourModel.h b/Engine/Plugins/VirtualProduction/RemoteControl/Source/RemoteControlUI/Private/UI/Behaviour/RCBehaviourModel.h
--- a/Engine/Plugins/VirtualProduction/RemoteControl/Source/RemoteControlUI/Private/UI/Behaviour/RCBehaviourModel.h
+++ b/Engine/Plugins/VirtualProduction/RemoteControl/Source/RemoteControlUI/Private/UI/Behaviour/RCBehaviourModel.h
@@ -46,6 +46,9 @@ public:
 	/** Set the Enabled state of our underlying Behaviour */
 	void SetIsBehaviourEnabled(const bool bIsEnabled);
 
+	/** Refreshes the UI widgets enabled state depending on whether the parent behaviour is currently enabled */
+	void RefreshIsBehaviourEnabled(const bool bIsEnabled);
+
 	/** The Actions List widget to be used for this Behaviour.
 	* The default actions table is used if a behaviour doesn't specify it.
 	*/
diff --git a/Engine/Plugins/VirtualProduction/RemoteControl/Source/RemoteControlUI/Private/UI/Behaviour/SRCBehaviourPanelList.cpp b/Engine/Plugins/VirtualProduction/RemoteControl/Source/RemoteControlUI/Private/UI/Behaviour/SRCBehaviourPanelList.cpp
--- a/Engine/Plugins/VirtualProduction/RemoteControl/Source/RemoteControlUI/Private/UI/Behaviour/SRCBehaviourPanelList.cpp
+++ b/Engine/Plugins/VirtualProduction/RemoteControl/Source/RemoteControlUI/Private/UI/Behaviour/SRCBehaviourPanelList.cpp
@@ -11,6 +11,7 @@
 #include "SRCBehaviourPanel.h"
 #include "SlateOptMacros.h"
 #include "Styling/RemoteControlStyles.h"
+#include "UI/Action/SRCActionPanel.h"
 #include "UI/Behaviour/Builtin/Bind/RCBehaviourBindModel.h"
 #include "UI/Behaviour/Builtin/Conditional/RCBehaviourConditionalModel.h"
 #include "UI/Behaviour/Builtin/Path/RCBehaviourSetAssetByPathModel.h"
@@ -97,6 +98,41 @@ void SRCBehaviourPanelList::AddNewLogicItem(UObject* InLogicItem)
 	RequestRefresh();
 }
 
+void SRCBehaviourPanelList::AddSpecialContextMenuOptions(FMenuBuilder& MenuBuilder)
+{
+	if (TSharedPtr<FRCBehaviourModel> BehaviourItem = GetSelectedBehaviourItem())
+	{
+		const bool bIsEnabled = BehaviourItem->IsBehaviourEnabled();
+
+		FUIAction Action(FExecuteAction::CreateSP(this, &SRCBehaviourPanelList::SetIsBehaviourEnabled, !bIsEnabled));
+
+		FText Label, Tooltip;
+		if (bIsEnabled)
+		{
+			Label = LOCTEXT("LabelDisableBehaviour", "Disable");
+			Tooltip = LOCTEXT("TooltipDisableBehaviour", "Disables the current behaviour. Actions will not be processed for this behaviour when the Controller value changes");
+		}
+		else
+		{
+			Label = LOCTEXT("LabelEnableBehaviour", "Enable");
+			Tooltip = LOCTEXT("TooltipDisableBehaviour", "Enables the current behaviour. Restores functioning of Actions associated with this behaviour.");
+		}
+
+		MenuBuilder.AddMenuEntry(Label, Tooltip, FSlateIcon(), Action);
+	}
+}
+
+void SRCBehaviourPanelList::SetIsBehaviourEnabled(const bool bIsEnabled)
+{
+	if (TSharedPtr<SRemoteControlPanel> RemoteControlPanel = GetRemoteControlPanel())
+	{
+		if (TSharedPtr<SRCActionPanel> ActionPanel = RemoteControlPanel->GetLogicActionPanel())
+		{
+			ActionPanel->SetIsBehaviourEnabled(bIsEnabled);
+		}
+	}
+}
+
 void SRCBehaviourPanelList::AddBehaviourToList(URCBehaviour* InBehaviour)
 {
 	if (!BehaviourPanelWeakPtr.IsValid())
diff --git a/Engine/Plugins/VirtualProduction/RemoteControl/Source/RemoteControlUI/Private/UI/Behaviour/SRCBehaviourPanelList.h b/Engine/Plugins/VirtualProduction/RemoteControl/Source/RemoteControlUI/Private/UI/Behaviour/SRCBehaviourPanelList.h
--- a/Engine/Plugins/VirtualProduction/RemoteControl/Source/RemoteControlUI/Private/UI/Behaviour/SRCBehaviourPanelList.h
+++ b/Engine/Plugins/VirtualProduction/RemoteControl/Source/RemoteControlUI/Private/UI/Behaviour/SRCBehaviourPanelList.h
@@ -64,8 +64,14 @@ public:
 
 	virtual void AddNewLogicItem(UObject* InLogicItem) override;
 
+	/** Allows Logic panels to add special functionality to the Context Menu based on context */
+	virtual void AddSpecialContextMenuOptions(FMenuBuilder& MenuBuilder) override;
+
 private:
 
+	/** Enables or Disables the currently selected behaviour*/
+	void SetIsBehaviourEnabled(const bool bIsEnabled);
+
 	void AddBehaviourToList(URCBehaviour* InBehaviour);
 
 	/** OnGenerateRow delegate for the Behaviours List View */
