Remove use of deprecated field.

diff --git a/Engine/Plugins/VirtualProduction/RemoteControl/Source/WebRemoteControl/Private/RemoteControlModels.h b/Engine/Plugins/VirtualProduction/RemoteControl/Source/WebRemoteControl/Private/RemoteControlModels.h
--- a/Engine/Plugins/VirtualProduction/RemoteControl/Source/WebRemoteControl/Private/RemoteControlModels.h
+++ b/Engine/Plugins/VirtualProduction/RemoteControl/Source/WebRemoteControl/Private/RemoteControlModels.h
@@ -551,9 +551,7 @@ struct FRCShortPresetDescription
 		}
 
 		ID = PresetId.ToString();
-PRAGMA_DISABLE_DEPRECATION_WARNINGS
-		Path = PresetAsset.ObjectPath;
-PRAGMA_ENABLE_DEPRECATION_WARNINGS
+		Path = *PresetAsset.GetObjectPathString();
 	}
 
 	FRCShortPresetDescription(const URemoteControlPreset* InPreset)
@@ -604,9 +602,7 @@ struct FRCAssetDescription
 	FRCAssetDescription(const FAssetData& InAsset)
 		: Name(InAsset.AssetName)
 		, Class(InAsset.AssetClassPath)
-PRAGMA_DISABLE_DEPRECATION_WARNINGS
-		, Path(InAsset.ObjectPath)
-PRAGMA_ENABLE_DEPRECATION_WARNINGS
+		, Path(*InAsset.GetObjectPathString())
 	{
 		Metadata = RemoteControlModels::SanitizeAssetMetadata(InAsset.TagsAndValues.CopyMap());
 	}
