1. Remote API Protocols - Masking Limitation

2. MIDI Muting/Isolating Channels

diff --git a/Engine/Plugins/VirtualProduction/RemoteControl/Source/RemoteControl/Private/RemoteControlField.cpp b/Engine/Plugins/VirtualProduction/RemoteControl/Source/RemoteControl/Private/RemoteControlField.cpp
--- a/Engine/Plugins/VirtualProduction/RemoteControl/Source/RemoteControl/Private/RemoteControlField.cpp
+++ b/Engine/Plugins/VirtualProduction/RemoteControl/Source/RemoteControl/Private/RemoteControlField.cpp
@@ -321,19 +321,6 @@ bool FRemoteControlProperty::IsBound() const
 	return !!GetProperty();
 }
 
-bool FRemoteControlProperty::HasOptionalMask() const
-{
-	if (const FStructProperty* StructProperty = CastField<FStructProperty>(GetProperty()))
-	{
-		return StructProperty->Struct == TBaseStructure<FLinearColor>::Get() ||
-			StructProperty->Struct == TBaseStructure<FVector4>::Get() ||
-			StructProperty->Struct == TBaseStructure<FQuat>::Get() ||
-			StructProperty->Struct == TBaseStructure<FIntVector4>::Get();
-	}
-
-	return FRemoteControlField::HasOptionalMask();
-}
-
 bool FRemoteControlProperty::SupportsMasking() const
 {
 	return IRemoteControlModule::Get().SupportsMasking(GetProperty());
diff --git a/Engine/Plugins/VirtualProduction/RemoteControl/Source/RemoteControl/Private/RemoteControlProtocolBinding.cpp b/Engine/Plugins/VirtualProduction/RemoteControl/Source/RemoteControl/Private/RemoteControlProtocolBinding.cpp
--- a/Engine/Plugins/VirtualProduction/RemoteControl/Source/RemoteControl/Private/RemoteControlProtocolBinding.cpp
+++ b/Engine/Plugins/VirtualProduction/RemoteControl/Source/RemoteControl/Private/RemoteControlProtocolBinding.cpp
@@ -698,7 +698,15 @@ bool FRemoteControlProtocolEntity::ApplyProtocolValueToProperty(double InProtoco
 		IRemoteControlModule::Get().ResolveObjectProperty(ObjectRef.Access, Object, ObjectRef.PropertyPathInfo, ObjectRef);
 
 		TSharedPtr<FRCMaskingOperation> MaskingOperation = MakeShared<FRCMaskingOperation>(ObjectRef.PropertyPathInfo, Object);
-		MaskingOperation->Masks = RemoteControlProperty->GetActiveMasks();
+
+		if (OverridenMasks == ERCMask::NoMask)
+		{
+			MaskingOperation->Masks = RemoteControlProperty->GetActiveMasks();
+		}
+		else
+		{
+			MaskingOperation->Masks = OverridenMasks;
+		}
 
 		// Cache the values.
 		IRemoteControlModule::Get().PerformMasking(MaskingOperation.ToSharedRef());
@@ -738,6 +746,21 @@ void FRemoteControlProtocolEntity::ResetDefaultBindingState()
 	BindingStatus = ERCBindingStatus::Unassigned;
 }
 
+void FRemoteControlProtocolEntity::ClearMask(ERCMask InMaskBit)
+{
+	OverridenMasks &= ~InMaskBit;
+}
+
+void FRemoteControlProtocolEntity::EnableMask(ERCMask InMaskBit)
+{
+	OverridenMasks |= InMaskBit;
+}
+
+bool FRemoteControlProtocolEntity::HasMask(ERCMask InMaskBit) const
+{
+	return (OverridenMasks & InMaskBit) != ERCMask::NoMask;
+}
+
 #if WITH_EDITOR
 
 const FName FRemoteControlProtocolEntity::GetPropertyName(const FName& ForColumnName)
diff --git a/Engine/Plugins/VirtualProduction/RemoteControl/Source/RemoteControl/Public/Factories/IRemoteControlMaskingFactory.h b/Engine/Plugins/VirtualProduction/RemoteControl/Source/RemoteControl/Public/Factories/IRemoteControlMaskingFactory.h
--- a/Engine/Plugins/VirtualProduction/RemoteControl/Source/RemoteControl/Public/Factories/IRemoteControlMaskingFactory.h
+++ b/Engine/Plugins/VirtualProduction/RemoteControl/Source/RemoteControl/Public/Factories/IRemoteControlMaskingFactory.h
@@ -51,7 +51,7 @@ public:
 	FGuid OperationId;
 
 	/** Masks to be applied. */
-	ERCMask Masks = ERCMask::NoMask;
+	ERCMask Masks = RC_AllMasks;
 
 	/** Holds Object reference. */
 	FRCObjectReference ObjectRef;
diff --git a/Engine/Plugins/VirtualProduction/RemoteControl/Source/RemoteControl/Public/RemoteControlField.h b/Engine/Plugins/VirtualProduction/RemoteControl/Source/RemoteControl/Public/RemoteControlField.h
--- a/Engine/Plugins/VirtualProduction/RemoteControl/Source/RemoteControl/Public/RemoteControlField.h
+++ b/Engine/Plugins/VirtualProduction/RemoteControl/Source/RemoteControl/Public/RemoteControlField.h
@@ -2,6 +2,7 @@
 #pragma once
 
 #include "CoreMinimal.h"
+#include "RemoteControlCommon.h"
 #include "RemoteControlEntity.h"
 #include "RemoteControlFieldPath.h"
 #include "RemoteControlProtocolBinding.h"
@@ -10,33 +11,6 @@
 class IRemoteControlPropertyHandle;
 struct FPropertyChangedEvent;
 
-UENUM()
-namespace ERCProtocolBinding
-{
-	enum Op
-	{
-		Added,
-
-		Removed
-	};
-}
-
-UENUM()
-enum class ERCMask : uint8
-{
-	NoMask = 0x00000000,
-
-	MaskA = 0x00000001,
-
-	MaskB = 0x00000002,
-
-	MaskC = 0x00000004,
-
-	MaskD = 0x00000008,
-};
-
-ENUM_CLASS_FLAGS(ERCMask);
-
 /**
  * The type of the exposed field.
  */
@@ -77,8 +51,6 @@ struct REMOTECONTROL_API FRemoteControlField : public FRemoteControlEntity
 	virtual void EnableMask(ERCMask InMaskBit);
 	/** Returns true if the given mask is enabled, false otherwise. */
 	virtual bool HasMask(ERCMask InMaskBit) const;
-	/** Returns true if this field has optional mask. */
-	virtual bool HasOptionalMask() const { return false; }
 	/** Returns true if this field supports masking. */
 	virtual bool SupportsMasking() const { return false; }
 
@@ -134,7 +106,7 @@ protected:
 	
 	/** Holds the actively enabled masks. */
 	UPROPERTY()
-	ERCMask ActiveMasks = ERCMask::NoMask;
+	ERCMask ActiveMasks = RC_AllMasks;
 
 protected:
 	FRemoteControlField(URemoteControlPreset* InPreset, EExposedFieldType InType, FName InLabel, FRCFieldPathInfo InFieldPathInfo, const TArray<URemoteControlBinding*> InBindings);
@@ -174,7 +146,6 @@ public:
 	//~ End FRemoteControlEntity interface
 	
 	//~ Begin FRemoteControlField interface
-	virtual bool HasOptionalMask() const override;
 	virtual bool SupportsMasking() const override;
 	//~ End FRemoteControlField interface
 
diff --git a/Engine/Plugins/VirtualProduction/RemoteControl/Source/RemoteControl/Public/RemoteControlProtocolBinding.h b/Engine/Plugins/VirtualProduction/RemoteControl/Source/RemoteControl/Public/RemoteControlProtocolBinding.h
--- a/Engine/Plugins/VirtualProduction/RemoteControl/Source/RemoteControl/Public/RemoteControlProtocolBinding.h
+++ b/Engine/Plugins/VirtualProduction/RemoteControl/Source/RemoteControl/Public/RemoteControlProtocolBinding.h
@@ -4,6 +4,7 @@
 
 #include "CoreMinimal.h"
 
+#include "RemoteControlCommon.h"
 #include "RCPropertyUtilities.h"
 #include "RCTypeTraits.h"
 #include "Containers/UnrealString.h"
@@ -426,6 +427,15 @@ public:
 	 */
 	virtual bool IsSame(const FRemoteControlProtocolEntity* InOther) { return false; }
 
+	/** Masking Support */
+
+	/** Disables the given mask. */
+	virtual void ClearMask(ERCMask InMaskBit);
+	/** Enables the given mask. */
+	virtual void EnableMask(ERCMask InMaskBit);
+	/** Returns true if the given mask is enabled, false otherwise. */
+	virtual bool HasMask(ERCMask InMaskBit) const;
+
 #if WITH_EDITOR
 
 	/** Retrieves the name of property corresponding to the given column name. */
@@ -535,7 +545,11 @@ protected:
 	 */
 	UPROPERTY()
 	TSet<FRemoteControlProtocolMapping> Mappings;
-	
+
+	/** Holds the overriden masks. */
+	UPROPERTY()
+	ERCMask OverridenMasks = ERCMask::NoMask;
+
 #if WITH_EDITOR
 
 	/**
@@ -657,7 +671,7 @@ public:
 	/** Custom struct serialize */
 	bool Serialize(FArchive& Ar);
 	friend FArchive& operator<<(FArchive& Ar, FRemoteControlProtocolBinding& InProtocolBinding);
-
+	
 private:
 	/** Return FRemoteControlProtocolEntity pointer from RemoteControlProtocolEntityPtr */
 	FRemoteControlProtocolEntity* GetRemoteControlProtocolEntity();
diff --git a/Engine/Plugins/VirtualProduction/RemoteControl/Source/RemoteControlCommon/Public/RemoteControlCommon.h b/Engine/Plugins/VirtualProduction/RemoteControl/Source/RemoteControlCommon/Public/RemoteControlCommon.h
new file mode 100644
--- /dev/null
+++ b/Engine/Plugins/VirtualProduction/RemoteControl/Source/RemoteControlCommon/Public/RemoteControlCommon.h
@@ -0,0 +1,37 @@
+// Copyright Epic Games, Inc. All Rights Reserved.
+
+#pragma once
+
+#include "CoreMinimal.h"
+
+#include "RemoteControlCommon.generated.h"
+
+UENUM()
+namespace ERCProtocolBinding
+{
+	enum Op
+	{
+		Added,
+
+		Removed
+	};
+}
+
+UENUM()
+enum class ERCMask : uint8
+{
+	NoMask = 0x00,
+
+	MaskA = 0x01,
+
+	MaskB = 0x02,
+
+	MaskC = 0x04,
+
+	MaskD = 0x08,
+};
+
+/** Mask for all remote control masking flags */
+#define RC_AllMasks (ERCMask)0xFF // All masks, used mainly for initialization
+
+ENUM_CLASS_FLAGS(ERCMask);
diff --git a/Engine/Plugins/VirtualProduction/RemoteControl/Source/RemoteControlProtocolWidgets/Private/Widgets/Masks/SRCProtocolBindingMask.cpp b/Engine/Plugins/VirtualProduction/RemoteControl/Source/RemoteControlProtocolWidgets/Private/Widgets/Masks/SRCProtocolBindingMask.cpp
new file mode 100644
--- /dev/null
+++ b/Engine/Plugins/VirtualProduction/RemoteControl/Source/RemoteControlProtocolWidgets/Private/Widgets/Masks/SRCProtocolBindingMask.cpp
@@ -0,0 +1,118 @@
+// Copyright Epic Games, Inc. All Rights Reserved.
+
+#include "SRCProtocolBindingMask.h"
+#include "IRemoteControlModule.h"
+#include "RemoteControlProtocolBinding.h"
+#include "ViewModels/ProtocolBindingViewModel.h"
+
+#define LOCTEXT_NAMESPACE "SRCProtocolBindingMask"
+
+void SRCProtocolBindingMask::Construct(const FArguments& InArgs, const TSharedRef<FProtocolBindingViewModel>& InViewModel)
+{
+	ViewModel = InViewModel;
+
+	const bool bHasOptionalMask = HasOptionalMask();
+
+	const bool bCanBeMasked = CanBeMasked();
+
+	SRCProtocolMaskTriplet::Construct(SRCProtocolMaskTriplet::FArguments()
+		.MaskA(ERCMask::MaskA)
+		.MaskB(ERCMask::MaskB)
+		.MaskC(ERCMask::MaskC)
+		.OptionalMask(ERCMask::MaskD)
+		.MaskingType(GetMaskingType())
+		.CanBeMasked(bCanBeMasked)
+		.EnableOptionalMask(bHasOptionalMask)
+	);
+}
+
+bool SRCProtocolBindingMask::CanBeMasked() const
+{
+	if (ViewModel.IsValid())
+	{
+		if (const FProperty* Property = ViewModel->GetProperty().Get())
+		{
+			return IRemoteControlModule::Get().SupportsMasking(Property);
+		}
+	}
+
+	return false;
+}
+
+EMaskingType SRCProtocolBindingMask::GetMaskingType()
+{
+	if (CanBeMasked())
+	{
+		if (const FProperty* Property = ViewModel->GetProperty().Get())
+		{
+			if (const FStructProperty* StructProperty = CastField<FStructProperty>(Property))
+			{
+				if (const EMaskingType* StructToMaskingType = RemoteControlProtocolMasking::StructsToMaskingTypes.Find(StructProperty->Struct))
+				{
+					return *StructToMaskingType;
+				}
+			}
+		}
+	}
+
+	return EMaskingType::Unsupported;
+}
+
+bool SRCProtocolBindingMask::HasOptionalMask() const
+{
+	if (ViewModel.IsValid())
+	{
+		if (const FProperty* Property = ViewModel->GetProperty().Get())
+		{
+			if (const FStructProperty* StructProperty = CastField<FStructProperty>(Property))
+			{
+				return RemoteControlProtocolMasking::OptionalMaskStructs.Contains(StructProperty->Struct);
+			}
+		}
+	}
+
+	return false;
+}
+
+ECheckBoxState SRCProtocolBindingMask::IsMaskEnabled(ERCMask InMaskBit) const
+{
+	if (ViewModel.IsValid())
+	{
+		if (const FRemoteControlProtocolBinding* RCProtocolBinding = ViewModel->GetBinding())
+		{
+			if (TSharedPtr<TStructOnScope<FRemoteControlProtocolEntity>> RCProtocolEntity = RCProtocolBinding->GetRemoteControlProtocolEntityPtr())
+			{
+				return (*RCProtocolEntity)->HasMask((ERCMask)InMaskBit) ? ECheckBoxState::Checked : ECheckBoxState::Unchecked;
+			}
+		}
+	}
+
+	return SRCProtocolMaskTriplet::IsMaskEnabled(InMaskBit);
+}
+
+void SRCProtocolBindingMask::SetMaskEnabled(ECheckBoxState NewState, ERCMask NewMaskBit)
+{
+	if (ViewModel.IsValid())
+	{
+		if (const FRemoteControlProtocolBinding* RCProtocolBinding = ViewModel->GetBinding())
+		{
+			if (TSharedPtr<TStructOnScope<FRemoteControlProtocolEntity>> RCProtocolEntity = RCProtocolBinding->GetRemoteControlProtocolEntityPtr())
+			{
+				if (NewState == ECheckBoxState::Checked)
+				{
+					(*RCProtocolEntity)->EnableMask((ERCMask)NewMaskBit);
+				}
+				else
+				{
+					(*RCProtocolEntity)->ClearMask((ERCMask)NewMaskBit);
+				}
+			}
+		}
+		else
+		{
+			SRCProtocolMaskTriplet::SetMaskEnabled(NewState, NewMaskBit);
+		}
+	}
+}
+
+#undef LOCTEXT_NAMESPACE
diff --git a/Engine/Plugins/VirtualProduction/RemoteControl/Source/RemoteControlProtocolWidgets/Private/Widgets/Masks/SRCProtocolBindingMask.h b/Engine/Plugins/VirtualProduction/RemoteControl/Source/RemoteControlProtocolWidgets/Private/Widgets/Masks/SRCProtocolBindingMask.h
new file mode 100644
--- /dev/null
+++ b/Engine/Plugins/VirtualProduction/RemoteControl/Source/RemoteControlProtocolWidgets/Private/Widgets/Masks/SRCProtocolBindingMask.h
@@ -0,0 +1,42 @@
+// Copyright Epic Games, Inc. All Rights Reserved.
+
+#pragma once
+
+#include "CoreMinimal.h"
+#include "Widgets/Masks/SRCProtocolMaskTriplet.h"
+
+enum class ERCMask : uint8;
+class FProtocolBindingViewModel;
+
+/**
+ * A widget that represents a mask (specific to protocol bindings).
+ */
+class SRCProtocolBindingMask : public SRCProtocolMaskTriplet, public IHasMaskExtensibility
+{
+public:
+	SLATE_BEGIN_ARGS(SRCProtocolBindingMask)
+	{}
+
+	SLATE_END_ARGS()
+
+	/** Constructs this widget with InArgs */
+	void Construct(const FArguments& InArgs, const TSharedRef<FProtocolBindingViewModel>& InViewModel);
+
+protected:
+
+	//~ Begin : IHasMaskExtensibility Interface
+	virtual bool CanBeMasked() const override;
+	virtual EMaskingType GetMaskingType() override;
+	virtual bool HasOptionalMask() const override;
+	//~ End : IHasMaskExtensibility Interface
+
+	//~ Begin : SRCProtocolMaskTriplet Interface
+	virtual ECheckBoxState IsMaskEnabled(ERCMask InMaskBit) const override;
+	virtual void SetMaskEnabled(ECheckBoxState NewState, ERCMask NewMaskBit) override;
+	//~ End : SRCProtocolMaskTriplet Interface
+
+private:
+
+	/** ViewModel for the Protocol Binding. */
+	TSharedPtr<FProtocolBindingViewModel> ViewModel;
+};
diff --git a/Engine/Plugins/VirtualProduction/RemoteControl/Source/RemoteControlProtocolWidgets/Private/Widgets/Masks/SRCProtocolMask.cpp b/Engine/Plugins/VirtualProduction/RemoteControl/Source/RemoteControlProtocolWidgets/Private/Widgets/Masks/SRCProtocolMask.cpp
--- a/Engine/Plugins/VirtualProduction/RemoteControl/Source/RemoteControlProtocolWidgets/Private/Widgets/Masks/SRCProtocolMask.cpp
+++ b/Engine/Plugins/VirtualProduction/RemoteControl/Source/RemoteControlProtocolWidgets/Private/Widgets/Masks/SRCProtocolMask.cpp
@@ -6,19 +6,6 @@
 
 #define LOCTEXT_NAMESPACE "SRCProtocolMask"
 
-namespace RemoteControlProtocolMasking
-{
-	static TMap<UScriptStruct*, EMaskingType> StructsToMaskingTypes = { { TBaseStructure<FColor>::Get(), EMaskingType::Color }
-		, { TBaseStructure<FLinearColor>::Get(), EMaskingType::Color }
-		, { TBaseStructure<FRotator>::Get(), EMaskingType::Rotator }
-		, { TBaseStructure<FVector>::Get(), EMaskingType::Vector }
-		, { TBaseStructure<FIntVector>::Get(), EMaskingType::Vector }
-		, { TBaseStructure<FVector4>::Get(), EMaskingType::Quat }
-		, { TBaseStructure<FIntVector4>::Get(), EMaskingType::Quat }
-		, { TBaseStructure<FQuat>::Get(), EMaskingType::Quat }
-	};
-}
-
 void SRCProtocolMask::Construct(const FArguments& InArgs, TWeakPtr<FRemoteControlField> InField)
 {
 	WeakField = MoveTemp(InField);
@@ -74,7 +61,13 @@ bool SRCProtocolMask::HasOptionalMask() const
 {
 	if (TSharedPtr<FRemoteControlField> RCField = WeakField.Pin())
 	{
-		return RCField->HasOptionalMask();
+		if (TSharedPtr<FRemoteControlProperty> RCProperty = StaticCastSharedPtr<FRemoteControlProperty>(RCField))
+		{
+			if (const FStructProperty* StructProperty = CastField<FStructProperty>(RCProperty->GetProperty()))
+			{
+				return RemoteControlProtocolMasking::OptionalMaskStructs.Contains(StructProperty->Struct);
+			}
+		}
 	}
 
 	return false;
diff --git a/Engine/Plugins/VirtualProduction/RemoteControl/Source/RemoteControlProtocolWidgets/Private/Widgets/SRCProtocolBinding.cpp b/Engine/Plugins/VirtualProduction/RemoteControl/Source/RemoteControlProtocolWidgets/Private/Widgets/SRCProtocolBinding.cpp
--- a/Engine/Plugins/VirtualProduction/RemoteControl/Source/RemoteControlProtocolWidgets/Private/Widgets/SRCProtocolBinding.cpp
+++ b/Engine/Plugins/VirtualProduction/RemoteControl/Source/RemoteControlProtocolWidgets/Private/Widgets/SRCProtocolBinding.cpp
@@ -6,9 +6,12 @@
 #include "SRCBindingWarning.h"
 #include "SRCProtocolRangeList.h"
 #include "ViewModels/ProtocolBindingViewModel.h"
+#include "Widgets/Masks/SRCProtocolBindingMask.h"
 #include "Widgets/SBoxPanel.h"
 #include "Widgets/SRCProtocolEntity.h"
+#include "Widgets/SRCProtocolShared.h"
 #include "Widgets/Input/SButton.h"
+#include "Widgets/Layout/SSpacer.h"
 #include "Widgets/Text/STextBlock.h"
 
 #define LOCTEXT_NAMESPACE "RemoteControlProtocolWidgets"
@@ -118,6 +121,14 @@ void SRCProtocolBinding::Construct(const FArguments& InArgs, const TSharedRef<ST
 				[
 					SNew(SRCProtocolEntity, ViewModel.ToSharedRef())
 				]
+				
+				+ SVerticalBox::Slot()
+				.Padding(Padding)
+				.VAlign(VAlign_Center)
+				.AutoHeight()
+				[
+					ConstructMaskingWidget()
+				]
 
 				+ SVerticalBox::Slot()
 				.Padding(Padding)
@@ -146,6 +157,55 @@ void SRCProtocolBinding::Construct(const FArguments& InArgs, const TSharedRef<ST
 		InOwnerTableView);
 }
 
+TSharedRef<SWidget> SRCProtocolBinding::ConstructMaskingWidget()
+{
+	const TSharedPtr<SWidget> LeftWidget = SNew(SHorizontalBox)
+
+		// Masking Label
+		+ SHorizontalBox::Slot()
+		.HAlign(HAlign_Left)
+		.VAlign(VAlign_Center)
+		.Padding(4.f, 2.f)
+		.AutoWidth()
+		[
+			SNew(STextBlock)
+			.Text(LOCTEXT("OverrideMaskingLabel", "Override Masks"))
+			.Font(FAppStyle::GetFontStyle("DetailsView.CategoryFontStyle"))
+			.ShadowOffset(FVector2D(1.f, 1.f))
+		]
+
+		+ SHorizontalBox::Slot()
+		.FillWidth(1.f)
+		[
+			SNew(SSpacer)
+		];
+
+		const TSharedPtr<SWidget> RightWidget = SNew(SHorizontalBox)
+
+			// Masking Widget
+			+ SHorizontalBox::Slot()
+			.VAlign(VAlign_Center)
+			.HAlign(HAlign_Left)
+			.Padding(4.f, 2.f)
+			.AutoWidth()
+			[
+				SNew(SRCProtocolBindingMask, ViewModel.ToSharedRef())
+			]
+
+			+ SHorizontalBox::Slot()
+			.FillWidth(1.f)
+			[
+				SNew(SSpacer)
+			];
+
+	TSharedRef<SWidget> Widget = SNew(RemoteControlProtocolWidgetUtils::SCustomSplitter)
+		.LeftWidget(LeftWidget.ToSharedRef())
+		.RightWidget(RightWidget.ToSharedRef())
+		.ColumnSizeData(PrimaryColumnSizeData);
+
+	return Widget;
+}
+
 FReply SRCProtocolBinding::OnDelete()
 {
 	ViewModel->Remove();
diff --git a/Engine/Plugins/VirtualProduction/RemoteControl/Source/RemoteControlProtocolWidgets/Private/Widgets/SRCProtocolBinding.h b/Engine/Plugins/VirtualProduction/RemoteControl/Source/RemoteControlProtocolWidgets/Private/Widgets/SRCProtocolBinding.h
--- a/Engine/Plugins/VirtualProduction/RemoteControl/Source/RemoteControlProtocolWidgets/Private/Widgets/SRCProtocolBinding.h
+++ b/Engine/Plugins/VirtualProduction/RemoteControl/Source/RemoteControlProtocolWidgets/Private/Widgets/SRCProtocolBinding.h
@@ -1,4 +1,4 @@
-﻿// Copyright Epic Games, Inc. All Rights Reserved.
+// Copyright Epic Games, Inc. All Rights Reserved.
 
 #pragma once
 
@@ -14,9 +14,6 @@ struct FRemoteControlProtocolEntity;
 struct FSlateColor;
 
 class FProtocolBindingViewModel;
-class FProtocolBindingViewModel;
-class FProtocolRangeViewModel;
-class FProtocolRangeViewModel;
 class FProtocolRangeViewModel;
 class ITableRow;
 class SRCProtocolRangeList;
@@ -42,6 +39,9 @@ public:
 	void Construct(const FArguments& InArgs, const TSharedRef<STableViewBase>& InOwnerTableView, const TSharedRef<FProtocolBindingViewModel>& InViewModel);
 
 private:
+	/** Constructs the masking widget. */
+	TSharedRef<SWidget> ConstructMaskingWidget();
+
 	/** Called when Delete button pressed. */
 	FReply OnDelete();
 
diff --git a/Engine/Plugins/VirtualProduction/RemoteControl/Source/RemoteControlProtocolWidgets/Public/Widgets/Masks/SRCProtocolMaskTriplet.h b/Engine/Plugins/VirtualProduction/RemoteControl/Source/RemoteControlProtocolWidgets/Public/Widgets/Masks/SRCProtocolMaskTriplet.h
--- a/Engine/Plugins/VirtualProduction/RemoteControl/Source/RemoteControlProtocolWidgets/Public/Widgets/Masks/SRCProtocolMaskTriplet.h
+++ b/Engine/Plugins/VirtualProduction/RemoteControl/Source/RemoteControlProtocolWidgets/Public/Widgets/Masks/SRCProtocolMaskTriplet.h
@@ -38,6 +38,26 @@ enum class EMaskingType : uint8
 	Unsupported
 };
 
+namespace RemoteControlProtocolMasking
+{
+	static TMap<UScriptStruct*, EMaskingType> StructsToMaskingTypes = { { TBaseStructure<FColor>::Get(), EMaskingType::Color }
+		, { TBaseStructure<FLinearColor>::Get(), EMaskingType::Color }
+		, { TBaseStructure<FRotator>::Get(), EMaskingType::Rotator }
+		, { TBaseStructure<FVector>::Get(), EMaskingType::Vector }
+		, { TBaseStructure<FIntVector>::Get(), EMaskingType::Vector }
+		, { TBaseStructure<FVector4>::Get(), EMaskingType::Quat }
+		, { TBaseStructure<FIntVector4>::Get(), EMaskingType::Quat }
+		, { TBaseStructure<FQuat>::Get(), EMaskingType::Quat }
+	};
+
+	static TSet<UScriptStruct*> OptionalMaskStructs = {
+		TBaseStructure<FLinearColor>::Get(),
+		TBaseStructure<FVector4>::Get(),
+		TBaseStructure<FQuat>::Get(),
+		TBaseStructure<FIntVector4>::Get()
+	};
+}
+
 /**
  * A widget that represents a triplet mask.
  */
