Fix localization key clash

diff --git a/Engine/Plugins/VirtualProduction/RemoteControl/Source/RemoteControlUI/Private/UI/Behaviour/SRCBehaviourPanelList.cpp b/Engine/Plugins/VirtualProduction/RemoteControl/Source/RemoteControlUI/Private/UI/Behaviour/SRCBehaviourPanelList.cpp
--- a/Engine/Plugins/VirtualProduction/RemoteControl/Source/RemoteControlUI/Private/UI/Behaviour/SRCBehaviourPanelList.cpp
+++ b/Engine/Plugins/VirtualProduction/RemoteControl/Source/RemoteControlUI/Private/UI/Behaviour/SRCBehaviourPanelList.cpp
@@ -115,7 +115,7 @@ void SRCBehaviourPanelList::AddSpecialContextMenuOptions(FMenuBuilder& MenuBuild
 		else
 		{
 			Label = LOCTEXT("LabelEnableBehaviour", "Enable");
-			Tooltip = LOCTEXT("TooltipDisableBehaviour", "Enables the current behaviour. Restores functioning of Actions associated with this behaviour.");
+			Tooltip = LOCTEXT("TooltipEnableBehaviour", "Enables the current behaviour. Restores functioning of Actions associated with this behaviour.");
 		}
 
 		MenuBuilder.AddMenuEntry(Label, Tooltip, FSlateIcon(), Action);
