Enable scrubbing live values for Controllers

diff --git a/Engine/Plugins/VirtualProduction/RemoteControl/Source/RemoteControl/Private/RCVirtualPropertyContainer.cpp b/Engine/Plugins/VirtualProduction/RemoteControl/Source/RemoteControl/Private/RCVirtualPropertyContainer.cpp
--- a/Engine/Plugins/VirtualProduction/RemoteControl/Source/RemoteControl/Private/RCVirtualPropertyContainer.cpp
+++ b/Engine/Plugins/VirtualProduction/RemoteControl/Source/RemoteControl/Private/RCVirtualPropertyContainer.cpp
@@ -13,12 +13,24 @@
 
 #define LOCTEXT_NAMESPACE "URCVirtualPropertyInContainer"
 
+void URCVirtualPropertyContainerBase::AddVirtualProperty(URCVirtualPropertyBase* InVirtualProperty)
+{
+	if(ensure(InVirtualProperty))
+	{
+		// Add property to Set
+		VirtualProperties.Add(InVirtualProperty);
+
+		// Initialize
+		InVirtualProperty->Init();
+	}
+}
+
 URCVirtualPropertyInContainer* URCVirtualPropertyContainerBase::AddProperty(const FName& InPropertyName, TSubclassOf<URCVirtualPropertyInContainer> InPropertyClass, const EPropertyBagPropertyType InValueType, UObject* InValueTypeObject)
 {
 	const FName PropertyName = GenerateUniquePropertyName(InPropertyName, InValueType, InValueTypeObject, this);
 	Bag.AddProperty(PropertyName, InValueType, InValueTypeObject);
 
-	// Ensure that the property has been succcefully added to the Bag
+	// Ensure that the property has been successfully added to the Bag
 	const FPropertyBagPropertyDesc* BagPropertyDesc = Bag.FindPropertyDescByName(PropertyName);
 	if (!ensure(BagPropertyDesc))
 	{
@@ -31,9 +43,8 @@ URCVirtualPropertyInContainer* URCVirtualPropertyContainerBase::AddProperty(cons
 	VirtualPropertyInContainer->PresetWeakPtr = PresetWeakPtr;
 	VirtualPropertyInContainer->ContainerWeakPtr = this;
 	VirtualPropertyInContainer->Id = FGuid::NewGuid();
-	
-	// Add property to Set
-	VirtualProperties.Add(VirtualPropertyInContainer);
+
+	AddVirtualProperty(VirtualPropertyInContainer);
 
 	return VirtualPropertyInContainer;
 }
@@ -55,8 +66,7 @@ URCVirtualPropertyInContainer* URCVirtualPropertyContainerBase::DuplicatePropert
 	VirtualPropertyInContainer->ContainerWeakPtr = this;
 	VirtualPropertyInContainer->Id = FGuid::NewGuid();
 
-	// Add property to Set
-	VirtualProperties.Add(VirtualPropertyInContainer);
+	AddVirtualProperty(VirtualPropertyInContainer);
 
 	return VirtualPropertyInContainer;
 }
@@ -99,7 +109,7 @@ URCVirtualPropertyInContainer* URCVirtualPropertyContainerBase::DuplicateVirtual
 		}
 
 		// Sync Virtual Properties List
-		VirtualProperties.Add(NewVirtualProperty);
+		AddVirtualProperty(NewVirtualProperty);
 
 		return NewVirtualProperty;
 	}
diff --git a/Engine/Plugins/VirtualProduction/RemoteControl/Source/RemoteControl/Private/RemoteControlPreset.cpp b/Engine/Plugins/VirtualProduction/RemoteControl/Source/RemoteControl/Private/RemoteControlPreset.cpp
--- a/Engine/Plugins/VirtualProduction/RemoteControl/Source/RemoteControl/Private/RemoteControlPreset.cpp
+++ b/Engine/Plugins/VirtualProduction/RemoteControl/Source/RemoteControl/Private/RemoteControlPreset.cpp
@@ -622,6 +622,15 @@ void URemoteControlPreset::SetControllerContainer(URCVirtualPropertyContainerBas
 }
 
 #if WITH_EDITOR
+
+void URemoteControlPreset::OnNotifyPreChangeVirtualProperty(const FPropertyChangedEvent& PropertyChangedEvent)
+{
+	if (ensure(ControllerContainer))
+	{
+		ControllerContainer->OnPreChangePropertyValue(PropertyChangedEvent);
+	}
+}
+
 void URemoteControlPreset::OnModifyVirtualProperty(const FPropertyChangedEvent& PropertyChangedEvent)
 {
 	if (ensure(ControllerContainer))
diff --git a/Engine/Plugins/VirtualProduction/RemoteControl/Source/RemoteControl/Public/RCVirtualProperty.h b/Engine/Plugins/VirtualProduction/RemoteControl/Source/RemoteControl/Public/RCVirtualProperty.h
--- a/Engine/Plugins/VirtualProduction/RemoteControl/Source/RemoteControl/Public/RCVirtualProperty.h
+++ b/Engine/Plugins/VirtualProduction/RemoteControl/Source/RemoteControl/Public/RCVirtualProperty.h
@@ -44,9 +44,16 @@ protected:
 	/** Pointer to Instanced property which is holds  bag of properties */
 	virtual const FInstancedPropertyBag* GetPropertyBagInstance() const;
 
+	/** Notifies pre change for Virtual Property value*/
+	virtual void OnPreChangePropertyValue() {}
+
+	/** Notifies post change for Virtual Property value*/
 	virtual void OnModifyPropertyValue() {}
 
 public:
+	/** Initialization routine. Called after the parent container has setup data for this property */
+	virtual void Init() {}
+
 	/** Returns const FProperty for this RC virtual property */
 	virtual const FProperty* GetProperty() const;
 
diff --git a/Engine/Plugins/VirtualProduction/RemoteControl/Source/RemoteControl/Public/RCVirtualPropertyContainer.h b/Engine/Plugins/VirtualProduction/RemoteControl/Source/RemoteControl/Public/RCVirtualPropertyContainer.h
--- a/Engine/Plugins/VirtualProduction/RemoteControl/Source/RemoteControl/Public/RCVirtualPropertyContainer.h
+++ b/Engine/Plugins/VirtualProduction/RemoteControl/Source/RemoteControl/Public/RCVirtualPropertyContainer.h
@@ -120,6 +120,9 @@ public:
 	/** Called after applying a transaction to the object. Used to broadcast Undo related container changes to UI. */
 	virtual void PostEditUndo();
 
+	/** Delegate when the value is being scrubbed in UI*/
+	virtual void OnPreChangePropertyValue(const FPropertyChangedEvent& PropertyChangedEvent) {}
+
 	/** Delegate when object changed */
 	virtual void OnModifyPropertyValue(const FPropertyChangedEvent& PropertyChangedEvent);
 #endif
@@ -128,6 +131,9 @@ public:
 	DECLARE_MULTICAST_DELEGATE(FOnVirtualPropertyContainerModified);
 	FOnVirtualPropertyContainerModified& OnVirtualPropertyContainerModified() { return OnVirtualPropertyContainerModifiedDelegate; }
 
+private:
+	void AddVirtualProperty(URCVirtualPropertyBase* InVirtualProperty);
+
 protected:
 	/** Holds bag of properties. */
 	UPROPERTY()
diff --git a/Engine/Plugins/VirtualProduction/RemoteControl/Source/RemoteControl/Public/RemoteControlPreset.h b/Engine/Plugins/VirtualProduction/RemoteControl/Source/RemoteControl/Public/RemoteControlPreset.h
--- a/Engine/Plugins/VirtualProduction/RemoteControl/Source/RemoteControl/Public/RemoteControlPreset.h
+++ b/Engine/Plugins/VirtualProduction/RemoteControl/Source/RemoteControl/Public/RemoteControlPreset.h
@@ -658,6 +658,10 @@ public:
 	}
 
 #if WITH_EDITOR
+	/** Called when a virtual property's value is being scrubbed by the user in the UI.
+	* This call is routed to the Controller for evaluating associated Logic & Behaviours*/
+	void OnNotifyPreChangeVirtualProperty(const FPropertyChangedEvent& PropertyChangedEvent);
+
 	/** Called when a virtual property is modified. This call is routed to the Controller for evaluating associated Logic & Behaviours*/
 	void OnModifyVirtualProperty(const FPropertyChangedEvent& PropertyChangedEvent);
 #endif
diff --git a/Engine/Plugins/VirtualProduction/RemoteControl/Source/RemoteControlLogic/Private/Action/Bind/RCPropertyBindAction.cpp b/Engine/Plugins/VirtualProduction/RemoteControl/Source/RemoteControlLogic/Private/Action/Bind/RCPropertyBindAction.cpp
--- a/Engine/Plugins/VirtualProduction/RemoteControl/Source/RemoteControlLogic/Private/Action/Bind/RCPropertyBindAction.cpp
+++ b/Engine/Plugins/VirtualProduction/RemoteControl/Source/RemoteControlLogic/Private/Action/Bind/RCPropertyBindAction.cpp
@@ -2,11 +2,17 @@
 
 #include "Action/Bind/RCPropertyBindAction.h"
 
+#include "Action/RCAction.h"
+#include "Components/SceneComponent.h"
 #include "Controller/RCController.h"
 #include "IRemoteControlPropertyHandle.h"
 #include "RCVirtualProperty.h"
 #include "RemoteControlPreset.h"
-#include "Action/RCAction.h"
+
+#if WITH_EDITOR
+#include "Editor/UnrealEdEngine.h"
+#include "UnrealEdGlobals.h"
+#endif
 
 static void SetStructPropertyFromController(const FProperty* RemoteControlProperty, TSharedRef<IRemoteControlPropertyHandle> RemoteControlHandle, const URCController* Controller)
 {
@@ -193,5 +199,23 @@ void URCPropertyBindAction::Execute() const
 		{
 			SetStructPropertyFromController(RemoteControlProperty, Handle.ToSharedRef(), Controller);
 		}
+
+		// Editor specific updates
+	#if WITH_EDITOR
+		for (UObject* Object : RemoteControlEntityAsProperty->GetBoundObjects())
+		{
+			// For Bind behaviour we want to refresh the pivot widget if we have just manipulated a Location vector via Controller
+			if (USceneComponent* SceneComponent = Cast<USceneComponent>(Object))
+			{
+				if (AActor* Actor = SceneComponent->GetOwner())
+				{
+					if (Actor->IsSelectedInEditor())
+					{
+						GUnrealEd->UpdatePivotLocationForSelection();
+					}
+				}
+			}
+		}
+	#endif
 	}
 }
\ No newline at end of file
diff --git a/Engine/Plugins/VirtualProduction/RemoteControl/Source/RemoteControlLogic/Private/Behaviour/Builtin/Bind/RCBehaviourBind.cpp b/Engine/Plugins/VirtualProduction/RemoteControl/Source/RemoteControlLogic/Private/Behaviour/Builtin/Bind/RCBehaviourBind.cpp
--- a/Engine/Plugins/VirtualProduction/RemoteControl/Source/RemoteControlLogic/Private/Behaviour/Builtin/Bind/RCBehaviourBind.cpp
+++ b/Engine/Plugins/VirtualProduction/RemoteControl/Source/RemoteControlLogic/Private/Behaviour/Builtin/Bind/RCBehaviourBind.cpp
@@ -13,6 +13,8 @@
 
 URCBehaviourBind::URCBehaviourBind()
 {
+	// For Bind Behaviour we want the users to be able to scrub UI widgets like float slides and watch the Exposed Property (and their level rendering like Location, Light intenstiy, etc) update in realtime
+	bExecuteBehavioursDuringPreChange = true;
 }
 
 void URCBehaviourBind::Initialize()
diff --git a/Engine/Plugins/VirtualProduction/RemoteControl/Source/RemoteControlLogic/Private/Controller/RCController.cpp b/Engine/Plugins/VirtualProduction/RemoteControl/Source/RemoteControlLogic/Private/Controller/RCController.cpp
--- a/Engine/Plugins/VirtualProduction/RemoteControl/Source/RemoteControlLogic/Private/Controller/RCController.cpp
+++ b/Engine/Plugins/VirtualProduction/RemoteControl/Source/RemoteControlLogic/Private/Controller/RCController.cpp
@@ -10,6 +10,29 @@
 
 #define LOCTEXT_NAMESPACE "RCController"
 
+void URCController::Init()
+{
+#if WITH_EDITORONLY_DATA
+	// Setup metadata
+	// Note: In the future this can be set directly by the user from UI
+	if (IsVectorType())
+	{
+		if (FProperty* Property = GetProperty())
+		{
+			// Provides users with a vector scrubbing experience that approximates the SceneComponent's Location Vector widget
+			Property->SetMetaData(TEXT("SliderExponent"), FString::Printf(TEXT("%f"), SliderExponent));
+		}
+	}
+#endif
+}
+
+void URCController::PostLoad()
+{
+	Super::PostLoad();
+
+	Init();
+}
+
 #if WITH_EDITOR
 void URCController::PostEditUndo()
 {
@@ -78,17 +101,36 @@ void URCController::EmptyBehaviours()
 	Behaviours.Empty();
 }
 
-void URCController::ExecuteBehaviours()
+void URCController::ExecuteBehaviours(const bool bIsPreChange/* = false*/)
 {
 	for (URCBehaviour* Behaviour : Behaviours)
 	{
-		if (Behaviour->bIsEnabled)
+		if (!Behaviour->bIsEnabled)
 		{
-			Behaviour->Execute();
+			continue;
 		}
+
+		if (bIsPreChange && !Behaviour->bExecuteBehavioursDuringPreChange)
+		{
+			continue;
+		}
+
+		Behaviour->Execute();
+
 	}
 }
 
+void URCController::OnPreChangePropertyValue()
+{
+	constexpr bool bIsPreChange = true;
+	ExecuteBehaviours(bIsPreChange);
+}
+
+void URCController::OnModifyPropertyValue()
+{
+	ExecuteBehaviours();
+}
+
 URCBehaviour* URCController::DuplicateBehaviour(URCController* InController, URCBehaviour* InBehaviour)
 {
 	URCBehaviour* NewBehaviour = DuplicateObject<URCBehaviour>(InBehaviour, InController);
diff --git a/Engine/Plugins/VirtualProduction/RemoteControl/Source/RemoteControlLogic/Private/Controller/RCControllerContainer.cpp b/Engine/Plugins/VirtualProduction/RemoteControl/Source/RemoteControlLogic/Private/Controller/RCControllerContainer.cpp
--- a/Engine/Plugins/VirtualProduction/RemoteControl/Source/RemoteControlLogic/Private/Controller/RCControllerContainer.cpp
+++ b/Engine/Plugins/VirtualProduction/RemoteControl/Source/RemoteControlLogic/Private/Controller/RCControllerContainer.cpp
@@ -4,20 +4,36 @@
 #include "Controller/RCController.h"
 
 #if WITH_EDITOR
-void URCControllerContainer::OnModifyPropertyValue(const FPropertyChangedEvent& Event)
+
+URCController* URCControllerContainer::GetControllerFromChangeEvent(const FPropertyChangedEvent& Event)
 {
 	if (const FProperty* FinalProperty = (Event.Property == Event.MemberProperty) ? Event.Property : Event.MemberProperty)
 	{
-		for (URCVirtualPropertyBase* VirtualProperty : VirtualProperties)
-		{
-			if (VirtualProperty && VirtualProperty->PropertyName == FinalProperty->GetFName())
-			{
-				if (URCController* Controller = Cast<URCController>(VirtualProperty))
-				{
-					Controller->OnModifyPropertyValue();
-				}
-			}
-		}
+		const FName PropertyName = FinalProperty->GetFName();
+
+		URCVirtualPropertyBase* VirtualProperty = GetVirtualProperty(PropertyName);
+
+		return Cast<URCController>(VirtualProperty);
+	}
+
+	return nullptr;
+}
+
+void URCControllerContainer::OnPreChangePropertyValue(const FPropertyChangedEvent& Event)
+{
+	if (URCController* Controller = GetControllerFromChangeEvent(Event))
+	{
+		Controller->OnPreChangePropertyValue();
+	}
+
+	Super::OnPreChangePropertyValue(Event);
+}
+
+void URCControllerContainer::OnModifyPropertyValue(const FPropertyChangedEvent& Event)
+{
+	if (URCController* Controller = GetControllerFromChangeEvent(Event))
+	{
+		Controller->OnModifyPropertyValue();
 	}
 
 	Super::OnModifyPropertyValue(Event);
diff --git a/Engine/Plugins/VirtualProduction/RemoteControl/Source/RemoteControlLogic/Public/Behaviour/RCBehaviour.h b/Engine/Plugins/VirtualProduction/RemoteControl/Source/RemoteControlLogic/Public/Behaviour/RCBehaviour.h
--- a/Engine/Plugins/VirtualProduction/RemoteControl/Source/RemoteControlLogic/Public/Behaviour/RCBehaviour.h
+++ b/Engine/Plugins/VirtualProduction/RemoteControl/Source/RemoteControlLogic/Public/Behaviour/RCBehaviour.h
@@ -120,4 +120,9 @@ public:
 	* If disabled, it will be not evaluated when the associated Controller changes */
 	UPROPERTY()
 	bool bIsEnabled = true;
+
+	/** Indicates whether we want a behavour to trigger during live scrubbing of values.
+	* For example, for a light intensity value controlled via Bind behavour we want the bound properties to update live even while the float widget is being scrubbed by the user. */
+	UPROPERTY()
+	bool bExecuteBehavioursDuringPreChange = false;
 };
diff --git a/Engine/Plugins/VirtualProduction/RemoteControl/Source/RemoteControlLogic/Public/Controller/RCController.h b/Engine/Plugins/VirtualProduction/RemoteControl/Source/RemoteControlLogic/Public/Controller/RCController.h
--- a/Engine/Plugins/VirtualProduction/RemoteControl/Source/RemoteControlLogic/Public/Controller/RCController.h
+++ b/Engine/Plugins/VirtualProduction/RemoteControl/Source/RemoteControlLogic/Public/Controller/RCController.h
@@ -1,4 +1,4 @@
-﻿// Copyright Epic Games, Inc. All Rights Reserved.
+﻿// CCopyright Epic Games, Inc. All Rights Reserved.
 
 #pragma once
 
@@ -19,6 +19,12 @@ class REMOTECONTROLLOGIC_API URCController : public URCVirtualPropertyInContaine
 	GENERATED_BODY()
 
 public:
+	/** Initialization routine. Called after the parent container has setup data for this property */
+	virtual void Init();
+
+	/** Called post serialization. Used to reinitialize Controller metadata after a Remote Control Preset asset is loaded*/
+	virtual void PostLoad() override;
+
 #if WITH_EDITOR
 	/** Called after applying a transaction to the object. Used to broadcast Undo related container changes to UI Used to broadcast Undo related container changes to UI */
 	virtual void PostEditUndo();
@@ -40,13 +46,13 @@ public:
 	virtual void EmptyBehaviours();
 
 	/** Execute all behaviours for this controller. */
-	virtual void ExecuteBehaviours();
+	virtual void ExecuteBehaviours(const bool bIsPreChange = false);
+
+	/** Pre-change notification for Controllers. Triggered while the user is scrubbing a float or vector slider in the UI*/
+	virtual void OnPreChangePropertyValue() override;
 
 	/** Handles modifications to controller value; evaluates all behaviours */
-	virtual void OnModifyPropertyValue() override
-	{
-		ExecuteBehaviours();
-	}
+	virtual void OnModifyPropertyValue() override;
 
 	/** Duplicates an existing behaviour and adds it to the behaviour set of a given Controller */
 	static URCBehaviour* DuplicateBehaviour(URCController* InController, URCBehaviour* InBehaviour);
@@ -58,4 +64,9 @@ public:
 	/** Set of the behaviours */
 	UPROPERTY()
 	TSet<TObjectPtr<URCBehaviour>> Behaviours;
+
+	/** Slider exponent for Vector Controllers. Determines how fast the value change travels while scrubbing the X/Y/Z sliders in UI
+	* The default value approximates the behaviour of SceneComponent Location Vectors (see SNumericVectorInputBox in ComponentTransformDetails.cpp) */
+	UPROPERTY()
+	float SliderExponent = 0.4f;
 };
diff --git a/Engine/Plugins/VirtualProduction/RemoteControl/Source/RemoteControlLogic/Public/Controller/RCControllerContainer.h b/Engine/Plugins/VirtualProduction/RemoteControl/Source/RemoteControlLogic/Public/Controller/RCControllerContainer.h
--- a/Engine/Plugins/VirtualProduction/RemoteControl/Source/RemoteControlLogic/Public/Controller/RCControllerContainer.h
+++ b/Engine/Plugins/VirtualProduction/RemoteControl/Source/RemoteControlLogic/Public/Controller/RCControllerContainer.h
@@ -6,6 +6,7 @@
 #include "RCControllerContainer.generated.h"
 
 class URCActionContainer;
+class URCController;
 
 /**
  * Controller Container which holds all virtual controller properties 
@@ -17,6 +18,12 @@ class REMOTECONTROLLOGIC_API URCControllerContainer : public URCVirtualPropertyC
 	
 private:
 #if WITH_EDITOR
+	/** Fetches the Controller underlying a given a Property Changed Event*/
+	URCController* GetControllerFromChangeEvent(const FPropertyChangedEvent& Event);
+
+	/** Delegate when the value is being scrubbed in UI*/
+	virtual void OnPreChangePropertyValue(const FPropertyChangedEvent& PropertyChangedEvent) override;
+
 	/** Delegate when object changed */
 	virtual void OnModifyPropertyValue(const FPropertyChangedEvent& PropertyChangedEvent) override;
 #endif
diff --git a/Engine/Plugins/VirtualProduction/RemoteControl/Source/RemoteControlUI/Private/UI/Controller/SRCControllerPanelList.cpp b/Engine/Plugins/VirtualProduction/RemoteControl/Source/RemoteControlUI/Private/UI/Controller/SRCControllerPanelList.cpp
--- a/Engine/Plugins/VirtualProduction/RemoteControl/Source/RemoteControlUI/Private/UI/Controller/SRCControllerPanelList.cpp
+++ b/Engine/Plugins/VirtualProduction/RemoteControl/Source/RemoteControlUI/Private/UI/Controller/SRCControllerPanelList.cpp
@@ -324,6 +324,7 @@ void SRCControllerPanelList::Construct(const FArguments& InArgs, const TSharedRe
 
 	FPropertyRowGeneratorArgs Args;
 	Args.bShouldShowHiddenProperties = true;
+	Args.NotifyHook = this;
 	PropertyRowGenerator = FModuleManager::GetModuleChecked<FPropertyEditorModule>("PropertyEditor").CreatePropertyRowGenerator(Args);
 	
 	Reset();
@@ -431,14 +432,19 @@ void SRCControllerPanelList::OnControllerAdded(const FName& InNewPropertyName)
 	Reset();
 }
 
+void SRCControllerPanelList::OnNotifyPreChangeProperties(const FPropertyChangedEvent& PropertyChangedEvent)
+{
+	if (URemoteControlPreset* Preset = GetPreset())
+	{
+		Preset->OnNotifyPreChangeVirtualProperty(PropertyChangedEvent);
+	}
+}
+
 void SRCControllerPanelList::OnFinishedChangingProperties(const FPropertyChangedEvent& PropertyChangedEvent)
 {
-	if (TSharedPtr< SRCControllerPanel> ControllerPanel = ControllerPanelWeakPtr.Pin())
+	if (URemoteControlPreset* Preset = GetPreset())
 	{
-		if (URemoteControlPreset* Preset = ControllerPanel->GetPreset())
-		{
-			Preset->OnModifyVirtualProperty(PropertyChangedEvent);
-		}
+		Preset->OnModifyVirtualProperty(PropertyChangedEvent);
 	}
 }
 
@@ -536,6 +542,19 @@ void SRCControllerPanelList::DeleteSelectedPanelItem()
 	DeleteItemFromLogicPanel<FRCControllerModel>(ControllerItems, ListView->GetSelectedItems());
 }
 
+void SRCControllerPanelList::NotifyPreChange(FEditPropertyChain* PropertyAboutToChange)
+{
+	// If a Vector is modified and the Z value changes, the sub property (corresponding to Z) gets notified to us.
+	// However for Controllers we are actually interested in the parent Struct property (corresponding to the Vector) as the Virtual Property is associated with Struct's FProperty (rather than its X/Y/Z components)
+	// For this reason the "Active Member Node" is extracted below from the child property
+	if (FEditPropertyChain::TDoubleLinkedListNode* ActiveMemberNode = PropertyAboutToChange->GetActiveMemberNode())
+	{
+		FPropertyChangedEvent PropertyChangedEvent(ActiveMemberNode->GetValue());
+
+		OnNotifyPreChangeProperties(PropertyChangedEvent);
+	}
+}
+
 void SRCControllerPanelList::EnterRenameMode()
 {
 	if (TSharedPtr<FRCControllerModel> SelectedItem = SelectedControllerItemWeakPtr.Pin())
diff --git a/Engine/Plugins/VirtualProduction/RemoteControl/Source/RemoteControlUI/Private/UI/Controller/SRCControllerPanelList.h b/Engine/Plugins/VirtualProduction/RemoteControl/Source/RemoteControlUI/Private/UI/Controller/SRCControllerPanelList.h
--- a/Engine/Plugins/VirtualProduction/RemoteControl/Source/RemoteControlUI/Private/UI/Controller/SRCControllerPanelList.h
+++ b/Engine/Plugins/VirtualProduction/RemoteControl/Source/RemoteControlUI/Private/UI/Controller/SRCControllerPanelList.h
@@ -2,12 +2,14 @@
 
 #pragma once
 #include "DragAndDrop/DecoratedDragDropOp.h"
+#include "Misc/NotifyHook.h"
 #include "UI/BaseLogicUI/SRCLogicPanelListBase.h"
 #include "UI/RemoteControlPanelStyle.h"
 #include "Widgets/Layout/SBorder.h"
 
 class FDragDropOperation;
 struct FRCPanelStyle;
+class FEditPropertyChain;
 class FRCControllerModel;
 class FRCLogicModeBase;
 class IPropertyRowGenerator;
@@ -58,7 +60,7 @@ private:
 * UI Widget for Controllers List
 * Used as part of the RC Logic Actions Panel.
 */
-class REMOTECONTROLUI_API SRCControllerPanelList : public SRCLogicPanelListBase
+class REMOTECONTROLUI_API SRCControllerPanelList : public SRCLogicPanelListBase, public FNotifyHook
 {
 public:
 	SLATE_BEGIN_ARGS(SRCControllerPanelList)
@@ -91,6 +93,10 @@ public:
 		return GetSelectedControllerItem();
 	}
 
+	/** FNotifyHook Interface Begin */
+	virtual void NotifyPreChange(FEditPropertyChain* PropertyAboutToChange) override;
+	/** FNotifyHook Interface End */
+
 	void EnterRenameMode();
 
 	int32 NumControllerItems() const
@@ -159,6 +165,11 @@ private:
 
 	void OnControllerContainerModified();
 
+	/** Pre Change listener for Controllers, propagated via FNotifyHook associated with PropertyRowGenerator
+	* Invoked while the user is scrubbing float or Vector sliders in the UI
+	*/
+	void OnNotifyPreChangeProperties(const FPropertyChangedEvent& PropertyChangedEvent);
+
 	/** Change listener for Controllers. Bound to the PropertyRowGenerator's delegate
 	* This is propagated to the corresponding Controller model (Virtual Property) for evaluating all associated Behaviours.
 	*/
