1. Remote Control Logic - Setup/Operation Toggle Not working

2. Remote Control Logic - Setup/Operation Toggle Should Only Appear in Logic Mode

diff --git a/Engine/Plugins/VirtualProduction/RemoteControl/Source/RemoteControlUI/Private/AssetEditor/RemoteControlPresetEditorToolkit.cpp b/Engine/Plugins/VirtualProduction/RemoteControl/Source/RemoteControlUI/Private/AssetEditor/RemoteControlPresetEditorToolkit.cpp
--- a/Engine/Plugins/VirtualProduction/RemoteControl/Source/RemoteControlUI/Private/AssetEditor/RemoteControlPresetEditorToolkit.cpp
+++ b/Engine/Plugins/VirtualProduction/RemoteControl/Source/RemoteControlUI/Private/AssetEditor/RemoteControlPresetEditorToolkit.cpp
@@ -97,17 +97,11 @@ void FRemoteControlPresetEditorToolkit::RegisterTabSpawners(const TSharedRef<FTa
 			.SetDisplayName(LOCTEXT("RemoteControlPanelMainTab", "Remote Control Panel"))
 			.SetGroup(WorkspaceMenuCategory.ToSharedRef())
 			.SetIcon(FSlateIcon(FAppStyle::GetAppStyleSetName(), "LevelEditor.Tabs.Details"));
-
-		InTabManager->RegisterTabSpawner(FRemoteControlUIModule::EntityDetailsTabName, FOnSpawnTab::CreateSP(this, &FRemoteControlPresetEditorToolkit::HandleTabManagerSpawnDetailsTab))
-			.SetDisplayName(LOCTEXT("RemoteControlPanelDetailsTab", "Entity Details"))
-			.SetGroup(WorkspaceMenuCategory.ToSharedRef())
-			.SetIcon(FSlateIcon(FAppStyle::GetAppStyleSetName(), "LevelEditor.Tabs.Details"));
 	}
 }
 
 void FRemoteControlPresetEditorToolkit::UnregisterTabSpawners(const TSharedRef<FTabManager>& InTabManager)
 {
-	InTabManager->UnregisterTabSpawner(FRemoteControlUIModule::EntityDetailsTabName);
 	InTabManager->UnregisterTabSpawner(FRemoteControlUIModule::RemoteControlPanelTabName);
 }
 
@@ -168,18 +162,6 @@ TSharedRef<SDockTab> FRemoteControlPresetEditorToolkit::HandleTabManagerSpawnPan
 		];
 }
 
-TSharedRef<SDockTab> FRemoteControlPresetEditorToolkit::HandleTabManagerSpawnDetailsTab(const FSpawnTabArgs& Args)
-{
-	check(Args.GetTabId() == FRemoteControlUIModule::EntityDetailsTabName);
-
-	return SNew(SDockTab)
-		.Label(LOCTEXT("ControlPanelEntityDetailsLabel", "Entity Details"))
-		.TabColorScale(GetTabColorScale())	
-		[
-			SNullWidget::NullWidget
-		];
-}
-
 void FRemoteControlPresetEditorToolkit::InvokePanelTab()
 {
 	struct Local
diff --git a/Engine/Plugins/VirtualProduction/RemoteControl/Source/RemoteControlUI/Private/AssetEditor/RemoteControlPresetEditorToolkit.h b/Engine/Plugins/VirtualProduction/RemoteControl/Source/RemoteControlUI/Private/AssetEditor/RemoteControlPresetEditorToolkit.h
--- a/Engine/Plugins/VirtualProduction/RemoteControl/Source/RemoteControlUI/Private/AssetEditor/RemoteControlPresetEditorToolkit.h
+++ b/Engine/Plugins/VirtualProduction/RemoteControl/Source/RemoteControlUI/Private/AssetEditor/RemoteControlPresetEditorToolkit.h
@@ -43,7 +43,6 @@ public:
 private:
 	//~ Handlers that spawn remote control tabs.
 	TSharedRef<SDockTab> HandleTabManagerSpawnPanelTab(const FSpawnTabArgs& Args);
-	TSharedRef<SDockTab> HandleTabManagerSpawnDetailsTab(const FSpawnTabArgs& Args);
 
 	/** Handle invoking the remote control tab. */
 	void InvokePanelTab();
diff --git a/Engine/Plugins/VirtualProduction/RemoteControl/Source/RemoteControlUI/Private/Commands/RemoteControlCommands.cpp b/Engine/Plugins/VirtualProduction/RemoteControl/Source/RemoteControlUI/Private/Commands/RemoteControlCommands.cpp
--- a/Engine/Plugins/VirtualProduction/RemoteControl/Source/RemoteControlUI/Private/Commands/RemoteControlCommands.cpp
+++ b/Engine/Plugins/VirtualProduction/RemoteControl/Source/RemoteControlUI/Private/Commands/RemoteControlCommands.cpp
@@ -19,9 +19,6 @@ FRemoteControlCommands::FRemoteControlCommands()
 
 void FRemoteControlCommands::RegisterCommands()
 {
-	// Toggle Edit Mode
-	UI_COMMAND(ToggleEditMode, "Toggle Editing", "Toggles edit mode in remote control panel which enables user to modify the exposed entities.", EUserInterfaceActionType::Check, FInputChord(EModifierKey::Control, EKeys::E));
-
 	// Save Preset
 	UI_COMMAND(SavePreset, "Save", "Saves this preset", EUserInterfaceActionType::Button, FInputChord(EModifierKey::Control, EKeys::S));
 
diff --git a/Engine/Plugins/VirtualProduction/RemoteControl/Source/RemoteControlUI/Private/Commands/RemoteControlCommands.h b/Engine/Plugins/VirtualProduction/RemoteControl/Source/RemoteControlUI/Private/Commands/RemoteControlCommands.h
--- a/Engine/Plugins/VirtualProduction/RemoteControl/Source/RemoteControlUI/Private/Commands/RemoteControlCommands.h
+++ b/Engine/Plugins/VirtualProduction/RemoteControl/Source/RemoteControlUI/Private/Commands/RemoteControlCommands.h
@@ -21,11 +21,6 @@ public:
 
 	//~ END : TCommands<> Implementation(s)
 
-	/**
-	 * Holds the information about UI Command that toggles edit mode in RC Panel.
-	 */
-	TSharedPtr<FUICommandInfo> ToggleEditMode;
-
 	/**
 	 * Holds the information about UI Command that finds the actively edited preset in the Content Browser.
 	 */
diff --git a/Engine/Plugins/VirtualProduction/RemoteControl/Source/RemoteControlUI/Private/RemoteControlUIModule.cpp b/Engine/Plugins/VirtualProduction/RemoteControl/Source/RemoteControlUI/Private/RemoteControlUIModule.cpp
--- a/Engine/Plugins/VirtualProduction/RemoteControl/Source/RemoteControlUI/Private/RemoteControlUIModule.cpp
+++ b/Engine/Plugins/VirtualProduction/RemoteControl/Source/RemoteControlUI/Private/RemoteControlUIModule.cpp
@@ -12,7 +12,6 @@
 #include "Framework/MultiBox/MultiBoxBuilder.h"
 #include "HAL/PlatformApplicationMisc.h"
 #include "IRemoteControlModule.h"
-#include "Interfaces/IMainFrameModule.h"
 #include "Kismet2/ComponentEditorUtils.h"
 #include "MaterialEditor/DEditorParameterValue.h"
 #include "Materials/Material.h"
@@ -40,7 +39,6 @@
 
 #define LOCTEXT_NAMESPACE "RemoteControlUI"
 
-const FName FRemoteControlUIModule::EntityDetailsTabName = "RemoteControl_EntityDetails";
 const FName FRemoteControlUIModule::RemoteControlPanelTabName = "RemoteControl_RemoteControlPanel";
 const FString IRemoteControlUIModule::SettingsIniSection = TEXT("RemoteControl");
 
@@ -302,23 +300,25 @@ void FRemoteControlUIModule::UnregisterMetadataCustomization(FName MetadataKey)
 
 TSharedRef<SRemoteControlPanel> FRemoteControlUIModule::CreateRemoteControlPanel(URemoteControlPreset* Preset, const TSharedPtr<IToolkitHost>& ToolkitHost)
 {
+	constexpr bool bIsInLiveMode = true;
+
 	if (TSharedPtr<SRemoteControlPanel> Panel = WeakActivePanel.Pin())
 	{
-		Panel->SetEditMode(false);
+		Panel->SetLiveMode(bIsInLiveMode);
 	}
 
 	TSharedRef<SRemoteControlPanel> PanelRef = SAssignNew(WeakActivePanel, SRemoteControlPanel, Preset, ToolkitHost)
-		.OnEditModeChange_Lambda(
-			[this](TSharedPtr<SRemoteControlPanel> Panel, bool bEditMode) 
+		.OnLiveModeChange_Lambda(
+			[this](TSharedPtr<SRemoteControlPanel> Panel, bool bLiveMode) 
 			{
-				// Activating the edit mode on a panel sets it as the active panel 
-				if (bEditMode)
+				// Activating the live mode on a panel sets it as the active panel 
+				if (bLiveMode)
 				{
 					if (TSharedPtr<SRemoteControlPanel> ActivePanel = WeakActivePanel.Pin())
 					{
 						if (ActivePanel != Panel)
 						{
-							ActivePanel->SetEditMode(false);
+							ActivePanel->SetLiveMode(true);
 						}
 					}
 					WeakActivePanel = MoveTemp(Panel);
@@ -402,29 +402,6 @@ void FRemoteControlUIModule::UnregisterEvents()
 	FEditorDelegates::PostUndoRedo.RemoveAll(this);
 }
 
-void FRemoteControlUIModule::ToggleEditMode()
-{
-	if (TSharedPtr<SRemoteControlPanel> Panel = GetPanelForObject(nullptr)) // Check whether the panel is active.
-	{
-		Panel->SetEditMode(Panel->IsInEditMode() ? false : true);
-	}
-}
-
-bool FRemoteControlUIModule::CanToggleEditMode() const
-{
-	return GetPanelForObject(nullptr).IsValid();
-}
-
-bool FRemoteControlUIModule::IsInEditMode() const
-{
-	if (TSharedPtr<SRemoteControlPanel> Panel = GetPanelForObject(nullptr)) // Check whether the panel is active.
-	{
-		return Panel->IsInEditMode();
-	}
-
-	return false;
-}
-
 URemoteControlPreset* FRemoteControlUIModule::GetActivePreset() const
 {
 	if (const TSharedPtr<SRemoteControlPanel> Panel = GetPanelForObject(nullptr))
@@ -462,33 +439,10 @@ void FRemoteControlUIModule::UnregisterAssetTools()
 void FRemoteControlUIModule::BindRemoteControlCommands()
 {
 	FRemoteControlCommands::Register();
-
-	const FRemoteControlCommands& Commands = FRemoteControlCommands::Get();
-
-	IMainFrameModule& MainFrame = FModuleManager::Get().LoadModuleChecked<IMainFrameModule>("MainFrame");
-
-	FUICommandList& ActionList = *MainFrame.GetMainFrameCommandBindings();
-
-	// Toggle Edit Mode
-
-	ActionList.MapAction(Commands.ToggleEditMode,
-		FExecuteAction::CreateRaw(this, &FRemoteControlUIModule::ToggleEditMode),
-		FCanExecuteAction::CreateRaw(this, &FRemoteControlUIModule::CanToggleEditMode),
-		FIsActionChecked::CreateRaw(this, &FRemoteControlUIModule::IsInEditMode),
-		FIsActionButtonVisible::CreateRaw(this, &FRemoteControlUIModule::CanToggleEditMode)
-	);
 }
 
 void FRemoteControlUIModule::UnbindRemoteControlCommands()
 {
-	const FRemoteControlCommands& Commands = FRemoteControlCommands::Get();
-
-	IMainFrameModule& MainFrame = FModuleManager::Get().LoadModuleChecked<IMainFrameModule>("MainFrame");
-
-	FUICommandList& ActionList = *MainFrame.GetMainFrameCommandBindings();
-
-	ActionList.UnmapAction(Commands.ToggleEditMode);
-
 	FRemoteControlCommands::Unregister();
 }
 
diff --git a/Engine/Plugins/VirtualProduction/RemoteControl/Source/RemoteControlUI/Private/RemoteControlUIModule.h b/Engine/Plugins/VirtualProduction/RemoteControl/Source/RemoteControlUI/Private/RemoteControlUIModule.h
--- a/Engine/Plugins/VirtualProduction/RemoteControl/Source/RemoteControlUI/Private/RemoteControlUIModule.h
+++ b/Engine/Plugins/VirtualProduction/RemoteControl/Source/RemoteControlUI/Private/RemoteControlUIModule.h
@@ -73,7 +73,6 @@ public:
 	}
 
 public:
-	static const FName EntityDetailsTabName;
 	static const FName RemoteControlPanelTabName;
 
 	TSharedPtr<SRCPanelTreeNode> GenerateEntityWidget(const FGenerateWidgetArgs& Args);
@@ -110,11 +109,6 @@ private:
 	void RegisterEvents();
 	void UnregisterEvents();
 
-	//~ Toggle Edit Mode RC Command.
-	void ToggleEditMode();
-	bool CanToggleEditMode() const;	
-	bool IsInEditMode() const;
-
 	/** Handle creating the row extensions.  */
 	void HandleCreatePropertyRowExtension(const FOnGenerateGlobalRowExtensionArgs& InArgs, TArray<FPropertyRowExtensionButton>& OutExtensions);
 
diff --git a/Engine/Plugins/VirtualProduction/RemoteControl/Source/RemoteControlUI/Private/UI/BaseLogicUI/SRCLogicPanelListBase.cpp b/Engine/Plugins/VirtualProduction/RemoteControl/Source/RemoteControlUI/Private/UI/BaseLogicUI/SRCLogicPanelListBase.cpp
--- a/Engine/Plugins/VirtualProduction/RemoteControl/Source/RemoteControlUI/Private/UI/BaseLogicUI/SRCLogicPanelListBase.cpp
+++ b/Engine/Plugins/VirtualProduction/RemoteControl/Source/RemoteControlUI/Private/UI/BaseLogicUI/SRCLogicPanelListBase.cpp
@@ -56,7 +56,8 @@ TSharedPtr<SWidget> SRCLogicPanelListBase::GetContextMenuWidget()
 	MenuBuilder.AddMenuEntry(Commands.DeleteEntity);
 
 	// 4. Delete All
-	FUIAction Action(FExecuteAction::CreateSP(this, &SRCLogicPanelListBase::RequestDeleteAllItems));
+	FUIAction Action(FExecuteAction::CreateSP(this, &SRCLogicPanelListBase::RequestDeleteAllItems)
+	, FCanExecuteAction::CreateSP(this, &SRCLogicPanelListBase::CanDeleteAllItems));
 
 	MenuBuilder.AddMenuEntry(LOCTEXT("DeleteAll", "Delete All"),
 		LOCTEXT("ContextMenuEditTooltip", "Delete all the rows in this list"),
@@ -70,6 +71,16 @@ TSharedPtr<SWidget> SRCLogicPanelListBase::GetContextMenuWidget()
 	return MenuWidget;
 }
 
+bool SRCLogicPanelListBase::CanDeleteAllItems() const
+{
+	if (TSharedPtr<SRemoteControlPanel> RemoteControlPanel = RemoteControlPanelWeakPtr.Pin())
+	{
+		return !RemoteControlPanel->IsInLiveMode();
+	}
+
+	return false;
+}
+
 void SRCLogicPanelListBase::RequestDeleteAllItems()
 {
 	if (TSharedPtr<SRCLogicPanelBase> LogicParentPanel = LogicPanelWeakPtr.Pin())
diff --git a/Engine/Plugins/VirtualProduction/RemoteControl/Source/RemoteControlUI/Private/UI/BaseLogicUI/SRCLogicPanelListBase.h b/Engine/Plugins/VirtualProduction/RemoteControl/Source/RemoteControlUI/Private/UI/BaseLogicUI/SRCLogicPanelListBase.h
--- a/Engine/Plugins/VirtualProduction/RemoteControl/Source/RemoteControlUI/Private/UI/BaseLogicUI/SRCLogicPanelListBase.h
+++ b/Engine/Plugins/VirtualProduction/RemoteControl/Source/RemoteControlUI/Private/UI/BaseLogicUI/SRCLogicPanelListBase.h
@@ -80,6 +80,9 @@ private:
 	/** Removes the given UI model item from the list of UI models for this panel list*/
 	virtual int32 RemoveModel(const TSharedPtr<FRCLogicModeBase> InModel) = 0;
 
+	/** Returns true when we can delete all items. */
+	bool CanDeleteAllItems() const;
+	
 	/** Requests the parent logic panel to delete all items in this list*/
 	void RequestDeleteAllItems();
 
diff --git a/Engine/Plugins/VirtualProduction/RemoteControl/Source/RemoteControlUI/Private/UI/Controller/SRCControllerPanel.cpp b/Engine/Plugins/VirtualProduction/RemoteControl/Source/RemoteControlUI/Private/UI/Controller/SRCControllerPanel.cpp
--- a/Engine/Plugins/VirtualProduction/RemoteControl/Source/RemoteControlUI/Private/UI/Controller/SRCControllerPanel.cpp
+++ b/Engine/Plugins/VirtualProduction/RemoteControl/Source/RemoteControlUI/Private/UI/Controller/SRCControllerPanel.cpp
@@ -33,6 +33,8 @@ void SRCControllerPanel::Construct(const FArguments& InArgs, const TSharedRef<SR
 {
 	SRCLogicPanelBase::Construct(SRCLogicPanelBase::FArguments(), InPanel);
 	
+	bIsInLiveMode = InArgs._LiveMode;
+
 	RCPanelStyle = &FRemoteControlPanelStyle::Get()->GetWidgetStyle<FRCPanelStyle>("RemoteControlPanel.MinorPanel");
 
 	// Controller Dock Panel
@@ -47,6 +49,7 @@ void SRCControllerPanel::Construct(const FArguments& InArgs, const TSharedRef<SR
 		.AddMetaData<FTagMetaData>(FTagMetaData(TEXT("Add Controller")))
 		.HAlign(HAlign_Center)
 		.VAlign(VAlign_Center)
+		.IsEnabled_Lambda([this]() { return !bIsInLiveMode.Get(); })
 		.ButtonStyle(&RCPanelStyle->FlatButtonStyle)
 		.ForegroundColor(FSlateColor::UseForeground())
 		.CollapseMenuOnParentFocus(true)
@@ -77,7 +80,7 @@ void SRCControllerPanel::Construct(const FArguments& InArgs, const TSharedRef<SR
 		.ButtonStyle(&RCPanelStyle->FlatButtonStyle)
 		.ToolTipText(LOCTEXT("DeleteSelectedControllerToolTip", "Deletes the selected controller."))
 		.OnClicked(this, &SRCControllerPanel::RequestDeleteSelectedItem)
-		.IsEnabled_Lambda([this]() { return ControllerPanelList.IsValid() && ControllerPanelList->NumSelectedLogicItems() > 0; })
+		.IsEnabled_Lambda([this]() { return ControllerPanelList.IsValid() && !bIsInLiveMode.Get() && ControllerPanelList->NumSelectedLogicItems() > 0; })
 		[
 			SNew(SBox)
 			.WidthOverride(RCPanelStyle->IconSize.X)
diff --git a/Engine/Plugins/VirtualProduction/RemoteControl/Source/RemoteControlUI/Private/UI/Controller/SRCControllerPanel.h b/Engine/Plugins/VirtualProduction/RemoteControl/Source/RemoteControlUI/Private/UI/Controller/SRCControllerPanel.h
--- a/Engine/Plugins/VirtualProduction/RemoteControl/Source/RemoteControlUI/Private/UI/Controller/SRCControllerPanel.h
+++ b/Engine/Plugins/VirtualProduction/RemoteControl/Source/RemoteControlUI/Private/UI/Controller/SRCControllerPanel.h
@@ -21,8 +21,9 @@ class REMOTECONTROLUI_API SRCControllerPanel : public SRCLogicPanelBase
 {
 public:
 	SLATE_BEGIN_ARGS(SRCControllerPanel)
-		{
-		}
+	{}
+
+		SLATE_ATTRIBUTE(bool, LiveMode)
 
 	SLATE_END_ARGS()
 
@@ -77,6 +78,9 @@ private:
 	/** Widget representing List of Controllers */
 	TSharedPtr<class SRCControllerPanelList> ControllerPanelList;
 
+	/** Whether the panel is in live mode. */
+	TAttribute<bool> bIsInLiveMode;
+
 	/** Panel Style reference. */
 	const FRCPanelStyle* RCPanelStyle;
 };
diff --git a/Engine/Plugins/VirtualProduction/RemoteControl/Source/RemoteControlUI/Private/UI/Drawers/SRCPanelDrawer.cpp b/Engine/Plugins/VirtualProduction/RemoteControl/Source/RemoteControlUI/Private/UI/Drawers/SRCPanelDrawer.cpp
--- a/Engine/Plugins/VirtualProduction/RemoteControl/Source/RemoteControlUI/Private/UI/Drawers/SRCPanelDrawer.cpp
+++ b/Engine/Plugins/VirtualProduction/RemoteControl/Source/RemoteControlUI/Private/UI/Drawers/SRCPanelDrawer.cpp
@@ -419,6 +419,16 @@ bool SRCPanelDrawer::UnregisterPanel(TSharedRef<FRCPanelDrawerArgs> InPanelToRem
 
 void SRCPanelDrawer::TogglePanel(TSharedRef<FRCPanelDrawerArgs> InPanel, const bool bRemoveDrawer)
 {
+	if (!CanToggleRCPanelDelegate.IsBound())
+	{
+		return;
+	}
+
+	if (!CanToggleRCPanelDelegate.Execute())
+	{
+		return;
+	}
+
 	auto PanelFinderPredicate = [InPanel](const TPair<TSharedRef<FRCPanelDrawerArgs>, TSharedRef<SRCPanelDrawerButton>>& TabAndButton)
 	{
 		return TabAndButton.Key == InPanel;
@@ -450,7 +460,7 @@ void SRCPanelDrawer::TogglePanel(TSharedRef<FRCPanelDrawerArgs> InPanel, const b
 		RegisterPanel(InPanel);
 	}
 
-	OnRCPanelToggledHandle.ExecuteIfBound(OpenedDrawer);
+	OnRCPanelToggledDelegate.ExecuteIfBound(OpenedDrawer);
 }
 
 void SRCPanelDrawer::UpdateAppearance()
diff --git a/Engine/Plugins/VirtualProduction/RemoteControl/Source/RemoteControlUI/Private/UI/Drawers/SRCPanelDrawer.h b/Engine/Plugins/VirtualProduction/RemoteControl/Source/RemoteControlUI/Private/UI/Drawers/SRCPanelDrawer.h
--- a/Engine/Plugins/VirtualProduction/RemoteControl/Source/RemoteControlUI/Private/UI/Drawers/SRCPanelDrawer.h
+++ b/Engine/Plugins/VirtualProduction/RemoteControl/Source/RemoteControlUI/Private/UI/Drawers/SRCPanelDrawer.h
@@ -23,6 +23,8 @@ enum class ERCPanels : uint8
 
 	RCP_OutputLog,
 
+	RCP_Live,
+
 	RCP_Count,
 };
 
@@ -119,6 +121,9 @@ DECLARE_DELEGATE_TwoParams(FOnRCPanelDrawerButtonPressed, TSharedRef<FRCPanelDra
 /** Delegate to be called when a drawer button is clicked. */
 DECLARE_DELEGATE_OneParam(FOnRCPanelToggled, ERCPanels /* PanelID */)
 
+/** Delegate to be called to check whether drawer button can be clicked or not. */
+DECLARE_DELEGATE_RetVal(bool, FCanToggleRCPanel)
+
 /**
  * A custom tab drawer tailored for RC Panels.
  */
@@ -152,8 +157,11 @@ public:
 	 */
 	bool UnregisterPanel(TSharedRef<FRCPanelDrawerArgs> InPanelToRemove);
 
+	/** Delegate to be called to check whether drawer button can be clicked or not. */
+	FCanToggleRCPanel& CanToggleRCPanel() { return CanToggleRCPanelDelegate; }
+	
 	/** Delegate to be called when a drawer button is clicked. */
-	FOnRCPanelToggled& OnRCPanelToggled() { return OnRCPanelToggledHandle; }
+	FOnRCPanelToggled& OnRCPanelToggled() { return OnRCPanelToggledDelegate; }
 
 	/** Toggles the panel. */
 	void TogglePanel(TSharedRef<FRCPanelDrawerArgs> InPanel, const bool bRemoveDrawer = false);
@@ -171,8 +179,11 @@ private:
 	/** Array of all registered panels. */
 	TArray<TPair<TSharedRef<FRCPanelDrawerArgs>, TSharedRef<SRCPanelDrawerButton>>> Panels;
 
+	/** Delegate to be called to check whether drawer button can be clicked or not. */
+	FCanToggleRCPanel CanToggleRCPanelDelegate;
+	
 	/** Delegate to be called when a drawer button is clicked. */
-	FOnRCPanelToggled OnRCPanelToggledHandle;
+	FOnRCPanelToggled OnRCPanelToggledDelegate;
 
 	/** Generally speaking one drawer is only ever open at once. */
 	ERCPanels OpenedDrawer;
diff --git a/Engine/Plugins/VirtualProduction/RemoteControl/Source/RemoteControlUI/Private/UI/SRCPanelExposedActor.cpp b/Engine/Plugins/VirtualProduction/RemoteControl/Source/RemoteControlUI/Private/UI/SRCPanelExposedActor.cpp
--- a/Engine/Plugins/VirtualProduction/RemoteControl/Source/RemoteControlUI/Private/UI/SRCPanelExposedActor.cpp
+++ b/Engine/Plugins/VirtualProduction/RemoteControl/Source/RemoteControlUI/Private/UI/SRCPanelExposedActor.cpp
@@ -16,7 +16,7 @@
 
 TSharedPtr<SRCPanelTreeNode> SRCPanelExposedActor::MakeInstance(const FGenerateWidgetArgs& Args)
 {
-	return SNew(SRCPanelExposedActor, StaticCastSharedPtr<FRemoteControlActor>(Args.Entity), Args.Preset, Args.ColumnSizeData).EditMode(Args.bIsInEditMode);
+	return SNew(SRCPanelExposedActor, StaticCastSharedPtr<FRemoteControlActor>(Args.Entity), Args.Preset, Args.ColumnSizeData).LiveMode(Args.bIsInLiveMode);
 }
 
 void SRCPanelExposedActor::Construct(const FArguments& InArgs, TWeakPtr<FRemoteControlActor> InWeakActor, URemoteControlPreset* InPreset, FRCColumnSizeData InColumnSizeData)
@@ -26,11 +26,11 @@ void SRCPanelExposedActor::Construct(const FArguments& InArgs, TWeakPtr<FRemoteC
 	const TSharedPtr<FRemoteControlActor> Actor = InWeakActor.Pin();
 	if (ensure(Actor))
 	{
-		Initialize(Actor->GetId(), InPreset, InArgs._EditMode);
+		Initialize(Actor->GetId(), InPreset, InArgs._LiveMode);
 		
 		ColumnSizeData = MoveTemp(InColumnSizeData);
 		WeakPreset = InPreset;
-		bEditMode = InArgs._EditMode;
+		bLiveMode = InArgs._LiveMode;
 		WeakActor = MoveTemp(InWeakActor);
 		
 		if (AActor* ResolvedActor = Actor->GetActor())
diff --git a/Engine/Plugins/VirtualProduction/RemoteControl/Source/RemoteControlUI/Private/UI/SRCPanelExposedActor.h b/Engine/Plugins/VirtualProduction/RemoteControl/Source/RemoteControlUI/Private/UI/SRCPanelExposedActor.h
--- a/Engine/Plugins/VirtualProduction/RemoteControl/Source/RemoteControlUI/Private/UI/SRCPanelExposedActor.h
+++ b/Engine/Plugins/VirtualProduction/RemoteControl/Source/RemoteControlUI/Private/UI/SRCPanelExposedActor.h
@@ -22,9 +22,9 @@ struct SRCPanelExposedActor : public SRCPanelExposedEntity
 	using SCompoundWidget::AsShared;
 	
 	SLATE_BEGIN_ARGS(SRCPanelExposedActor)
-		: _EditMode(true)
+		: _LiveMode(false)
 	{}
-		SLATE_ATTRIBUTE(bool, EditMode)
+		SLATE_ATTRIBUTE(bool, LiveMode)
 	SLATE_END_ARGS()
 
 	static TSharedPtr<SRCPanelTreeNode> MakeInstance(const FGenerateWidgetArgs& Args);
@@ -50,6 +50,6 @@ private:
 	TWeakObjectPtr<URemoteControlPreset> WeakPreset;
 	/** Weak ptr to the remote control actor structure. */
 	TWeakPtr<FRemoteControlActor> WeakActor;
-	/** Holds this row's panel edit mode. */
-	TAttribute<bool> bEditMode;
+	/** Holds this row's panel live mode. */
+	TAttribute<bool> bLiveMode;
 };
diff --git a/Engine/Plugins/VirtualProduction/RemoteControl/Source/RemoteControlUI/Private/UI/SRCPanelExposedEntitiesList.cpp b/Engine/Plugins/VirtualProduction/RemoteControl/Source/RemoteControlUI/Private/UI/SRCPanelExposedEntitiesList.cpp
--- a/Engine/Plugins/VirtualProduction/RemoteControl/Source/RemoteControlUI/Private/UI/SRCPanelExposedEntitiesList.cpp
+++ b/Engine/Plugins/VirtualProduction/RemoteControl/Source/RemoteControlUI/Private/UI/SRCPanelExposedEntitiesList.cpp
@@ -30,6 +30,7 @@
 #include "Widgets/Input/SButton.h"
 #include "Widgets/Layout/SSpacer.h"
 #include "Widgets/Layout/SSeparator.h"
+#include "Widgets/Layout/SWidgetSwitcher.h"
 #include "Widgets/Views/SListView.h"
 #include "Widgets/Views/STableRow.h"
 #include "Widgets/Views/STreeView.h"
@@ -167,7 +168,7 @@ private:
 
 void SRCPanelExposedEntitiesList::Construct(const FArguments& InArgs, URemoteControlPreset* InPreset, TWeakPtr<FRCPanelWidgetRegistry> InWidgetRegistry)
 {
-	bIsInEditMode = InArgs._EditMode;
+	bIsInLiveMode = InArgs._LiveMode;
 	bIsInProtocolsMode = InArgs._ProtocolsMode;
 	Preset = TStrongObjectPtr<URemoteControlPreset>(InPreset);
 	OnEntityListUpdatedDelegate = InArgs._OnEntityListUpdated;
@@ -211,7 +212,7 @@ void SRCPanelExposedEntitiesList::Construct(const FArguments& InArgs, URemoteCon
 	// Add New Group Button
 	TSharedPtr<SWidget> NewGroupButton = SNew(SButton)
 		.AddMetaData<FTagMetaData>(FTagMetaData(TEXT("Add Group")))
-		.IsEnabled_Lambda([this]() { return bIsInEditMode.Get(); })
+		.IsEnabled_Lambda([this]() { return !bIsInLiveMode.Get(); })
 		.HAlign(HAlign_Center)
 		.VAlign(VAlign_Center)
 		.ForegroundColor(FSlateColor::UseForeground())
@@ -232,8 +233,8 @@ void SRCPanelExposedEntitiesList::Construct(const FArguments& InArgs, URemoteCon
 	// Delete All Groups Button
 	TSharedPtr<SWidget> DeleteAllGroupsButton = SNew(SButton)
 		.AddMetaData<FTagMetaData>(FTagMetaData(TEXT("Delete All Groups")))
-		.IsEnabled_Lambda([this]() { return bIsInEditMode.Get() && FieldGroups.Num() > 0; })
-		.Visibility_Lambda([this]() { return (bIsInEditMode.Get() && FieldGroups.Num() > 0) ? EVisibility::Visible : EVisibility::Collapsed; })
+		.IsEnabled_Lambda([this]() { return !bIsInLiveMode.Get() && FieldGroups.Num() > 0; })
+		.Visibility_Lambda([this]() { return (!bIsInLiveMode.Get() && FieldGroups.Num() > 0) ? EVisibility::Visible : EVisibility::Collapsed; })
 		.HAlign(HAlign_Center)
 		.VAlign(VAlign_Center)
 		.ForegroundColor(FSlateColor::UseForeground())
@@ -297,6 +298,7 @@ void SRCPanelExposedEntitiesList::Construct(const FArguments& InArgs, URemoteCon
 	// Exposed Entities Dock Panel
 	TSharedPtr<SRCMinorPanel> ExposeDockPanel = SNew(SRCMinorPanel)
 		.HeaderLabel(this, &SRCPanelExposedEntitiesList::HandleEntityListHeaderLabel)
+		.Visibility_Lambda([this]() {return bIsInLiveMode.Get() ? EVisibility::Collapsed : EVisibility::Visible; })
 		.EnableFooter(true)
 		[
 			FieldsListView.ToSharedRef()
@@ -314,8 +316,8 @@ void SRCPanelExposedEntitiesList::Construct(const FArguments& InArgs, URemoteCon
 	// Delete All Entities Button
 	TSharedPtr<SWidget> DeleteAllEntitiesButton = SNew(SButton)
 		.AddMetaData<FTagMetaData>(FTagMetaData(TEXT("Delete All Entities")))
-		.IsEnabled_Lambda([this]() { return bIsInEditMode.Get() && FieldWidgetMap.Num() > 0; })
-		.Visibility_Lambda([this]() { return (bIsInEditMode.Get() && FieldWidgetMap.Num() > 0) ? EVisibility::Visible : EVisibility::Collapsed; })
+		.IsEnabled_Lambda([this]() { return !bIsInLiveMode.Get() && FieldWidgetMap.Num() > 0; })
+		.Visibility_Lambda([this]() { return (!bIsInLiveMode.Get() && FieldWidgetMap.Num() > 0) ? EVisibility::Visible : EVisibility::Collapsed; })
 		.HAlign(HAlign_Center)
 		.VAlign(VAlign_Center)
 		.ForegroundColor(FSlateColor::UseForeground())
@@ -714,7 +716,7 @@ void SRCPanelExposedEntitiesList::GenerateListWidgets()
 			Args.Preset = Preset.Get();
 			Args.WidgetRegistry = WidgetRegistry;
 			Args.ColumnSizeData = ColumnSizeData;
-			Args.bIsInEditMode = bIsInEditMode;
+			Args.bIsInLiveMode = bIsInLiveMode;
 
 			FieldWidgetMap.Add(Entity->GetId(), FRemoteControlUIModule::Get().GenerateEntityWidget(Args));
 		}
@@ -756,7 +758,7 @@ void SRCPanelExposedEntitiesList::RefreshGroups()
 			.OnFieldDropEvent(this, &SRCPanelExposedEntitiesList::OnDropOnGroup)
 			.OnGetGroupId(this, &SRCPanelExposedEntitiesList::GetGroupId)
 			.OnDeleteGroup(this, &SRCPanelExposedEntitiesList::OnDeleteGroup)
-			.EditMode(bIsInEditMode);
+			.LiveMode(bIsInLiveMode);
 		
 		FieldGroups.Add(FieldGroup);
 		FieldGroup->GetNodes().Reserve(RCGroup.GetFields().Num());
@@ -1088,7 +1090,7 @@ void SRCPanelExposedEntitiesList::OnEntityAdded(const FGuid& InEntityId)
 	Args.Preset = Preset.Get();
 	Args.WidgetRegistry = WidgetRegistry;
 	Args.ColumnSizeData = ColumnSizeData;
-	Args.bIsInEditMode = bIsInEditMode;
+	Args.bIsInLiveMode = bIsInLiveMode;
 	Args.Entity = Preset->GetExposedEntity(InEntityId).Pin();
 
 	ExposeEntity(FRemoteControlUIModule::Get().GenerateEntityWidget(Args));
@@ -1126,7 +1128,7 @@ void SRCPanelExposedEntitiesList::OnGroupAdded(const FRemoteControlPresetGroup&
 		.OnFieldDropEvent(this, &SRCPanelExposedEntitiesList::OnDropOnGroup)
 		.OnGetGroupId(this, &SRCPanelExposedEntitiesList::GetGroupId)
 		.OnDeleteGroup(this, &SRCPanelExposedEntitiesList::OnDeleteGroup)
-		.EditMode(bIsInEditMode);
+		.LiveMode(bIsInLiveMode);
 	
 	FieldGroups.Add(FieldGroup);
 	
diff --git a/Engine/Plugins/VirtualProduction/RemoteControl/Source/RemoteControlUI/Private/UI/SRCPanelExposedEntitiesList.h b/Engine/Plugins/VirtualProduction/RemoteControl/Source/RemoteControlUI/Private/UI/SRCPanelExposedEntitiesList.h
--- a/Engine/Plugins/VirtualProduction/RemoteControl/Source/RemoteControlUI/Private/UI/SRCPanelExposedEntitiesList.h
+++ b/Engine/Plugins/VirtualProduction/RemoteControl/Source/RemoteControlUI/Private/UI/SRCPanelExposedEntitiesList.h
@@ -58,9 +58,9 @@ class SRCPanelExposedEntitiesList : public SCompoundWidget
 {
 public:
 	SLATE_BEGIN_ARGS(SRCPanelExposedEntitiesList)
-		: _EditMode(true)
+		: _LiveMode(false)
 	{}
-		SLATE_ATTRIBUTE(bool, EditMode)
+		SLATE_ATTRIBUTE(bool, LiveMode)
 		SLATE_ATTRIBUTE(bool, ProtocolsMode)
 		SLATE_ATTRIBUTE(TSharedPtr<SWidget>, ExposeComboButton)
 		SLATE_EVENT(FSimpleDelegate, OnEntityListUpdated)
@@ -192,8 +192,8 @@ private:
 	TArray<TSharedPtr<SRCPanelTreeNode>> FieldEntities;
 	/** Map of field ids to field widgets. */
 	TMap<FGuid, TSharedPtr<SRCPanelTreeNode>> FieldWidgetMap;
-	/** Whether the panel is in edit mode. */
-	TAttribute<bool> bIsInEditMode;
+	/** Whether the panel is in live mode. */
+	TAttribute<bool> bIsInLiveMode;
 	/** Whether the panel is in protocols mode. */
 	TAttribute<bool> bIsInProtocolsMode;
 	/** Holds the preset asset. */
diff --git a/Engine/Plugins/VirtualProduction/RemoteControl/Source/RemoteControlUI/Private/UI/SRCPanelExposedEntity.cpp b/Engine/Plugins/VirtualProduction/RemoteControl/Source/RemoteControlUI/Private/UI/SRCPanelExposedEntity.cpp
--- a/Engine/Plugins/VirtualProduction/RemoteControl/Source/RemoteControlUI/Private/UI/SRCPanelExposedEntity.cpp
+++ b/Engine/Plugins/VirtualProduction/RemoteControl/Source/RemoteControlUI/Private/UI/SRCPanelExposedEntity.cpp
@@ -99,11 +99,11 @@ void SRCPanelExposedEntity::EnterRenameMode()
 	}
 }
 
-void SRCPanelExposedEntity::Initialize(const FGuid& InEntityId, URemoteControlPreset* InPreset, const TAttribute<bool>& InbEditMode)
+void SRCPanelExposedEntity::Initialize(const FGuid& InEntityId, URemoteControlPreset* InPreset, const TAttribute<bool>& InbLiveMode)
 {
 	EntityId = InEntityId;
 	Preset = InPreset;
-	bEditMode = InbEditMode;
+	bLiveMode = InbLiveMode;
 
 	RCPanelStyle = &FRemoteControlPanelStyle::Get()->GetWidgetStyle<FRCPanelStyle>("RemoteControlPanel.MinorPanel");
 
@@ -200,9 +200,9 @@ TSharedRef<SWidget> SRCPanelExposedEntity::CreateInvalidWidget()
 		];
 }
 
-EVisibility SRCPanelExposedEntity::GetVisibilityAccordingToEditMode(EVisibility NonEditModeVisibility) const
+EVisibility SRCPanelExposedEntity::GetVisibilityAccordingToLiveMode(EVisibility NonEditModeVisibility) const
 {
-	return bEditMode.Get() ? EVisibility::Visible : NonEditModeVisibility;
+	return !bLiveMode.Get() ? EVisibility::Visible : NonEditModeVisibility;
 }
 
 TSharedRef<SWidget> SRCPanelExposedEntity::CreateRebindMenuContent()
@@ -303,7 +303,7 @@ TSharedRef<SWidget> SRCPanelExposedEntity::CreateEntityWidget(TSharedPtr<SWidget
 		.BorderImage(this, &SRCPanelExposedEntity::GetBorderImage);
 	
 	Args.DragHandle = SNew(SBox)
-		.Visibility(this, &SRCPanelExposedEntity::GetVisibilityAccordingToEditMode, EVisibility::Collapsed)
+		.Visibility(this, &SRCPanelExposedEntity::GetVisibilityAccordingToLiveMode, EVisibility::Collapsed)
 		[
 			SNew(SRCPanelDragHandle<FExposedEntityDragDrop>, GetRCId())
 			.Widget(Widget)
@@ -331,7 +331,7 @@ TSharedRef<SWidget> SRCPanelExposedEntity::CreateEntityWidget(TSharedPtr<SWidget
 			.Text(FText::FromName(CachedLabel))
 			.OnTextCommitted(this, &SRCPanelExposedEntity::OnLabelCommitted)
 			.OnVerifyTextChanged(this, &SRCPanelExposedEntity::OnVerifyItemLabelChanged)
-			.IsReadOnly_Lambda([this]() { return !bEditMode.Get(); })
+			.IsReadOnly_Lambda([this]() { return bLiveMode.Get(); })
 			.HighlightText_Lambda([this]() { return HighlightText.Get().ToString().Len() > 3 ? HighlightText.Get() : FText::GetEmpty(); })
 		];
 
diff --git a/Engine/Plugins/VirtualProduction/RemoteControl/Source/RemoteControlUI/Private/UI/SRCPanelExposedEntity.h b/Engine/Plugins/VirtualProduction/RemoteControl/Source/RemoteControlUI/Private/UI/SRCPanelExposedEntity.h
--- a/Engine/Plugins/VirtualProduction/RemoteControl/Source/RemoteControlUI/Private/UI/SRCPanelExposedEntity.h
+++ b/Engine/Plugins/VirtualProduction/RemoteControl/Source/RemoteControlUI/Private/UI/SRCPanelExposedEntity.h
@@ -31,13 +31,13 @@ struct SRCPanelExposedEntity : public SRCPanelTreeNode
 	virtual void SetHighlightText(const FText& InHightlightText = FText::GetEmpty()) override { HighlightText = InHightlightText; }
 
 protected:
-	void Initialize(const FGuid& InEntityId, URemoteControlPreset* InPreset, const TAttribute<bool>& InbEditMode);
+	void Initialize(const FGuid& InEntityId, URemoteControlPreset* InPreset, const TAttribute<bool>& InbLiveMode);
 	
 	/** Create a widget that displays the rebind button. */
 	TSharedRef<SWidget> CreateInvalidWidget();
 
 	/** Get the widget's visibility according to the panel's mode. */
-	EVisibility GetVisibilityAccordingToEditMode(EVisibility NonEditModeVisibility) const;
+	EVisibility GetVisibilityAccordingToLiveMode(EVisibility NonEditModeVisibility) const;
 
 	/** Create an exposed entity widget with a drag handle and unexpose button. */
 	TSharedRef<SWidget> CreateEntityWidget(TSharedPtr<SWidget> ValueWidget, TSharedPtr<SWidget> ResetWidget = SNullWidget::NullWidget, const FText& OptionalWarningMessage = FText::GetEmpty());
@@ -47,8 +47,8 @@ protected:
 	FGuid EntityId;
 	/** The underlying preset. */
 	TWeakObjectPtr<URemoteControlPreset> Preset;
-	/** Whether the panel is in edit mode or not. */
-	TAttribute<bool> bEditMode;
+	/** Whether the panel is in live mode or not. */
+	TAttribute<bool> bLiveMode;
 	/** Display name of the entity. */
 	FName CachedLabel;
 	/** Text to be highlighted while searching. */
diff --git a/Engine/Plugins/VirtualProduction/RemoteControl/Source/RemoteControlUI/Private/UI/SRCPanelExposedField.cpp b/Engine/Plugins/VirtualProduction/RemoteControl/Source/RemoteControlUI/Private/UI/SRCPanelExposedField.cpp
--- a/Engine/Plugins/VirtualProduction/RemoteControl/Source/RemoteControlUI/Private/UI/SRCPanelExposedField.cpp
+++ b/Engine/Plugins/VirtualProduction/RemoteControl/Source/RemoteControlUI/Private/UI/SRCPanelExposedField.cpp
@@ -66,7 +66,7 @@ namespace ExposedFieldUtils
 
 TSharedPtr<SRCPanelTreeNode> SRCPanelExposedField::MakeInstance(const FGenerateWidgetArgs& Args)
 {
-	return SNew(SRCPanelExposedField, StaticCastSharedPtr<FRemoteControlField>(Args.Entity), Args.ColumnSizeData, Args.WidgetRegistry).Preset(Args.Preset).EditMode(Args.bIsInEditMode).HighlightText(Args.HighlightText);
+	return SNew(SRCPanelExposedField, StaticCastSharedPtr<FRemoteControlField>(Args.Entity), Args.ColumnSizeData, Args.WidgetRegistry).Preset(Args.Preset).LiveMode(Args.bIsInLiveMode).HighlightText(Args.HighlightText);
 }
 
 void SRCPanelExposedField::Construct(const FArguments& InArgs, TWeakPtr<FRemoteControlField> InField, FRCColumnSizeData InColumnSizeData, TWeakPtr<FRCPanelWidgetRegistry> InWidgetRegistry)
@@ -82,7 +82,7 @@ void SRCPanelExposedField::Construct(const FArguments& InArgs, TWeakPtr<FRemoteC
 
 	if (TSharedPtr<FRemoteControlField> FieldPtr = WeakField.Pin())
 	{
-		Initialize(FieldPtr->GetId(), InArgs._Preset.Get(), InArgs._EditMode);
+		Initialize(FieldPtr->GetId(), InArgs._Preset.Get(), InArgs._LiveMode);
 	
 		CachedLabel = FieldPtr->GetLabel();
 		EntityId = FieldPtr->GetId();
diff --git a/Engine/Plugins/VirtualProduction/RemoteControl/Source/RemoteControlUI/Private/UI/SRCPanelExposedField.h b/Engine/Plugins/VirtualProduction/RemoteControl/Source/RemoteControlUI/Private/UI/SRCPanelExposedField.h
--- a/Engine/Plugins/VirtualProduction/RemoteControl/Source/RemoteControlUI/Private/UI/SRCPanelExposedField.h
+++ b/Engine/Plugins/VirtualProduction/RemoteControl/Source/RemoteControlUI/Private/UI/SRCPanelExposedField.h
@@ -29,11 +29,11 @@ class SWidget;
 struct SRCPanelExposedField : public SRCPanelExposedEntity
 {
 	SLATE_BEGIN_ARGS(SRCPanelExposedField)
-		: _EditMode(true)
+		: _LiveMode(false)
 		, _HighlightText()
 		, _Preset(nullptr)
 	{}
-		SLATE_ATTRIBUTE(bool, EditMode)
+		SLATE_ATTRIBUTE(bool, LiveMode)
 		SLATE_ATTRIBUTE(FText, HighlightText)
 		SLATE_ATTRIBUTE(URemoteControlPreset*, Preset)
 	SLATE_END_ARGS()
diff --git a/Engine/Plugins/VirtualProduction/RemoteControl/Source/RemoteControlUI/Private/UI/SRCPanelFieldGroup.cpp b/Engine/Plugins/VirtualProduction/RemoteControl/Source/RemoteControlUI/Private/UI/SRCPanelFieldGroup.cpp
--- a/Engine/Plugins/VirtualProduction/RemoteControl/Source/RemoteControlUI/Private/UI/SRCPanelFieldGroup.cpp
+++ b/Engine/Plugins/VirtualProduction/RemoteControl/Source/RemoteControlUI/Private/UI/SRCPanelFieldGroup.cpp
@@ -46,7 +46,7 @@ void SRCPanelGroup::Construct(const FArguments& InArgs, URemoteControlPreset* In
 	OnFieldDropEvent = InArgs._OnFieldDropEvent;
 	OnGetGroupId = InArgs._OnGetGroupId;
 	OnDeleteGroup = InArgs._OnDeleteGroup;
-	bEditMode = InArgs._EditMode;
+	bLiveMode = InArgs._LiveMode;
 
 	TSharedRef<SWidget> LeftColumn = 
 		SNew(SHorizontalBox)
@@ -69,7 +69,7 @@ void SRCPanelGroup::Construct(const FArguments& InArgs, URemoteControlPreset* In
 		[
 			SNew(SBox)
 			.Padding(2.f, 0.f)
-			.Visibility(this, &SRCPanelGroup::GetVisibilityAccordingToEditMode, EVisibility::Collapsed)
+			.Visibility(this, &SRCPanelGroup::GetVisibilityAccordingToLiveMode, EVisibility::Collapsed)
 			[
 				SNew(SRCPanelDragHandle<FFieldGroupDragDropOp>, Id)
 				.Widget(SharedThis(this))
@@ -88,7 +88,7 @@ void SRCPanelGroup::Construct(const FArguments& InArgs, URemoteControlPreset* In
 			.Text(FText::FromName(Name))
 			.OnTextCommitted(this, &SRCPanelGroup::OnLabelCommitted)
 			.OnVerifyTextChanged(this, &SRCPanelGroup::OnVerifyItemLabelChanged)
-			.IsReadOnly_Lambda([this]() { return !bEditMode.Get(); })
+			.IsReadOnly_Lambda([this]() { return bLiveMode.Get(); })
 		];
 
 	ChildSlot
@@ -243,9 +243,9 @@ const FSlateBrush* SRCPanelGroup::GetBorderImage() const
 	return FRemoteControlPanelStyle::Get()->GetBrush("RemoteControlPanel.TransparentBorder");
 }
 
-EVisibility SRCPanelGroup::GetVisibilityAccordingToEditMode(EVisibility NonEditModeVisibility) const
+EVisibility SRCPanelGroup::GetVisibilityAccordingToLiveMode(EVisibility NonEditModeVisibility) const
 {
-	return bEditMode.Get() ? EVisibility::Visible : NonEditModeVisibility;
+	return !bLiveMode.Get() ? EVisibility::Visible : NonEditModeVisibility;
 }
 
 bool SRCPanelGroup::OnVerifyItemLabelChanged(const FText& InLabel, FText& OutErrorMessage)
diff --git a/Engine/Plugins/VirtualProduction/RemoteControl/Source/RemoteControlUI/Private/UI/SRCPanelFieldGroup.h b/Engine/Plugins/VirtualProduction/RemoteControl/Source/RemoteControlUI/Private/UI/SRCPanelFieldGroup.h
--- a/Engine/Plugins/VirtualProduction/RemoteControl/Source/RemoteControlUI/Private/UI/SRCPanelFieldGroup.h
+++ b/Engine/Plugins/VirtualProduction/RemoteControl/Source/RemoteControlUI/Private/UI/SRCPanelFieldGroup.h
@@ -32,7 +32,7 @@ public:
 	DECLARE_DELEGATE_OneParam(FOnDeleteGroup, const FGuid& /*GroupId*/);
 
 	SLATE_BEGIN_ARGS(SRCPanelGroup)
-		: _EditMode(true)
+		: _LiveMode(false)
 	{}
 		SLATE_ARGUMENT(FGuid, Id)
 		SLATE_ARGUMENT(FName, Name)
@@ -40,7 +40,7 @@ public:
 		SLATE_EVENT(FOnFieldDropEvent, OnFieldDropEvent)
 		SLATE_EVENT(FOnGetGroupId, OnGetGroupId)
 		SLATE_EVENT(FOnDeleteGroup, OnDeleteGroup)
-		SLATE_ATTRIBUTE(bool, EditMode)
+		SLATE_ATTRIBUTE(bool, LiveMode)
 		
 	SLATE_END_ARGS()
 
@@ -79,7 +79,7 @@ private:
 	/** Get the border image according to the current selection. */
 	const FSlateBrush* GetBorderImage() const;
 	/** Get the visibility according to the panel's current mode. */
-	EVisibility GetVisibilityAccordingToEditMode(EVisibility DefaultHiddenVisibility) const;
+	EVisibility GetVisibilityAccordingToLiveMode(EVisibility DefaultHiddenVisibility) const;
 
 	bool OnVerifyItemLabelChanged(const FText& InLabel, FText& OutErrorMessage);
 
@@ -96,8 +96,8 @@ private:
 	FOnDeleteGroup OnDeleteGroup;
 	/** Holds the text box for the group name. */
 	TSharedPtr<SInlineEditableTextBlock> NameTextBox;
-	/** Whether the panel is currently in edit mode. */
-	TAttribute<bool> bEditMode;
+	/** Whether the panel is currently in live mode. */
+	TAttribute<bool> bLiveMode;
 	/** Whether the group needs to be renamed. (As requested by a click on the rename button) */
 	bool bNeedsRename = false;
 	/** Weak ptr to the preset that contains the field group. */
diff --git a/Engine/Plugins/VirtualProduction/RemoteControl/Source/RemoteControlUI/Private/UI/SRemoteControlPanel.cpp b/Engine/Plugins/VirtualProduction/RemoteControl/Source/RemoteControlUI/Private/UI/SRemoteControlPanel.cpp
--- a/Engine/Plugins/VirtualProduction/RemoteControl/Source/RemoteControlUI/Private/UI/SRemoteControlPanel.cpp
+++ b/Engine/Plugins/VirtualProduction/RemoteControl/Source/RemoteControlUI/Private/UI/SRemoteControlPanel.cpp
@@ -143,6 +143,11 @@ public:
 	{
 		if (const TSharedPtr<SRemoteControlPanel> PinnedOwner = Owner.Pin())
 		{
+			if (PinnedOwner->IsInLiveMode())
+			{
+				return false;
+			}
+
 			if (!SlateApp.HasFocusedDescendants(PinnedOwner.ToSharedRef()))
 			{
 				return false;
@@ -329,7 +334,7 @@ private:
 
 void SRemoteControlPanel::Construct(const FArguments& InArgs, URemoteControlPreset* InPreset, TSharedPtr<IToolkitHost> InToolkitHost)
 {
-	OnEditModeChange = InArgs._OnEditModeChange;
+	OnLiveModeChange = InArgs._OnLiveModeChange;
 	Preset = TStrongObjectPtr<URemoteControlPreset>(InPreset);
 	WidgetRegistry = MakeShared<FRCPanelWidgetRegistry>();
 	ToolkitHost = InToolkitHost;
@@ -347,7 +352,7 @@ void SRemoteControlPanel::Construct(const FArguments& InArgs, URemoteControlPres
 	AddToolbarWidget(SNew(SCheckBox)
 		.Style(&RCPanelStyle->ToggleButtonStyle)
 		.ToolTipText(LOCTEXT("ShowLogTooltip", "Show/Hide remote control log."))
-		.IsChecked_Lambda([]() { return FRemoteControlLogger::Get().IsEnabled() ? ECheckBoxState::Checked : ECheckBoxState::Unchecked; })
+		.IsChecked_Lambda([this]() { return (FRemoteControlLogger::Get().IsEnabled() && !bIsInLiveMode) ? ECheckBoxState::Checked : ECheckBoxState::Unchecked; })
 		.OnCheckStateChanged(this, &SRemoteControlPanel::OnLogCheckboxToggle)
 		.Padding(4.f)
 		[
@@ -355,21 +360,6 @@ void SRemoteControlPanel::Construct(const FArguments& InArgs, URemoteControlPres
 			.Text(LOCTEXT("ShowLogLabel", "Log"))
 			.Justification(ETextJustify::Center)
 			.TextStyle(&RCPanelStyle->PanelTextStyle)
-		]);
-
-	// Edit Mode
-	AddToolbarWidget(SNew(SCheckBox)
-		.Style(&RCPanelStyle->ToggleButtonStyle)
-		.ToolTipText(LOCTEXT("EditModeTooltip", "Toggle Editing (Ctrl + E)"))
-		.IsChecked_Lambda([this]() { return this->bIsInEditMode ? ECheckBoxState::Checked : ECheckBoxState::Unchecked; })
-		.Padding(4.f)
-		.OnCheckStateChanged(this, &SRemoteControlPanel::OnEditModeCheckboxToggle)
-		.Content()
-		[
-			SNew(STextBlock)
-			.Text(LOCTEXT("EditLabel", "Edit"))
-			.Justification(ETextJustify::Center)
-			.TextStyle(&RCPanelStyle->PanelTextStyle)
 		]
 	);
 
@@ -410,7 +400,7 @@ void SRemoteControlPanel::Construct(const FArguments& InArgs, URemoteControlPres
 				CachedExposedPropertyArgs.Reset();
 			}
 		)
-		.EditMode_Lambda([this]() { return bIsInEditMode; })
+		.LiveMode_Lambda([this]() {return bIsInLiveMode; })
 		.ProtocolsMode_Lambda([this]() { return IsInProtocolsMode() && ActivePanel == ERCPanels::RCP_Protocols; });
 
 	EntityList->OnSelectionChange().AddSP(this, &SRemoteControlPanel::UpdateEntityDetailsView);
@@ -549,10 +539,11 @@ void SRemoteControlPanel::Construct(const FArguments& InArgs, URemoteControlPres
 		.EnableFooter(true);
 
 	ControllerPanel = SNew(SRCControllerPanel, SharedThis(this))
-		.Visibility_Lambda([this] { return bIsLogicPanelEnabled && (ActivePanel == ERCPanels::RCP_Properties || ActivePanel == ERCPanels::RCP_None) ? EVisibility::Visible : EVisibility::Collapsed; });
+		.LiveMode_Lambda([this]() { return bIsInLiveMode; })
+		.Visibility_Lambda([this] { return (bIsLogicPanelEnabled || bIsInLiveMode) && (ActivePanel == ERCPanels::RCP_Properties || ActivePanel == ERCPanels::RCP_Live || ActivePanel == ERCPanels::RCP_None) ? EVisibility::Visible : EVisibility::Collapsed; });
 
 	BehaviourPanel = SNew(SRCBehaviourPanel, SharedThis(this))
-		.Visibility_Lambda([this] { return bIsLogicPanelEnabled && (ActivePanel == ERCPanels::RCP_Properties || ActivePanel == ERCPanels::RCP_None) ? EVisibility::Visible : EVisibility::Collapsed; });
+		.Visibility_Lambda([this] { return !bIsInLiveMode && bIsLogicPanelEnabled && (ActivePanel == ERCPanels::RCP_Properties || ActivePanel == ERCPanels::RCP_None) ? EVisibility::Visible : EVisibility::Collapsed; });
 
 	ControllersAndBehavioursPanel->AddPanel(ControllerPanel.ToSharedRef(), 0.5f);
 		
@@ -573,7 +564,7 @@ void SRemoteControlPanel::Construct(const FArguments& InArgs, URemoteControlPres
 
 	// Actions Panel.
 	ActionPanel = SNew(SRCActionPanel, SharedThis(this))
-		.Visibility_Lambda([this] { return bIsLogicPanelEnabled && (ActivePanel == ERCPanels::RCP_Properties || ActivePanel == ERCPanels::RCP_None) ? EVisibility::Visible : EVisibility::Collapsed; });
+		.Visibility_Lambda([this] { return !bIsInLiveMode && bIsLogicPanelEnabled && (ActivePanel == ERCPanels::RCP_Properties || ActivePanel == ERCPanels::RCP_None) ? EVisibility::Visible : EVisibility::Collapsed; });
 
 	LogicPanel->AddPanel(ControllersAndBehavioursPanel, 0.67f);
 	LogicPanel->AddPanel(ActionPanel.ToSharedRef(), 0.33f);
@@ -707,6 +698,27 @@ void SRemoteControlPanel::Construct(const FArguments& InArgs, URemoteControlPres
 			[
 				OutputLogDockPanel
 			]
+
+			// 6. Live (Index : 5).
+			+ SWidgetSwitcher::Slot()
+			[
+				SNew(SSplitter)
+				.Orientation(EOrientation::Orient_Horizontal)						
+				
+				// Exposed entities List
+				+ SSplitter::Slot()
+				.Value(0.5f)
+				[
+					EntityList.ToSharedRef()
+				]
+			
+				// Logic Panel
+				+ SSplitter::Slot()
+				.Value(0.5f)
+				[
+					LogicPanel
+				]
+			]
 		];
 
 	ChildSlot
@@ -1099,7 +1111,7 @@ TSharedRef<SWidget> SRemoteControlPanel::CreateExposeButton()
 	
 	return SAssignNew(ExposeComboButton, SComboButton)
 		.AddMetaData<FTagMetaData>(FTagMetaData(TEXT("Expose Functions")))
-		.IsEnabled_Lambda([this]() { return this->bIsInEditMode; })
+		.IsEnabled_Lambda([this]() { return this->bIsInLiveMode; })
 		.HAlign(HAlign_Center)
 		.VAlign(VAlign_Center)
 		.ButtonStyle(&RCPanelStyle->FlatButtonStyle)
@@ -1399,6 +1411,12 @@ void SRemoteControlPanel::RegisterPanels()
 	if (PanelDrawer.IsValid())
 	{
 		PanelDrawer->OnRCPanelToggled().BindSP(this, &SRemoteControlPanel::OnRCPanelToggled);
+		
+		PanelDrawer->CanToggleRCPanel().BindLambda([this]()
+			{
+				return !bIsInLiveMode;
+			}
+		);
 
 		{// Properties Panel
 			TSharedRef<FRCPanelDrawerArgs> PropertiesPanel = MakeShared<FRCPanelDrawerArgs>(ERCPanels::RCP_Properties);
@@ -1458,6 +1476,21 @@ void SRemoteControlPanel::RegisterPanels()
 
 			RegisteredDrawers.Add(OutputLogPanel->GetPanelID(), OutputLogPanel);
 		}
+
+		{// Live Panel
+			TSharedRef<FRCPanelDrawerArgs> LivePanel = MakeShared<FRCPanelDrawerArgs>(ERCPanels::RCP_Live);
+
+			LivePanel->bDrawnByDefault = true;
+			LivePanel->bRotateIconBy90 = false;
+			LivePanel->DrawerVisibility = bIsInLiveMode ? EVisibility::Visible : EVisibility::Collapsed;
+			LivePanel->Label = LOCTEXT("LivePanelPanelLabel", "Live");
+			LivePanel->ToolTip = LOCTEXT("LivePanelTooltip", "Open live panel.");
+			LivePanel->Icon = FSlateIcon(FAppStyle::GetAppStyleSetName(), "LevelEditor.Tabs.StatsViewer");
+
+			PanelDrawer->RegisterPanel(LivePanel);
+
+			RegisteredDrawers.Add(LivePanel->GetPanelID(), LivePanel);
+		}
 	}
 }
 
@@ -1466,6 +1499,7 @@ void SRemoteControlPanel::UnregisterPanels()
 	if (PanelDrawer.IsValid())
 	{
 		PanelDrawer->OnRCPanelToggled().Unbind();
+		PanelDrawer->CanToggleRCPanel().Unbind();
 
 		for (TMap<ERCPanels, TSharedRef<FRCPanelDrawerArgs>>::TIterator RegisteredDrawer = RegisteredDrawers.CreateIterator(); RegisteredDrawer; ++RegisteredDrawer)
 		{
@@ -1541,12 +1575,6 @@ void SRemoteControlPanel::Unexpose(const FRCExposesPropertyArgs& InPropertyArgs)
 	}
 }
 
-void SRemoteControlPanel::OnEditModeCheckboxToggle(ECheckBoxState State)
-{
-	bIsInEditMode = (State == ECheckBoxState::Checked) ? true : false;
-	OnEditModeChange.ExecuteIfBound(SharedThis(this), bIsInEditMode);
-}
-
 void SRemoteControlPanel::OnLogCheckboxToggle(ECheckBoxState State)
 {
 	const bool bIsLogEnabled = (State == ECheckBoxState::Checked) ? true : false;
@@ -1619,28 +1647,6 @@ void SRemoteControlPanel::ExposeActor(AActor* Actor)
 	}
 }
 
-void SRemoteControlPanel::ToggleDetailsView()
-{
-	const FTabId TabId = FTabId(FRemoteControlUIModule::EntityDetailsTabName);
-	
-	if (TSharedPtr<IToolkitHost> PinnedToolkit = ToolkitHost.Pin())
-	{
-		if (TSharedPtr<SDockTab> ExistingTab = PinnedToolkit->GetTabManager()->FindExistingLiveTab(TabId))
-		{
-			ExistingTab->RequestCloseTab();
-		}
-		else
-		{
-			// Request the Tab Manager to invoke the tab. This will spawn the tab if needed, otherwise pull it to focus. This assumes
-			// that the Toolkit Host's Tab Manager has already registered a tab with a NullWidget for content.
-			if (TSharedPtr<SDockTab> EntityDetailsTab = PinnedToolkit->GetTabManager()->TryInvokeTab(TabId))
-			{
-				EntityDetailsTab->SetContent(CreateEntityDetailsView());
-			}
-		}
-	}
-}
-
 TSharedRef<SWidget> SRemoteControlPanel::CreateEntityDetailsView()
 {
 	WrappedEntityDetailsView = SNew(SBorder)
@@ -1920,7 +1926,35 @@ void SRemoteControlPanel::GenerateToolbar()
 			.DefaultMode("Setup")
 			.OnModeSwitched_Lambda([this](const SRCModeSwitcher::FRCMode& NewMode)
 				{
-					UE_LOG(LogRemoteControl, Warning, TEXT("Active Mode : %s"), *NewMode.ModeId.ToString());
+					if (NewMode.ModeId == TEXT("Operation"))
+					{
+						bIsInLiveMode = true;
+					}
+					else if (NewMode.ModeId == TEXT("Setup"))
+					{
+						bIsInLiveMode = false;
+					}
+
+					bIsLogicPanelEnabled = bIsInLiveMode;
+
+					if (PanelDrawer.IsValid())
+					{
+						TSharedRef<FRCPanelDrawerArgs> LivePanel = RegisteredDrawers.FindChecked(ERCPanels::RCP_Live);
+
+						LivePanel->DrawerVisibility = bIsInLiveMode ? EVisibility::Visible : EVisibility::Collapsed;
+
+						// When we are not in Live Mode collapse the drawer.
+						PanelDrawer->TogglePanel(LivePanel, !bIsInLiveMode);
+
+						if (!bIsInLiveMode)
+						{
+							TSharedRef<FRCPanelDrawerArgs> PropertiesPanel = RegisteredDrawers.FindChecked(ERCPanels::RCP_Properties);
+
+							PanelDrawer->TogglePanel(PropertiesPanel);
+						}
+					}
+
+					OnLiveModeChange.ExecuteIfBound(SharedThis(this), bIsInLiveMode);
 				}
 			)
 		
@@ -2185,12 +2219,12 @@ void SRemoteControlPanel::ToggleProtocolMappings_Execute()
 
 bool SRemoteControlPanel::CanToggleProtocolsMode() const
 {
-	return bIsInEditMode;
+	return !bIsInLiveMode;
 }
 
 bool SRemoteControlPanel::IsInProtocolsMode() const
 {
-	return bIsInEditMode && bIsInProtocolsMode;
+	return bIsInProtocolsMode;
 }
 
 void SRemoteControlPanel::ToggleLogicEditor_Execute()
@@ -2207,12 +2241,12 @@ void SRemoteControlPanel::ToggleLogicEditor_Execute()
 
 bool SRemoteControlPanel::CanToggleLogicPanel() const
 {
-	return bIsInEditMode;
+	return !bIsInLiveMode;
 }
 
 bool SRemoteControlPanel::IsLogicPanelEnabled() const
 {
-	return bIsInEditMode && bIsLogicPanelEnabled;
+	return !bIsInLiveMode && bIsLogicPanelEnabled;
 }
 
 void SRemoteControlPanel::OnRCPanelToggled(ERCPanels InPanelID)
@@ -2221,7 +2255,7 @@ void SRemoteControlPanel::OnRCPanelToggled(ERCPanels InPanelID)
 	{
 		ActivePanel = InPanelID;
 		
-		if (EntityList.IsValid())
+		if (EntityList.IsValid() && !bIsInLiveMode)
 		{
 			const bool bToggleProtcolMode = IsInProtocolsMode() && ActivePanel == ERCPanels::RCP_Protocols;
 
@@ -2288,12 +2322,17 @@ void SRemoteControlPanel::DeleteEntity_Execute()
 
 bool SRemoteControlPanel::CanDeleteEntity() const
 {
+	if (bIsInLiveMode)
+	{
+		return false;
+	}
+
 	if (TSharedPtr<SRCLogicPanelBase> ActiveLogicPanel = GetActiveLogicPanel())
 	{
 		return ActiveLogicPanel->GetSelectedLogicItem() != nullptr; // User has focus on a logic panel
 	}
 
-	if (SelectedEntity.IsValid() && Preset.IsValid() && bIsInEditMode)
+	if (SelectedEntity.IsValid() && Preset.IsValid())
 	{
 		// Do not allow default group to be deleted.
 		return !Preset->Layout.IsDefaultGroup(SelectedEntity->GetRCId());
@@ -2320,12 +2359,17 @@ void SRemoteControlPanel::RenameEntity_Execute() const
 
 bool SRemoteControlPanel::CanRenameEntity() const
 {
+	if (bIsInLiveMode)
+	{
+		return false;
+	}
+
 	if (ControllerPanel->IsListFocused())
 	{
 		return true;
 	}
 
-	if (SelectedEntity.IsValid() && Preset.IsValid() && bIsInEditMode)
+	if (SelectedEntity.IsValid() && Preset.IsValid())
 	{
 		// Do not allow default group to be renamed.
 		return !Preset->Layout.IsDefaultGroup(SelectedEntity->GetRCId());
@@ -2350,6 +2394,11 @@ void SRemoteControlPanel::CopyItem_Execute()
 
 bool SRemoteControlPanel::CanCopyItem() const
 {
+	if (bIsInLiveMode)
+	{
+		return false;
+	}
+
 	if (TSharedPtr<SRCLogicPanelBase> ActiveLogicPanel = GetActiveLogicPanel())
 	{
 		return ActiveLogicPanel->GetSelectedLogicItem().IsValid();
@@ -2368,6 +2417,11 @@ void SRemoteControlPanel::PasteItem_Execute()
 
 bool SRemoteControlPanel::CanPasteItem() const
 {
+	if (bIsInLiveMode)
+	{
+		return false;
+	}
+
 	if(LogicClipboardItem)
 	{
 		if (TSharedPtr<SRCLogicPanelBase> ActiveLogicPanel = GetActiveLogicPanel())
@@ -2395,6 +2449,11 @@ void SRemoteControlPanel::DuplicateItem_Execute()
 
 bool SRemoteControlPanel::CanDuplicateItem() const
 {
+	if (bIsInLiveMode)
+	{
+		return false;
+	}
+
 	if (TSharedPtr<SRCLogicPanelBase> ActiveLogicPanel = GetActiveLogicPanel())
 	{
 		return ActiveLogicPanel->GetSelectedLogicItem().IsValid();
diff --git a/Engine/Plugins/VirtualProduction/RemoteControl/Source/RemoteControlUI/Private/UI/SRemoteControlPanel.h b/Engine/Plugins/VirtualProduction/RemoteControl/Source/RemoteControlUI/Private/UI/SRemoteControlPanel.h
--- a/Engine/Plugins/VirtualProduction/RemoteControl/Source/RemoteControlUI/Private/UI/SRemoteControlPanel.h
+++ b/Engine/Plugins/VirtualProduction/RemoteControl/Source/RemoteControlUI/Private/UI/SRemoteControlPanel.h
@@ -49,7 +49,7 @@ class URCAction;
 class FRCControllerModel;
 class FRCBehaviourModel;
 
-DECLARE_DELEGATE_TwoParams(FOnEditModeChange, TSharedPtr<SRemoteControlPanel> /* Panel */, bool /* bEditModeChange */);
+DECLARE_DELEGATE_TwoParams(FOnLiveModeChange, TSharedPtr<SRemoteControlPanel> /* Panel */, bool /* bEditModeChange */);
 
 /**
  * UI representation of a remote control preset.
@@ -58,7 +58,7 @@ DECLARE_DELEGATE_TwoParams(FOnEditModeChange, TSharedPtr<SRemoteControlPanel> /*
 class SRemoteControlPanel : public SCompoundWidget, public FSelfRegisteringEditorUndoClient
 {
 	SLATE_BEGIN_ARGS(SRemoteControlPanel) {}
-		SLATE_EVENT(FOnEditModeChange, OnEditModeChange)
+		SLATE_EVENT(FOnLiveModeChange, OnLiveModeChange)
 		SLATE_ARGUMENT(bool, AllowGrouping)
 	SLATE_END_ARGS()
 
@@ -113,9 +113,9 @@ public:
 	void ToggleProperty(const FRCExposesPropertyArgs& InArgs);
 
 	/**
-	 * @return Whether or not the panel is in edit mode.
+	 * @return Whether or not the panel is in live mode.
 	 */
-	bool IsInEditMode() const { return bIsInEditMode; }
+	bool IsInLiveMode() const { return bIsInLiveMode; }
 
 	/**
 	 * Get the selected group.
@@ -126,9 +126,9 @@ public:
 	 * Set the edit mode of the panel.
 	 * @param bEditMode The desired mode.
 	 */
-	void SetEditMode(bool bEditMode)
+	void SetLiveMode(bool bLiveMode)
 	{
-		bIsInEditMode = bEditMode;
+		bIsInLiveMode = bLiveMode;
 	}
 
 	/**
@@ -199,9 +199,6 @@ private:
 	/** Handle exposing an actor. */
 	void ExposeActor(AActor* Actor);
 
-	/** Opens or closes the entity details view tab. */
-	void ToggleDetailsView();
-
 	/** Handles disabling CPU throttling. */
 	FReply OnClickDisableUseLessCPU() const;
 
@@ -346,16 +343,16 @@ private:
 	static const FName AuxiliaryRemoteControlPanelToolBarName;
 	/** Holds the preset asset. */
 	TStrongObjectPtr<URemoteControlPreset> Preset;
-	/** Whether the panel is in edit mode. */
-	bool bIsInEditMode = true;
 	/** Whether the panel is in protocols mode. */
 	bool bIsInProtocolsMode = false;
+	/** Whether the panel is in live (or) operation mode. */
+	bool bIsInLiveMode = false;
 	/** Whether the logic panel is enabled or not. */
 	bool bIsLogicPanelEnabled = false;
 	/** Whether objects need to be re-resolved because PIE Started or ended. */
 	bool bTriggerRefreshForPIE = false;
-	/** Delegate called when the edit mode changes. */
-	FOnEditModeChange OnEditModeChange;
+	/** Delegate called when the live mode changes. */
+	FOnLiveModeChange OnLiveModeChange;
 	/** Holds the blueprint library picker */
 	TSharedPtr<SRCPanelFunctionPicker> BlueprintPicker;
 	/** Holds the actor function picker */
diff --git a/Engine/Plugins/VirtualProduction/RemoteControl/Source/RemoteControlUI/Public/IRemoteControlUIModule.h b/Engine/Plugins/VirtualProduction/RemoteControl/Source/RemoteControlUI/Public/IRemoteControlUIModule.h
--- a/Engine/Plugins/VirtualProduction/RemoteControl/Source/RemoteControlUI/Public/IRemoteControlUIModule.h
+++ b/Engine/Plugins/VirtualProduction/RemoteControl/Source/RemoteControlUI/Public/IRemoteControlUIModule.h
@@ -32,7 +32,7 @@ struct FGenerateWidgetArgs
 	URemoteControlPreset* Preset = nullptr;
 	FRCColumnSizeData ColumnSizeData;
 	TSharedPtr<FRemoteControlEntity> Entity;
-	TAttribute<bool> bIsInEditMode;
+	TAttribute<bool> bIsInLiveMode;
 	TWeakPtr<FRCPanelWidgetRegistry> WidgetRegistry;
 	FText HighlightText;
 };
