Material Actions

diff --git a/Engine/Plugins/VirtualProduction/RemoteControl/Source/RemoteControlLogic/Private/Action/RCActionContainer.cpp b/Engine/Plugins/VirtualProduction/RemoteControl/Source/RemoteControlLogic/Private/Action/RCActionContainer.cpp
--- a/Engine/Plugins/VirtualProduction/RemoteControl/Source/RemoteControlLogic/Private/Action/RCActionContainer.cpp
+++ b/Engine/Plugins/VirtualProduction/RemoteControl/Source/RemoteControlLogic/Private/Action/RCActionContainer.cpp
@@ -83,11 +83,47 @@ URCPropertyAction* URCActionContainer::AddPropertyAction(const TSharedRef<const
 
 		return nullptr;
 	}
-	
-	if (FRCObjectReference ObjectRef; IRemoteControlModule::Get().ResolveObjectProperty(ERCAccess::READ_ACCESS, InRemoteControlProperty->GetBoundObjects()[0], InRemoteControlProperty->FieldPathInfo.ToString(), ObjectRef))
+
+	bool bFoundMatchingContainer = false;
+
+	if (FProperty* Property = InRemoteControlProperty->GetProperty())
+	{
+		for (URCAction* Action : Actions)
+		{
+			if (URCPropertyAction* PropertyAction = Cast<URCPropertyAction>(Action))
+			{
+				if (PropertyAction->PropertySelfContainer)
+				{
+					if(FProperty* InputFieldProperty = PropertyAction->PropertySelfContainer->GetProperty())
+					{
+						if(InputFieldProperty->GetFName() == Property->GetFName())
+						{
+							if (const TSharedPtr<const FRemoteControlProperty> RemoteControlPropertyForAction = PropertyAction->GetRemoteControlProperty())
+							{
+								if (InRemoteControlProperty->GetBoundObject() == RemoteControlPropertyForAction->GetBoundObject())
+								{
+									// We already have an Action associated with this property, so this must be a container type like Array.
+									// Both the UI widget and Execute will operate on a unique index associated with this Action for this container property
+									NewPropertyAction->PropertySelfContainer = PropertyAction->PropertySelfContainer;
+
+									bFoundMatchingContainer = true;
+								}
+							}
+						}
+					}
+				}
+			}
+		}
+	}
+
+	if(!bFoundMatchingContainer)
 	{
-		const FName& PropertyName = InRemoteControlProperty->GetProperty()->GetFName();
-		NewPropertyAction->PropertySelfContainer->DuplicatePropertyWithCopy(PropertyName, InRemoteControlProperty->GetProperty(), (uint8*)ObjectRef.ContainerAdress);
+		// Create an input field for the Action by duplicating the Remote Control Property associated with it
+		if (FRCObjectReference ObjectRef; IRemoteControlModule::Get().ResolveObjectProperty(ERCAccess::READ_ACCESS, InRemoteControlProperty->GetBoundObjects()[0], InRemoteControlProperty->FieldPathInfo.ToString(), ObjectRef))
+		{
+			const FName& PropertyName = InRemoteControlProperty->GetProperty()->GetFName();
+			NewPropertyAction->PropertySelfContainer->DuplicatePropertyWithCopy(PropertyName, InRemoteControlProperty->GetProperty(), (uint8*)ObjectRef.ContainerAdress);
+		}
 	}
 
 	AddAction(NewPropertyAction);
diff --git a/Engine/Plugins/VirtualProduction/RemoteControl/Source/RemoteControlUI/Private/UI/Action/RCActionModel.cpp b/Engine/Plugins/VirtualProduction/RemoteControl/Source/RemoteControlUI/Private/UI/Action/RCActionModel.cpp
--- a/Engine/Plugins/VirtualProduction/RemoteControl/Source/RemoteControlUI/Private/UI/Action/RCActionModel.cpp
+++ b/Engine/Plugins/VirtualProduction/RemoteControl/Source/RemoteControlUI/Private/UI/Action/RCActionModel.cpp
@@ -190,18 +190,49 @@ FRCPropertyActionType::FRCPropertyActionType(URCPropertyAction* InPropertyAction
 
 	if (InPropertyAction)
 	{
+		// Generate UI widget for Action input
 		if (const TSharedPtr<FStructOnScope> StructOnScope = InPropertyAction->PropertySelfContainer->CreateStructOnScope())
 		{
 			PropertyRowGenerator->SetStructure(StructOnScope);
-
+			
 			for (const TSharedRef<IDetailTreeNode>& CategoryNode : PropertyRowGenerator->GetRootTreeNodes())
 			{
 				TArray<TSharedRef<IDetailTreeNode>> Children;
 				CategoryNode->GetChildren(Children);
+
 				for (const TSharedRef<IDetailTreeNode>& Child : Children)
 				{
-					DetailTreeNodeWeakPtr = Child;
-					break;
+					// Special handling for any Remote Control Property that is actually a single element of an Array
+					// For example: "Override Materials" of Static Mesh Actor when exposed as a property in the Remote Control preset
+					//
+					if (TSharedPtr<FRemoteControlProperty> RemoteControlProperty = InPropertyAction->GetRemoteControlProperty())
+					{
+						if (ensure(RemoteControlProperty->FieldPathInfo.GetSegmentCount() > 0))
+						{
+							const FRCFieldPathSegment& LastSegment = RemoteControlProperty->FieldPathInfo.Segments.Last();
+							const int32 ArrayIndex = LastSegment.ArrayIndex;
+
+							// Is this a Container element?
+							if (ArrayIndex != INDEX_NONE)
+							{
+								TArray<TSharedRef<IDetailTreeNode>> InnerChildren;
+								Child->GetChildren(InnerChildren);
+
+								// If yes, extract the Detail node corresponding to that element (rather than the parent container)
+								if (ensure(InnerChildren.IsValidIndex(ArrayIndex)))
+								{
+									DetailTreeNodeWeakPtr = InnerChildren[ArrayIndex];
+								}
+							}
+						}
+					}
+
+					// For regular properties (non-container)
+					if (!DetailTreeNodeWeakPtr.IsValid())
+					{
+						DetailTreeNodeWeakPtr = Child;
+						break;
+					}
 				}
 			}
 		}
