Fix NDisplay exposed properties on viewports not being able to be rebound
Fix some ndisplay properties changing to <multiple values> after being changed in the preset

diff --git a/Engine/Plugins/VirtualProduction/RemoteControl/Source/RemoteControl/Private/RemoteControlBinding.cpp b/Engine/Plugins/VirtualProduction/RemoteControl/Source/RemoteControl/Private/RemoteControlBinding.cpp
--- a/Engine/Plugins/VirtualProduction/RemoteControl/Source/RemoteControl/Private/RemoteControlBinding.cpp
+++ b/Engine/Plugins/VirtualProduction/RemoteControl/Source/RemoteControl/Private/RemoteControlBinding.cpp
@@ -136,8 +136,10 @@ void URemoteControlLevelDependantBinding::SetBoundObject(const TSoftObjectPtr<UO
 		SubLevelSelectionMapByPath.FindOrAdd(OuterLevel->GetWorld()) = OuterLevel;
 		
 		Name = EditorObject->GetName();
+
+		const bool bShouldSetSubObjectContext = !BindingContext.HasValidSubObjectPath() && !EditorObject->IsA<AActor>();
 		
-		if (BindingContext.OwnerActorName.IsNone() || !GetDefault<URemoteControlSettings>()->bUseRebindingContext) 
+		if (BindingContext.OwnerActorName.IsNone() || bShouldSetSubObjectContext || !GetDefault<URemoteControlSettings>()->bUseRebindingContext)
 		{
 			InitializeBindingContext(InObject.Get());
 		}
@@ -215,7 +217,7 @@ UObject* URemoteControlLevelDependantBinding::Resolve() const
 
 		LevelWithLastSuccessfulResolve = Object->GetTypedOuter<ULevel>();
 
-		if (BindingContext.OwnerActorName.IsNone())
+		if (BindingContext.OwnerActorName.IsNone() || (!Object->IsA<AActor>() && !BindingContext.HasValidSubObjectPath()))
 		{
 			InitializeBindingContext(Object);
 		}
@@ -346,25 +348,23 @@ void URemoteControlLevelDependantBinding::InitializeBindingContext(UObject* InOb
 			BindingContext.OwnerActorClass = InObject->GetClass();
 			BindingContext.OwnerActorName = InObject->GetFName();
 		}
-		else if (InObject->IsA<UActorComponent>())
-		{
-			AActor* Owner = InObject->GetTypedOuter<AActor>();
-			check(Owner);
-			BindingContext.ComponentName = InObject->GetFName();
-			BindingContext.OwnerActorClass = Owner->GetClass();
-			BindingContext.OwnerActorName = Owner->GetFName();
-		}
 		else
 		{
 			AActor* Owner = InObject->GetTypedOuter<AActor>();
 			check(Owner);
 			BindingContext.ComponentName = InObject->GetFName();
+			BindingContext.SubObjectPath = InObject->GetPathName(Owner);
 			BindingContext.OwnerActorClass = Owner->GetClass();
 			BindingContext.OwnerActorName = Owner->GetFName();
 		}
 	}
 }
 
+TSoftObjectPtr<UObject> URemoteControlLevelDependantBinding::GetLastBoundObject() const
+{
+	return BoundObjectMapByPath.FindRef(LevelWithLastSuccessfulResolve.ToSoftObjectPath());
+}
+
 UClass* URemoteControlLevelDependantBinding::GetSupportedOwnerClass() const
 {
 	return BindingContext.OwnerActorClass.LoadSynchronous();
diff --git a/Engine/Plugins/VirtualProduction/RemoteControl/Source/RemoteControl/Private/RemoteControlField.cpp b/Engine/Plugins/VirtualProduction/RemoteControl/Source/RemoteControl/Private/RemoteControlField.cpp
--- a/Engine/Plugins/VirtualProduction/RemoteControl/Source/RemoteControl/Private/RemoteControlField.cpp
+++ b/Engine/Plugins/VirtualProduction/RemoteControl/Source/RemoteControl/Private/RemoteControlField.cpp
@@ -16,6 +16,34 @@
 #include "UObject/StructOnScope.h"
 #include "UObject/UnrealType.h"
 
+
+namespace RemoteControlResolvers
+{
+	FName DisplayClusterActorName = "DisplayClusterRootActor";
+	FName ConfigDataPropertyName = "CurrentConfigData";
+
+	UObject* NDisplayConfigDataResolver(const UObject* NDisplayActor)
+	{
+		if (NDisplayActor)
+		{
+			UClass* SuperClass = NDisplayActor->GetClass()->GetSuperClass();
+			if (SuperClass && SuperClass->GetFName() == DisplayClusterActorName)
+			{
+				if (FObjectProperty* ConfigData = CastField<FObjectProperty>(SuperClass->FindPropertyByName(ConfigDataPropertyName)))
+				{
+					return ConfigData->GetObjectPropertyValue_InContainer(NDisplayActor);
+				}
+			}
+		}
+
+		return nullptr;
+	}
+
+	static TMap<FName, TFunction<UObject* (const UObject*)>> CustomResolvers = {
+		{ DisplayClusterActorName, NDisplayConfigDataResolver }
+	};
+}
+
 FRemoteControlField::FRemoteControlField(URemoteControlPreset* InPreset, EExposedFieldType InType, FName InLabel, FRCFieldPathInfo InFieldPathInfo, const TArray<URemoteControlBinding*> InBindings)
 	: FRemoteControlEntity(InPreset, InLabel, InBindings)
 	, FieldType(InType)
@@ -63,6 +91,20 @@ void FRemoteControlField::BindObject(UObject* InObjectToBind)
 			}
 			else
 			{
+				UClass* Class = InObjectToBind->GetClass();
+				if (Class->GetSuperClass())
+				{
+					if (TFunction<UObject* (const UObject*)>* Resolver = RemoteControlResolvers::CustomResolvers.Find(Class->GetSuperClass()->GetFName()))
+					{
+						UObject* CustomResolvedObject = (*Resolver)(InObjectToBind);
+						if (CustomResolvedObject && CustomResolvedObject->GetClass() == ResolvedOwnerClass)
+						{
+							FRemoteControlEntity::BindObject(CustomResolvedObject);
+							return;
+						}
+					}
+				}
+
 				// Search for a matching component if the root component was not a match.
 				FRemoteControlEntity::BindObject(Actor->GetComponentByClass(ResolvedOwnerClass));
 			}
@@ -89,6 +131,16 @@ bool FRemoteControlField::CanBindObject(const UObject* InObjectToBind) const
 				}
 				return false;
 			}
+
+			UClass* Class = InObjectToBind->GetClass();
+			if (Class->GetSuperClass())
+			{
+				if (TFunction<UObject* (const UObject*)>* Resolver = RemoteControlResolvers::CustomResolvers.Find(Class->GetSuperClass()->GetFName()))
+				{
+					UObject* CustomResolvedObject = (*Resolver)(InObjectToBind);
+					return CustomResolvedObject && CustomResolvedObject->GetClass() == ResolvedOwnerClass;
+				}
+			}
 		}
 	}
 	return false;
diff --git a/Engine/Plugins/VirtualProduction/RemoteControl/Source/RemoteControl/Private/RemoteControlPreset.cpp b/Engine/Plugins/VirtualProduction/RemoteControl/Source/RemoteControl/Private/RemoteControlPreset.cpp
--- a/Engine/Plugins/VirtualProduction/RemoteControl/Source/RemoteControl/Private/RemoteControlPreset.cpp
+++ b/Engine/Plugins/VirtualProduction/RemoteControl/Source/RemoteControl/Private/RemoteControlPreset.cpp
@@ -1030,6 +1030,70 @@ void URemoteControlPreset::PostLoadProperties()
 	}
 }
 
+void URemoteControlPreset::HandleDisplayClusterConfigChange(UObject* DisplayClusterConfigData)
+{
+#if WITH_EDITOR
+	AActor* OwnerActor = DisplayClusterConfigData->GetTypedOuter<AActor>();
+	GEditor->GetTimerManager()->SetTimerForNextTick(FTimerDelegate::CreateLambda([OwnerActor, PresetPtr = TWeakObjectPtr<URemoteControlPreset>{ this }]()
+	{
+		TSet<URemoteControlBinding*> ModifiedBindings;
+
+		if (URemoteControlPreset* Preset = PresetPtr.Get())
+		{
+			for (URemoteControlBinding* Binding : Preset->Bindings)
+			{
+				UObject* NewObject = nullptr;
+				UObject* ResolvedBinding = Binding->Resolve();
+
+				static const FName NDisplayConfigurationData = "DisplayClusterConfigurationData";
+				if (ResolvedBinding && ResolvedBinding->GetClass()->GetFName() == NDisplayConfigurationData)
+				{
+					AActor* BindingOwnerActor = ResolvedBinding->GetTypedOuter<AActor>();
+					if (OwnerActor == BindingOwnerActor)
+					{
+						// Replace binding
+						static const FName ConfigDataName = "CurrentConfigData";
+						FProperty* ConfigDataProperty = BindingOwnerActor->GetClass()->FindPropertyByName(ConfigDataName);
+
+						if (void* NewConfigData = ConfigDataProperty->ContainerPtrToValuePtr<void>(BindingOwnerActor))
+						{
+							if (FObjectProperty* ObjectProperty = CastField<FObjectProperty>(ConfigDataProperty))
+							{
+								NewObject = ObjectProperty->GetObjectPropertyValue(NewConfigData);
+							}
+						}
+					}
+				}
+
+				if (NewObject)
+				{
+					ModifiedBindings.Add(Binding);
+					PresetPtr->Modify();
+					Binding->Modify();
+					Binding->SetBoundObject(NewObject);
+				}
+
+				for (const TSharedPtr<FRemoteControlField>& Entity : PresetPtr->Registry->GetExposedEntities<FRemoteControlField>())
+				{
+					for (TWeakObjectPtr<URemoteControlBinding> WeakBinding : Entity->Bindings)
+					{
+						if (!WeakBinding.IsValid())
+						{
+							continue;
+						}
+
+						if (ModifiedBindings.Contains(WeakBinding.Get()))
+						{
+							PresetPtr->PerFrameUpdatedEntities.Add(Entity->GetId());
+						}
+					}
+				}
+			}
+		}
+	}));
+#endif
+}
+
 UWorld* URemoteControlPreset::GetWorld(const URemoteControlPreset* Preset, bool bAllowPIE)
 {
 	if (Preset && Preset->IsEmbeddedPreset())
@@ -1433,6 +1497,15 @@ void URemoteControlPreset::OnObjectPropertyChanged(UObject* Object, struct FProp
 {
 	// Objects modified should have run through the preobjectmodified. If interesting, they will be cached
 	TRACE_CPUPROFILER_EVENT_SCOPE(URemoteControlPreset::OnObjectPropertyChanged);
+
+	if (Object && Object->GetClass()->GetName() == TEXT("DisplayClusterConfigurationData"))
+	{
+		TRACE_CPUPROFILER_EVENT_SCOPE(URemoteControlPreset::OnObjectPropertyChanged::HandleDisplayClusterConfigChange);
+		// If a display cluster is modified, it will get invalidated so we need to update bindings to the new one.
+		// Since it does not show up in the ObjectReplaced callback, we have to do it on property change.
+		HandleDisplayClusterConfigChange(Object);
+		return;
+	}
  
 	if (Event.Property == nullptr)
 	{
@@ -1695,7 +1768,7 @@ void URemoteControlPreset::RegisterDelegates()
 
 	FEditorDelegates::PostPIEStarted.AddUObject(this, &URemoteControlPreset::OnPieEvent);
 	FEditorDelegates::EndPIE.AddUObject(this, &URemoteControlPreset::OnPieEvent);
-	
+
 	FCoreUObjectDelegates::OnObjectsReplaced.AddUObject(this, &URemoteControlPreset::OnReplaceObjects);
 
 	FEditorDelegates::MapChange.AddUObject(this, &URemoteControlPreset::OnMapChange);
@@ -1722,12 +1795,12 @@ void URemoteControlPreset::UnregisterDelegates()
 			Blueprint->OnCompiled().RemoveAll(this);
 		}
 	}
-	
+
 	FEditorDelegates::MapChange.RemoveAll(this);
 
 
 	FCoreUObjectDelegates::OnObjectsReplaced.RemoveAll(this);
-	
+
 	FEditorDelegates::EndPIE.RemoveAll(this);
 	FEditorDelegates::PostPIEStarted.RemoveAll(this);
 
@@ -1748,7 +1821,7 @@ void URemoteControlPreset::OnActorDeleted(AActor* Actor)
 	UWorld* World = Actor->GetWorld();
 
 	TSet<URemoteControlBinding*> ModifiedBindings;
-	
+
 	if (World && !World->IsPreviewWorld())
 	{
 		for (auto It = Bindings.CreateIterator(); It; ++It)
@@ -1765,14 +1838,14 @@ void URemoteControlPreset::OnActorDeleted(AActor* Actor)
 
 void URemoteControlPreset::OnPieEvent(bool)
 {
-	GEditor->GetTimerManager()->SetTimerForNextTick(FTimerDelegate::CreateLambda([PresetPtr = TWeakObjectPtr<URemoteControlPreset>{this}]()
+	GEditor->GetTimerManager()->SetTimerForNextTick(FTimerDelegate::CreateLambda([PresetPtr = TWeakObjectPtr<URemoteControlPreset>{ this }]()
 	{
 		if (PresetPtr.IsValid() && PresetPtr->Registry)
 		{
 			for (TSharedPtr<FRemoteControlEntity> Entity : PresetPtr->Registry->GetExposedEntities<FRemoteControlEntity>())
 			{
 				PresetPtr->PerFrameUpdatedEntities.Add(Entity->GetId());
-			}	
+			}
 		}
 	}));
 }
@@ -1780,14 +1853,19 @@ void URemoteControlPreset::OnPieEvent(bool)
 void URemoteControlPreset::OnReplaceObjects(const TMap<UObject*, UObject*>& ReplacementObjectMap)
 {
 	TSet<URemoteControlBinding*> ModifiedBindings;
-	
+
 	for (URemoteControlBinding* Binding : Bindings)
 	{
 		UObject* NewObject = nullptr;
+		UObject* ResolvedBinding = Binding->Resolve();
 
-		if (UObject* Replacement = ReplacementObjectMap.FindRef(Binding->Resolve()))
+
+		if (!NewObject)
 		{
-			NewObject = Replacement;
+			if (UObject* Replacement = ReplacementObjectMap.FindRef(ResolvedBinding))
+			{
+				NewObject = Replacement;
+			}
 		}
 
 		if (NewObject)
diff --git a/Engine/Plugins/VirtualProduction/RemoteControl/Source/RemoteControl/Private/RemoteControlPresetRebindingManager.cpp b/Engine/Plugins/VirtualProduction/RemoteControl/Source/RemoteControl/Private/RemoteControlPresetRebindingManager.cpp
--- a/Engine/Plugins/VirtualProduction/RemoteControl/Source/RemoteControl/Private/RemoteControlPresetRebindingManager.cpp
+++ b/Engine/Plugins/VirtualProduction/RemoteControl/Source/RemoteControl/Private/RemoteControlPresetRebindingManager.cpp
@@ -253,7 +253,6 @@ void FRemoteControlPresetRebindingManager::Rebind_Legacy(URemoteControlPreset* P
 	}
 }
 
-
 void FRemoteControlPresetRebindingManager::Rebind_NewAlgo(URemoteControlPreset* Preset)
 {
 	check(Preset);
@@ -341,16 +340,16 @@ void FRemoteControlPresetRebindingManager::Rebind_NewAlgo(URemoteControlPreset*
 		return nullptr;
 	};
 
-	auto TryRebindPair = [&BoundObjects](URemoteControlBinding* InBindingToRebind, AActor* InActor, UObject* InActorComponent)
+	auto TryRebindPair = [&](URemoteControlBinding* InBindingToRebind, AActor* InActor, UObject* InSubObject)
 	{
-		TPair<AActor*, UObject*> Pair = TPair<AActor*, UObject*> { InActor, InActorComponent };
+		TPair<AActor*, UObject*> Pair = TPair<AActor*, UObject*> { InActor, InSubObject };
 		if (!BoundObjects.Contains(Pair))
 		{
 			BoundObjects.Add(Pair);
 			InBindingToRebind->Modify();
-			if (InActorComponent)
+			if (InSubObject)
 			{
-				InBindingToRebind->SetBoundObject(InActorComponent);
+				InBindingToRebind->SetBoundObject(InSubObject);
 			}
 			else
 			{
@@ -361,18 +360,23 @@ void FRemoteControlPresetRebindingManager::Rebind_NewAlgo(URemoteControlPreset*
 		return false;
 	};
 
-	auto GetComponentBasedOnName = [](URemoteControlLevelDependantBinding* InBinding, UObject* InObject) -> UObject*
+	auto GetSubObjectBasedOnName = [](URemoteControlLevelDependantBinding* InBinding, UObject* InOwnerObject) -> UObject*
 	{
-		FName InitialComponentName = InBinding->BindingContext.ComponentName;
-		UObject* TargetComponent = FindObject<UObject>(InObject, *InitialComponentName.ToString());
-		if (TargetComponent && TargetComponent->GetClass()->IsChildOf(InBinding->BindingContext.SupportedClass.LoadSynchronous()))
+		FString SubObjectPath = InBinding->BindingContext.ComponentName.ToString();
+		if (InBinding->BindingContext.HasValidSubObjectPath())
+		{
+			SubObjectPath = InBinding->BindingContext.SubObjectPath;
+		}
+
+		UObject* TargetSubObject = FindObject<UObject>(InOwnerObject, *SubObjectPath);
+		if (TargetSubObject && TargetSubObject->GetClass()->IsChildOf(InBinding->BindingContext.SupportedClass.LoadSynchronous()))
 		{
-			return TargetComponent;
+			return TargetSubObject;
 		}
 		return nullptr;
 	};
 
-	auto RebindComponentBasedOnClass = [TryRebindPair](URemoteControlLevelDependantBinding* InBinding, AActor* InActorMatch)
+	auto RebindComponentBasedOnClass = [&](URemoteControlLevelDependantBinding* InBinding, AActor* InActorMatch)
 	{
 		if (UClass* SupportedClass = InBinding->BindingContext.SupportedClass.LoadSynchronous())
 		{
@@ -393,7 +397,7 @@ void FRemoteControlPresetRebindingManager::Rebind_NewAlgo(URemoteControlPreset*
 		return false;
 	};
 
-	auto RebindUsingActorName = [FindByName, &BoundObjects, TryRebindPair](URemoteControlLevelDependantBinding* BindingToRebind, const TArray<AActor*>& ObjectsWithSupportedClass)
+	auto RebindUsingActorName = [&](URemoteControlLevelDependantBinding* BindingToRebind, const TArray<AActor*>& ObjectsWithSupportedClass)
 	{
 		if (UObject* Match = FindByName(ObjectsWithSupportedClass, BindingToRebind->BindingContext.OwnerActorName))
 		{
@@ -402,7 +406,7 @@ void FRemoteControlPresetRebindingManager::Rebind_NewAlgo(URemoteControlPreset*
 		return false;
 	};
 
-	auto RebindUsingClass = [FindByName, &BoundObjects, TryRebindPair](URemoteControlLevelDependantBinding* BindingToRebind, const TArray<AActor*>& ObjectsWithSupportedClass)
+	auto RebindUsingClass = [&](URemoteControlLevelDependantBinding* BindingToRebind, const TArray<AActor*>& ObjectsWithSupportedClass)
 	{
 		for (UObject* PotentialMatch : ObjectsWithSupportedClass)
 		{
@@ -414,11 +418,11 @@ void FRemoteControlPresetRebindingManager::Rebind_NewAlgo(URemoteControlPreset*
 		return false;
 	};
 
-	auto TryRebindComponent = [GetComponentBasedOnName, TryRebindPair, RebindComponentBasedOnClass](URemoteControlLevelDependantBinding* BindingToRebind, UObject* PotentialMatch)
+	auto TryRebindSubObject = [&](URemoteControlLevelDependantBinding* BindingToRebind, UObject* PotentialMatch)
 	{
-		if (UObject* TargetComponent = GetComponentBasedOnName(BindingToRebind, PotentialMatch))
+		if (UObject* TargetSubObject = GetSubObjectBasedOnName(BindingToRebind, PotentialMatch))
 		{
-			return TryRebindPair(BindingToRebind, TargetComponent->GetTypedOuter<AActor>(), TargetComponent);
+			return TryRebindPair(BindingToRebind, TargetSubObject->GetTypedOuter<AActor>(), TargetSubObject);
 		}
 		else
 		{
@@ -426,52 +430,115 @@ void FRemoteControlPresetRebindingManager::Rebind_NewAlgo(URemoteControlPreset*
 		}
 	};
 
-	// Core of the rebinding algo
-	for (URemoteControlLevelDependantBinding* InvalidBinding : InvalidBindings)
+	auto FindSubObjectToRebind = [&](URemoteControlLevelDependantBinding* InvalidBinding, const TArray<AActor*>& ObjectsWithSupportedClass)
 	{
-		if (UClass* OwnerClass = InvalidBinding->BindingContext.OwnerActorClass.LoadSynchronous())
+		if (UObject* Match = FindByName(ObjectsWithSupportedClass, InvalidBinding->BindingContext.OwnerActorName))
+		{
+			// Found an owner object with matching name. Try to find a component under it with a matching name.
+			return TryRebindSubObject(InvalidBinding, Match);
+		}
+		else
 		{
-			const bool bRebindingComponent = !InvalidBinding->BindingContext.ComponentName.IsNone();
-			TArray<AActor*> ObjectsWithSupportedClass;
-			constexpr bool bAllowPIE = false;
-			if (!RCPresetRebindingManager::GetActorsOfClass(Preset->GetWorld(bAllowPIE), OwnerClass, ObjectsWithSupportedClass))
+			// Could not find an actor with the same name as the initial binding, rely only on class instead.
+			for (UObject* Object : ObjectsWithSupportedClass)
 			{
-				continue;
+				if (TryRebindSubObject(InvalidBinding, Object))
+				{
+					return true;
+				}
 			}
 
-			if (bRebindingComponent)
+			return false;
+		}
+	};
+
+	auto GetPotentialActors = [&](URemoteControlLevelDependantBinding* InvalidBinding, TArray<AActor*>& OutObjectsWithSupportedClass)-> bool
+	{
+		UClass* OwnerClass = InvalidBinding->BindingContext.OwnerActorClass.LoadSynchronous();
+		constexpr bool bAllowPIE = false;
+		return RCPresetRebindingManager::GetActorsOfClass(Preset->GetWorld(bAllowPIE), OwnerClass, OutObjectsWithSupportedClass);
+	};
+
+	auto Rebind = [&](URemoteControlLevelDependantBinding* InvalidBinding)
+	{
+		const FRemoteControlInitialBindingContext& BindingContext = InvalidBinding->BindingContext;
+		TArray<AActor*> ObjectsWithSupportedClass;
+		
+		if (GetPotentialActors(InvalidBinding, ObjectsWithSupportedClass))
+		{
+			if (BindingContext.HasValidSubObjectPath() || BindingContext.HasValidComponentName())
+			{
+				return FindSubObjectToRebind(InvalidBinding, ObjectsWithSupportedClass);
+			}
+			else
 			{
-				if (UObject* Match = FindByName(ObjectsWithSupportedClass, InvalidBinding->BindingContext.OwnerActorName))
+				if (!RebindUsingActorName(InvalidBinding, ObjectsWithSupportedClass))
 				{
-					// Found an owner object with matching name. Try to find a component under it with a matching name.
-					TryRebindComponent(InvalidBinding, Match);
+					return RebindUsingClass(InvalidBinding, ObjectsWithSupportedClass);
 				}
 				else
 				{
-					// Could not find an actor with the same name as the initial binding, rely only on class instead.
-					for (UObject* Object : ObjectsWithSupportedClass)
-					{
-						if (TryRebindComponent(InvalidBinding, Object))
-						{
-							break;
-						}
-					}
+					return true;
 				}
 			}
-			else
+		}
+		return false;
+	};
+
+	auto TryNDisplayRebind = [&](URemoteControlLevelDependantBinding* InvalidBinding)
+	{
+		// Used when the preset does not contain the full component path, we try to find an ndisplay viewport on a new object.
+		if (UBlueprintGeneratedClass* BPClass = Cast <UBlueprintGeneratedClass>(InvalidBinding->BindingContext.OwnerActorClass.Get()))
+		{
+			if (BPClass->GetSuperStruct() && BPClass->GetSuperStruct()->GetName() == TEXT("DisplayClusterRootActor"))
 			{
-				if (!RebindUsingActorName(InvalidBinding, ObjectsWithSupportedClass))
+				TArray<AActor*> ObjectsWithSupportedClass;
+				if (!GetPotentialActors(InvalidBinding, ObjectsWithSupportedClass))
+				{
+					return;
+				}
+
+				// We manually get the path from the last bound object since it was not saved.
+				// We are looking for a path like DisplayClusterConfigurationData_0.DisplayClusterConfigurationCluster_0.Node_1.VP_1
+				// We will then combine this with display cluster root actors in the scene to attempt finding a matching object.
+				FString LastObjectPath = InvalidBinding->GetLastBoundObject().ToSoftObjectPath().ToString();
+				if (LastObjectPath.IsEmpty())
+				{
+					return;
+				}
+
+				const int32 ConfigDataIndex = LastObjectPath.Find(TEXT("DisplayClusterConfigurationData"));
+
+				if (ConfigDataIndex == INDEX_NONE)
 				{
-					RebindUsingClass(InvalidBinding, ObjectsWithSupportedClass);
+					return;
+				}
+
+				FString SubObjectSubPath = LastObjectPath.RightChop(ConfigDataIndex);
+
+				for (AActor* Actor : ObjectsWithSupportedClass)
+				{
+					FTopLevelAssetPath AssetPath = FTopLevelAssetPath(Actor->GetPathName() + SubObjectSubPath);
+					if (UObject* SubObject = FindObject<UObject>(Actor, *SubObjectSubPath))
+					{
+						TryRebindPair(InvalidBinding, Actor, SubObject);
+					}
 				}
 			}
 		}
+	};
+
+	for (URemoteControlLevelDependantBinding* InvalidBinding : InvalidBindings)
+	{
+		if (!Rebind(InvalidBinding))
+		{
+			TryNDisplayRebind(InvalidBinding);
+		}
 	}
 
 	Preset->RemoveUnusedBindings();
 }
 
-
 TArray<FGuid> FRemoteControlPresetRebindingManager::RebindAllEntitiesUnderSameActor(URemoteControlPreset* Preset, URemoteControlBinding* InitialBinding, AActor* NewActor)
 {
 	check(Preset && InitialBinding && NewActor);
diff --git a/Engine/Plugins/VirtualProduction/RemoteControl/Source/RemoteControl/Public/RemoteControlBinding.h b/Engine/Plugins/VirtualProduction/RemoteControl/Source/RemoteControl/Public/RemoteControlBinding.h
--- a/Engine/Plugins/VirtualProduction/RemoteControl/Source/RemoteControl/Public/RemoteControlBinding.h
+++ b/Engine/Plugins/VirtualProduction/RemoteControl/Source/RemoteControl/Public/RemoteControlBinding.h
@@ -99,6 +99,12 @@ struct FRemoteControlInitialBindingContext
 	UPROPERTY()
 	FName ComponentName;
 
+	/**
+	 * Path to the subobject if any.
+	 */
+	UPROPERTY()
+	FString SubObjectPath;
+
 	/**
 	 * The class of the actor that's targeted by this binding or the class of the actor that owns the component that is targeted.
 	 */
@@ -124,13 +130,24 @@ struct FRemoteControlInitialBindingContext
 		return LHS.SupportedClass == RHS.SupportedClass
 			&& LHS.ComponentName == RHS.ComponentName
 			&& LHS.OwnerActorClass == RHS.OwnerActorClass
-			&& LHS.OwnerActorName == RHS.OwnerActorName;
+			&& LHS.OwnerActorName == RHS.OwnerActorName
+			&& LHS.SubObjectPath == RHS.SubObjectPath;
 	}
 
 	friend bool operator!=(const FRemoteControlInitialBindingContext& LHS, const FRemoteControlInitialBindingContext& RHS)
 	{
 		return !(LHS == RHS);
 	}
+
+	bool HasValidComponentName() const
+	{
+		return !ComponentName.IsNone();
+	}
+
+	bool HasValidSubObjectPath() const
+	{
+		return !SubObjectPath.IsEmpty();
+	}
 };
 
 UCLASS(BlueprintType)
@@ -180,6 +197,8 @@ private:
 
 	void InitializeBindingContext(UObject* InObject) const;
 
+	TSoftObjectPtr<UObject> GetLastBoundObject() const;
+
 private:
 #if WITH_EDITORONLY_DATA
 	UPROPERTY(meta = (DeprecatedProperty, DeprecationMessage = "BoundObjectMap is deprecated, please use BoundObjectMapByPath instead."))
diff --git a/Engine/Plugins/VirtualProduction/RemoteControl/Source/RemoteControl/Public/RemoteControlPreset.h b/Engine/Plugins/VirtualProduction/RemoteControl/Source/RemoteControl/Public/RemoteControlPreset.h
--- a/Engine/Plugins/VirtualProduction/RemoteControl/Source/RemoteControl/Public/RemoteControlPreset.h
+++ b/Engine/Plugins/VirtualProduction/RemoteControl/Source/RemoteControl/Public/RemoteControlPreset.h
@@ -766,6 +766,9 @@ private:
 
 	/** Call post load function for exposed properties. */
 	void PostLoadProperties();
+
+	/** Handle a change on a ndisplay config data, used to replace bindings. */
+	void HandleDisplayClusterConfigChange(UObject* DisplayClusterConfigData);
 	
 private:
 	/** Preset unique ID */
diff --git a/Engine/Plugins/VirtualProduction/RemoteControl/Source/RemoteControlUI/Private/UI/RCPanelWidgetRegistry.cpp b/Engine/Plugins/VirtualProduction/RemoteControl/Source/RemoteControlUI/Private/UI/RCPanelWidgetRegistry.cpp
--- a/Engine/Plugins/VirtualProduction/RemoteControl/Source/RemoteControlUI/Private/UI/RCPanelWidgetRegistry.cpp
+++ b/Engine/Plugins/VirtualProduction/RemoteControl/Source/RemoteControlUI/Private/UI/RCPanelWidgetRegistry.cpp
@@ -134,6 +134,7 @@ TSharedPtr<IDetailTreeNode> FRCPanelWidgetRegistry::GetObjectTreeNode(UObject* I
 		Generator = FModuleManager::GetModuleChecked<FPropertyEditorModule>("PropertyEditor").CreatePropertyRowGenerator(Args);
 		Generator->SetCustomValidatePropertyNodesFunction(FOnValidatePropertyRowGeneratorNodes::CreateLambda(MoveTemp(ValidationLambda)));
 		Generator->SetObjects({InObject});
+		Generator->OnRowsRefreshed().AddRaw(this, &FRCPanelWidgetRegistry::OnRowsRefreshed, Generator);
 		ObjectToRowGenerator.Add(WeakObject, Generator);
 	}
 	
@@ -250,3 +251,18 @@ TSharedPtr<IDetailTreeNode> FRCPanelWidgetRegistry::FindNDisplayTreeNode(UObject
 
 	return WidgetRegistryUtils::FindNode(Generator->GetRootTreeNodes(), InField, InFindMethod);
 }
+
+void FRCPanelWidgetRegistry::OnRowsRefreshed(TSharedPtr<IPropertyRowGenerator> Generator)
+{
+	if (Generator)
+	{
+		const TArray<TWeakObjectPtr<UObject>> WeakSelectedObjects = Generator->GetSelectedObjects();
+		TArray<UObject*> SelectedObjects;
+		SelectedObjects.Reserve(WeakSelectedObjects.Num());
+		for (const TWeakObjectPtr<UObject>& WeakObject : WeakSelectedObjects)
+		{
+			SelectedObjects.Add(WeakObject.Get());
+		}
+		OnObjectRefreshedDelegate.Broadcast(SelectedObjects);
+	}
+}
diff --git a/Engine/Plugins/VirtualProduction/RemoteControl/Source/RemoteControlUI/Private/UI/RCPanelWidgetRegistry.h b/Engine/Plugins/VirtualProduction/RemoteControl/Source/RemoteControlUI/Private/UI/RCPanelWidgetRegistry.h
--- a/Engine/Plugins/VirtualProduction/RemoteControl/Source/RemoteControlUI/Private/UI/RCPanelWidgetRegistry.h
+++ b/Engine/Plugins/VirtualProduction/RemoteControl/Source/RemoteControlUI/Private/UI/RCPanelWidgetRegistry.h
@@ -16,6 +16,8 @@ enum class ERCFindNodeMethod
 	Path  // Find a node by providing a property path.
 };
 
+DECLARE_MULTICAST_DELEGATE_OneParam(FOnObjectsRefreshed, const TArray<UObject*>&);
+
 /** Local special case handlers to generate widgets to be displayed */
 using FRCTreeNodeFinderIsException = TFunction< bool(UObject* /*InObject*/, const FString& /*InField*/, ERCFindNodeMethod /*InFindMethod*/) >;
 using FRCTreeNodeFinderFunction = TFunction< TSharedPtr<IDetailTreeNode>(UObject* /*InObject*/, const FString& /*InField*/, ERCFindNodeMethod /*InFindMethod*/) >;
@@ -62,11 +64,20 @@ public:
 	 */
 	void Clear();
 
+	/**
+	 * Get the on object refresh delegate.
+	 */
+	FOnObjectsRefreshed& OnObjectRefreshed() 
+	{
+		return OnObjectRefreshedDelegate; 
+	}
+
 private:
 
 	/** Special handlers for ndisplay root actor and its details customizations */
 	bool IsNDisplayObject(UObject* InObject, const FString& InField, ERCFindNodeMethod InFindMethod);
 	TSharedPtr<IDetailTreeNode> FindNDisplayTreeNode(UObject* InObject, const FString& InField, ERCFindNodeMethod InFindMethod);
+	void OnRowsRefreshed(TSharedPtr<IPropertyRowGenerator> Generator);
 
 private:
 	/** Map of objects to row generator, used to have one row generator per object. */
@@ -77,4 +88,6 @@ private:
 	TMap<TPair<TWeakObjectPtr<UObject>, FString>, TWeakPtr<IDetailTreeNode>> TreeNodeCache;
 	/** List of tree node finder handlers for certain type */
 	TArray<FRCTreeNodeFinderHandler> SpecialTreeNodeHandlers;
+	/** Called when a generator gets refreshed. (ie. forcefully refreshed by a customization) */
+	FOnObjectsRefreshed OnObjectRefreshedDelegate;
 };
diff --git a/Engine/Plugins/VirtualProduction/RemoteControl/Source/RemoteControlUI/Private/UI/SRCPanelExposedField.cpp b/Engine/Plugins/VirtualProduction/RemoteControl/Source/RemoteControlUI/Private/UI/SRCPanelExposedField.cpp
--- a/Engine/Plugins/VirtualProduction/RemoteControl/Source/RemoteControlUI/Private/UI/SRCPanelExposedField.cpp
+++ b/Engine/Plugins/VirtualProduction/RemoteControl/Source/RemoteControlUI/Private/UI/SRCPanelExposedField.cpp
@@ -98,6 +98,14 @@ void SRCPanelExposedField::Construct(const FArguments& InArgs, TWeakPtr<FRemoteC
 	}
 }
 
+SRCPanelExposedField::~SRCPanelExposedField()
+{
+	if (TSharedPtr<FRCPanelWidgetRegistry> Registry = WidgetRegistry.Pin())
+	{
+		Registry->OnObjectRefreshed().RemoveAll(this);
+	}
+}
+
 void SRCPanelExposedField::GetNodeChildren(TArray<TSharedPtr<SRCPanelTreeNode>>& OutChildren) const
 {
 	OutChildren.Append(ChildWidgets);
@@ -327,6 +335,9 @@ TSharedRef<SWidget> SRCPanelExposedField::ConstructWidget()
 		{
 			if (TSharedPtr<FRCPanelWidgetRegistry> Registry = WidgetRegistry.Pin())
 			{
+				Registry->OnObjectRefreshed().RemoveAll(this);
+				Registry->OnObjectRefreshed().AddSP(this, &SRCPanelExposedField::OnObjectsRefreshed);
+
 				if (TSharedPtr<IDetailTreeNode> Node = Registry->GetObjectTreeNode(Objects[0], Field->FieldPathInfo.ToPathPropertyString(), ERCFindNodeMethod::Path))
 				{
 					TArray<TSharedRef<IDetailTreeNode>> ChildNodes;
@@ -360,7 +371,7 @@ TSharedRef<SWidget> SRCPanelExposedField::ConstructWidget()
 						{
 							TSharedPtr<SRCPanelExposedField> FieldPtr = WeakFieldPtr.Pin();
 							TSharedPtr<IPropertyHandle> PropertyHandle = Node->CreatePropertyHandle();
-							if (FieldPtr && PropertyHandle)
+							if (FieldPtr && PropertyHandle && PropertyHandle->IsValidHandle())
 							{
 								void* DataPtr;
 								PropertyHandle->GetValueData(DataPtr);
@@ -561,6 +572,24 @@ FReply SRCPanelExposedField::OnClickFunctionButton()
 	return FReply::Handled();
 }
 
+void SRCPanelExposedField::OnObjectsRefreshed(const TArray<UObject*>& RefreshedObjects)
+{
+	if (TSharedPtr<FRemoteControlField> Field = WeakField.Pin())
+	{
+		TSet<UObject*> BoundObjectsSet;
+		BoundObjectsSet.Append(Field->GetBoundObjects());
+
+		TSet<UObject*> RefreshedObjectsSet;
+		RefreshedObjectsSet.Append(RefreshedObjects);
+
+		const TSet<UObject*> CommonObjects = BoundObjectsSet.Intersect(RefreshedObjectsSet);
+		if (CommonObjects.Num())
+		{
+			Refresh();
+		}
+	}
+}
+
 void SRCPanelFieldChildNode::Construct(const FArguments& InArgs, const TSharedRef<IDetailTreeNode>& InNode, FRCColumnSizeData InColumnSizeData)
 {
 	TArray<TSharedRef<IDetailTreeNode>> ChildNodes;
diff --git a/Engine/Plugins/VirtualProduction/RemoteControl/Source/RemoteControlUI/Private/UI/SRCPanelExposedField.h b/Engine/Plugins/VirtualProduction/RemoteControl/Source/RemoteControlUI/Private/UI/SRCPanelExposedField.h
--- a/Engine/Plugins/VirtualProduction/RemoteControl/Source/RemoteControlUI/Private/UI/SRCPanelExposedField.h
+++ b/Engine/Plugins/VirtualProduction/RemoteControl/Source/RemoteControlUI/Private/UI/SRCPanelExposedField.h
@@ -41,6 +41,7 @@ struct SRCPanelExposedField : public SRCPanelExposedEntity
 	static TSharedPtr<SRCPanelTreeNode> MakeInstance(const FGenerateWidgetArgs& Args);
 
 	void Construct(const FArguments& InArgs, TWeakPtr<FRemoteControlField> Field, FRCColumnSizeData ColumnSizeData, TWeakPtr<FRCPanelWidgetRegistry> InWidgetRegistry);
+	virtual ~SRCPanelExposedField();
 
 	//~ SRCPanelTreeNode Interface 
 	virtual void GetNodeChildren(TArray<TSharedPtr<SRCPanelTreeNode>>& OutChildren) const override;
@@ -90,6 +91,8 @@ private:
 	TSharedRef<SWidget> ConstructCallFunctionButton(bool bIsEnabled = true);
 	/** Handles calling an exposed function.*/
 	FReply OnClickFunctionButton();
+	/** Called when the underlying generator gets refreshed. In that case we must re-fetch the widget. */
+	void OnObjectsRefreshed(const TArray<UObject*>& RefreshedObjects);
 
 private:
 	/** Weak pointer to the underlying RC Field. */
