1. RC Functions Cannot Add

diff --git a/Engine/Plugins/VirtualProduction/RemoteControl/Source/RemoteControlUI/Private/UI/SRCPanelExposedEntitiesList.cpp b/Engine/Plugins/VirtualProduction/RemoteControl/Source/RemoteControlUI/Private/UI/SRCPanelExposedEntitiesList.cpp
--- a/Engine/Plugins/VirtualProduction/RemoteControl/Source/RemoteControlUI/Private/UI/SRCPanelExposedEntitiesList.cpp
+++ b/Engine/Plugins/VirtualProduction/RemoteControl/Source/RemoteControlUI/Private/UI/SRCPanelExposedEntitiesList.cpp
@@ -400,7 +400,8 @@ void SRCPanelExposedEntitiesList::Construct(const FArguments& InArgs, URemoteCon
 	// Expose Button
 	ExposeDockPanel->AddHeaderToolbarItem(EToolbar::Left, PlaceholderBox.ToSharedRef());
 	ExposeDockPanel->AddHeaderToolbarItem(EToolbar::Right, ModeSwitcher.ToSharedRef());
-	ExposeDockPanel->AddFooterToolbarItem(EToolbar::Right, InArgs._ExposeComboButton.Get().ToSharedRef());
+	ExposeDockPanel->AddFooterToolbarItem(EToolbar::Left, InArgs._ExposeActorsComboButton.Get().ToSharedRef());
+	ExposeDockPanel->AddFooterToolbarItem(EToolbar::Left, InArgs._ExposeFunctionsComboButton.Get().ToSharedRef());
 	ExposeDockPanel->AddFooterToolbarItem(EToolbar::Right, DeleteAllEntitiesButton.ToSharedRef());
 
 	ExposePanel->AddPanel(ExposeDockPanel.ToSharedRef(), 0.75f);
diff --git a/Engine/Plugins/VirtualProduction/RemoteControl/Source/RemoteControlUI/Private/UI/SRCPanelExposedEntitiesList.h b/Engine/Plugins/VirtualProduction/RemoteControl/Source/RemoteControlUI/Private/UI/SRCPanelExposedEntitiesList.h
--- a/Engine/Plugins/VirtualProduction/RemoteControl/Source/RemoteControlUI/Private/UI/SRCPanelExposedEntitiesList.h
+++ b/Engine/Plugins/VirtualProduction/RemoteControl/Source/RemoteControlUI/Private/UI/SRCPanelExposedEntitiesList.h
@@ -62,7 +62,8 @@ public:
 	{}
 		SLATE_ATTRIBUTE(bool, LiveMode)
 		SLATE_ATTRIBUTE(bool, ProtocolsMode)
-		SLATE_ATTRIBUTE(TSharedPtr<SWidget>, ExposeComboButton)
+		SLATE_ATTRIBUTE(TSharedPtr<SWidget>, ExposeFunctionsComboButton)
+		SLATE_ATTRIBUTE(TSharedPtr<SWidget>, ExposeActorsComboButton)
 		SLATE_EVENT(FSimpleDelegate, OnEntityListUpdated)
 	SLATE_END_ARGS()
 
diff --git a/Engine/Plugins/VirtualProduction/RemoteControl/Source/RemoteControlUI/Private/UI/SRemoteControlPanel.cpp b/Engine/Plugins/VirtualProduction/RemoteControl/Source/RemoteControlUI/Private/UI/SRemoteControlPanel.cpp
--- a/Engine/Plugins/VirtualProduction/RemoteControl/Source/RemoteControlUI/Private/UI/SRemoteControlPanel.cpp
+++ b/Engine/Plugins/VirtualProduction/RemoteControl/Source/RemoteControlUI/Private/UI/SRemoteControlPanel.cpp
@@ -392,7 +392,8 @@ void SRemoteControlPanel::Construct(const FArguments& InArgs, URemoteControlPres
 	EntityProtocolDetails = SNew(SBox);
 	
 	EntityList = SNew(SRCPanelExposedEntitiesList, Preset.Get(), WidgetRegistry)
-		.ExposeComboButton(CreateExposeButton())
+		.ExposeActorsComboButton(CreateExposeActorsButton())
+		.ExposeFunctionsComboButton(CreateExposeFunctionsButton())
 		.OnEntityListUpdated_Lambda([this] ()
 			{
 				UpdateEntityDetailsView(EntityList->GetSelectedEntity());
@@ -1001,7 +1002,7 @@ TSharedRef<SWidget> SRemoteControlPanel::CreateCPUThrottleWarning() const
 		];
 }
 
-TSharedRef<SWidget> SRemoteControlPanel::CreateExposeButton()
+TSharedRef<SWidget> SRemoteControlPanel::CreateExposeFunctionsButton()
 {	
 	FMenuBuilder MenuBuilder(true, nullptr);
 	
@@ -1061,6 +1062,60 @@ TSharedRef<SWidget> SRemoteControlPanel::CreateExposeButton()
 			LOCTEXT("ActorFunctionSubMenuToolTip", "Expose an actor's function."),
 			ActorFunctionPicker.ToSharedRef()
 		);
+	}
+
+	MenuBuilder.EndSection();
+	
+	return SAssignNew(ExposeFunctionsComboButton, SComboButton)
+		.AddMetaData<FTagMetaData>(FTagMetaData(TEXT("Expose Functions")))
+		.IsEnabled_Lambda([this]() { return !this->bIsInLiveMode; })
+		.HAlign(HAlign_Center)
+		.VAlign(VAlign_Center)
+		.ButtonStyle(&RCPanelStyle->FlatButtonStyle)
+		.ForegroundColor(FSlateColor::UseForeground())
+		.CollapseMenuOnParentFocus(true)
+		.HasDownArrow(false)
+		.ContentPadding(FMargin(4.f, 2.f))
+		.ButtonContent()
+		[
+			SNew(SBox)
+			.WidthOverride(RCPanelStyle->IconSize.X)
+			.HeightOverride(RCPanelStyle->IconSize.Y)
+			[
+				SNew(SImage)
+				.ColorAndOpacity(FSlateColor::UseForeground())
+				.Image(FAppStyle::GetBrush("GraphEditor.Function_16x"))
+			]
+		]
+		.MenuContent()
+		[
+			MenuBuilder.MakeWidget()
+		];
+}
+
+TSharedRef<SWidget> SRemoteControlPanel::CreateExposeActorsButton()
+{	
+	FMenuBuilder MenuBuilder(true, nullptr);
+	
+	MenuBuilder.BeginSection(NAME_None, LOCTEXT("ExposeHeader", "Expose"));
+	{
+		constexpr bool bNoIndent = true;
+		constexpr bool bSearchable = false;
+
+		auto CreatePickerSubMenu = [this, bNoIndent, bSearchable, &MenuBuilder] (const FText& Label, const FText& ToolTip, const TSharedRef<SWidget>& Widget)
+		{
+			MenuBuilder.AddSubMenu(
+				Label,
+				ToolTip,
+				FNewMenuDelegate::CreateLambda(
+					[this, bNoIndent, bSearchable, Widget](FMenuBuilder& MenuBuilder)
+					{
+						MenuBuilder.AddWidget(Widget, FText::GetEmpty(), bNoIndent, bSearchable);
+						FSlateApplication::Get().SetKeyboardFocus(Widget, EFocusCause::Navigation);
+					}
+				)
+			);
+		};
 
 		// SObjectPropertyEntryBox does not support non-level editor worlds.
 		if (!Preset.IsValid() || !Preset->IsEmbeddedPreset())
@@ -1107,11 +1162,12 @@ TSharedRef<SWidget> SRemoteControlPanel::CreateExposeButton()
 			CreateExposeByClassWidget()
 		);
 	}
+
 	MenuBuilder.EndSection();
 	
-	return SAssignNew(ExposeComboButton, SComboButton)
-		.AddMetaData<FTagMetaData>(FTagMetaData(TEXT("Expose Functions")))
-		.IsEnabled_Lambda([this]() { return this->bIsInLiveMode; })
+	return SAssignNew(ExposeActorsComboButton, SComboButton)
+		.AddMetaData<FTagMetaData>(FTagMetaData(TEXT("Expose Actors")))
+		.IsEnabled_Lambda([this]() { return !this->bIsInLiveMode; })
 		.HAlign(HAlign_Center)
 		.VAlign(VAlign_Center)
 		.ButtonStyle(&RCPanelStyle->FlatButtonStyle)
@@ -1127,7 +1183,7 @@ TSharedRef<SWidget> SRemoteControlPanel::CreateExposeButton()
 			[
 				SNew(SImage)
 				.ColorAndOpacity(FSlateColor::UseForeground())
-				.Image(FAppStyle::GetBrush("GraphEditor.Function_16x"))
+				.Image(FAppStyle::GetBrush("ClassIcon.Actor"))
 			]
 		]
 		.MenuContent()
@@ -1189,9 +1245,9 @@ TSharedRef<SWidget> SRemoteControlPanel::CreateExposeByClassWidget()
 				}
 			}
 
-			if (ExposeComboButton)
+			if (ExposeActorsComboButton)
 			{
-				ExposeComboButton->SetIsOpen(false);
+				ExposeActorsComboButton->SetIsOpen(false);
 			}
 		}));
 
@@ -1597,7 +1653,7 @@ void SRemoteControlPanel::OnLogCheckboxToggle(ECheckBoxState State)
 		}
 	}
 }
- 
+
 void SRemoteControlPanel::OnBlueprintReinstanced()
 {
 	Refresh();
@@ -1615,9 +1671,9 @@ void SRemoteControlPanel::ExposeProperty(UObject* Object, FRCFieldPathInfo Path)
 
 void SRemoteControlPanel::ExposeFunction(UObject* Object, UFunction* Function)
 {
-	if (ExposeComboButton)
+	if (ExposeFunctionsComboButton)
 	{
-		ExposeComboButton->SetIsOpen(false);
+		ExposeFunctionsComboButton->SetIsOpen(false);
 	}
 	
 	FScopedTransaction Transaction(LOCTEXT("ExposeFunction", "ExposeFunction"));
diff --git a/Engine/Plugins/VirtualProduction/RemoteControl/Source/RemoteControlUI/Private/UI/SRemoteControlPanel.h b/Engine/Plugins/VirtualProduction/RemoteControl/Source/RemoteControlUI/Private/UI/SRemoteControlPanel.h
--- a/Engine/Plugins/VirtualProduction/RemoteControl/Source/RemoteControlUI/Private/UI/SRemoteControlPanel.h
+++ b/Engine/Plugins/VirtualProduction/RemoteControl/Source/RemoteControlUI/Private/UI/SRemoteControlPanel.h
@@ -181,9 +181,6 @@ private:
 	/** Unexpose a field from the preset. */
 	void Unexpose(const FRCExposesPropertyArgs& InArgs);
 
-	/** Handle the using toggling the edit mode check box. */
-	void OnEditModeCheckboxToggle(ECheckBoxState State);
-
 	/** Handler called when a blueprint is reinstanced. */
 	void OnBlueprintReinstanced();
 
@@ -206,7 +203,10 @@ private:
 	TSharedRef<SWidget> CreateCPUThrottleWarning() const;
 
 	/** Create expose button, allowing to expose blueprints and actor functions. */
-	TSharedRef<SWidget> CreateExposeButton();
+	TSharedRef<SWidget> CreateExposeFunctionsButton();
+	
+	/** Create expose button, allowing to expose actors. */
+	TSharedRef<SWidget> CreateExposeActorsButton();
 
 	/** Create expose by class menu content */
 	TSharedRef<SWidget> CreateExposeByClassWidget();
@@ -361,8 +361,10 @@ private:
 	TSharedPtr<SRCPanelFunctionPicker> SubsystemFunctionPicker;
 	/** Holds the exposed entity list view. */
 	TSharedPtr<SRCPanelExposedEntitiesList> EntityList;
-	/** Holds the combo button that allows exposing functions and actors. */
-	TSharedPtr<SComboButton> ExposeComboButton;
+	/** Holds the combo button that allows exposing functions. */
+	TSharedPtr<SComboButton> ExposeFunctionsComboButton;
+	/** Holds the combo button that allows exposing actors. */
+	TSharedPtr<SComboButton> ExposeActorsComboButton;
 	/** Caches all the classes of actors in the current level. */
 	TSet<TWeakObjectPtr<const UClass>> CachedClassesInLevel;
 	/** Holds the class picker used to expose all actors of class. */
