Fix ndisplay config rebind all not working if generated for a different config (Viewport properties excluded)
Add context menu for rebinding viewport properties

diff --git a/Engine/Plugins/VirtualProduction/RemoteControl/Source/RemoteControl/Private/RemoteControlBinding.cpp b/Engine/Plugins/VirtualProduction/RemoteControl/Source/RemoteControl/Private/RemoteControlBinding.cpp
--- a/Engine/Plugins/VirtualProduction/RemoteControl/Source/RemoteControl/Private/RemoteControlBinding.cpp
+++ b/Engine/Plugins/VirtualProduction/RemoteControl/Source/RemoteControl/Private/RemoteControlBinding.cpp
@@ -2,6 +2,8 @@
 
 #include "RemoteControlBinding.h"
 
+#include "Components/BillboardComponent.h"
+#include "Engine/Brush.h"
 #include "Engine/Engine.h"
 #include "Engine/EngineTypes.h"
 #include "Engine/World.h"
@@ -21,6 +23,49 @@
 
 #define LOCTEXT_NAMESPACE "RemoteControlBinding"
 
+
+namespace UE::RemoteControlBinding
+{
+	bool IsValidObjectForRebinding(UObject* InObject, UWorld* PresetWorld)
+	{
+		return InObject
+			&& !InObject->IsA<UBillboardComponent>()
+			&& InObject->GetName().Find(TEXT("TRASH_")) == INDEX_NONE
+			&& InObject->GetTypedOuter<UPackage>() != GetTransientPackage()
+			&& InObject->GetWorld() == PresetWorld;
+	}
+
+	bool IsValidActorForRebinding(AActor* InActor, UWorld* PresetWorld)
+	{
+		return IsValidObjectForRebinding(InActor, PresetWorld) &&
+#if WITH_EDITOR
+			InActor->IsEditable() &&
+			InActor->IsListedInSceneOutliner() &&
+#endif
+			!InActor->IsTemplate() &&
+			InActor->GetClass() != ABrush::StaticClass() && // Workaround Brush being listed as visible in the scene outliner even though it's not.
+			!InActor->HasAnyFlags(RF_Transient);
+	}
+
+	bool IsValidSubObjectForRebinding(UObject* InComponent, UWorld* PresetWorld)
+	{
+		if (!IsValidObjectForRebinding(InComponent, PresetWorld))
+		{
+			return false;
+		}
+
+		if (AActor* OuterActor = InComponent->GetTypedOuter<AActor>())
+		{
+			if (!IsValidActorForRebinding(OuterActor, PresetWorld))
+			{
+				return false;
+			}
+		}
+
+		return true;
+	}
+}
+
 namespace
 {
 	/**
diff --git a/Engine/Plugins/VirtualProduction/RemoteControl/Source/RemoteControl/Private/RemoteControlPresetRebindingManager.cpp b/Engine/Plugins/VirtualProduction/RemoteControl/Source/RemoteControl/Private/RemoteControlPresetRebindingManager.cpp
--- a/Engine/Plugins/VirtualProduction/RemoteControl/Source/RemoteControl/Private/RemoteControlPresetRebindingManager.cpp
+++ b/Engine/Plugins/VirtualProduction/RemoteControl/Source/RemoteControl/Private/RemoteControlPresetRebindingManager.cpp
@@ -4,6 +4,7 @@
 
 #include "Algo/AllOf.h"
 #include "Algo/Transform.h"
+#include "Engine/BlueprintGeneratedClass.h"
 #include "Components/ActorComponent.h"
 #include "Engine/Brush.h"
 #include "Engine/Classes/Components/ActorComponent.h"
@@ -24,44 +25,6 @@ static TAutoConsoleVariable<int32> CVarRemoteControlRebindingUseLegacyAlgo(TEXT(
 
 namespace RCPresetRebindingManager
 {
-
-	bool IsValidObjectForRebinding(UObject* InObject, UWorld* PresetWorld)
-	{
-		return InObject
-			&& InObject->GetTypedOuter<UPackage>() != GetTransientPackage()
-			&& InObject->GetWorld() == PresetWorld;
-	}
-
-	bool IsValidActorForRebinding(AActor* InActor, UWorld* PresetWorld)
-	{
-		return IsValidObjectForRebinding(InActor, PresetWorld) &&
-#if WITH_EDITOR
-            InActor->IsEditable() &&
-            InActor->IsListedInSceneOutliner() &&
-#endif
-			!InActor->IsTemplate() &&
-            InActor->GetClass() != ABrush::StaticClass() && // Workaround Brush being listed as visible in the scene outliner even though it's not.
-            !InActor->HasAnyFlags(RF_Transient);
-	}
-
-	bool IsValidComponentForRebinding(UActorComponent* InComponent, UWorld* PresetWorld)
-	{
-		if (!IsValidObjectForRebinding(InComponent, PresetWorld))
-		{
-			return false;
-		}
-		
-		if (AActor* OuterActor = InComponent->GetOwner())
-		{
-			if (!IsValidActorForRebinding(OuterActor, PresetWorld))
-			{
-				return false;
-			}
-		}
-
-		return true;
-	}
-
 	TArray<UClass*> GetRelevantClassesForObject(const TArray<UClass*>& InClasses, UObject* InObject)
 	{
 		return InClasses.FilterByPredicate([&InObject](UClass* InClass) { return InClass && InObject->IsA(InClass); });
@@ -78,8 +41,8 @@ namespace RCPresetRebindingManager
 		auto AddRelevantClassesForObject = [&ObjectMap, &RelevantClasses, PresetWorld] (UObject* Object)
 		{
 			if (!Object
-				|| (Object->IsA<AActor>() && !IsValidActorForRebinding(CastChecked<AActor>(Object), PresetWorld))
-				|| (Object->IsA<UActorComponent>() && !IsValidComponentForRebinding(CastChecked<UActorComponent>(Object), PresetWorld)))
+				|| (Object->IsA<AActor>() && !UE::RemoteControlBinding::IsValidActorForRebinding(CastChecked<AActor>(Object), PresetWorld))
+				|| (Object->IsA<UActorComponent>() && !UE::RemoteControlBinding::IsValidSubObjectForRebinding(Object, PresetWorld)))
 			{
 				return;
 			}
@@ -109,11 +72,29 @@ namespace RCPresetRebindingManager
 
 	bool GetActorsOfClass(UWorld* PresetWorld, UClass* InTargetClass, TArray<AActor*>& OutActors)
 	{
-		auto IsObjectOfClass = [InTargetClass, PresetWorld](UObject* Object)
+		auto WasGeneratedBySameBlueprintClass = [InTargetClass](UObject* Object)
+		{
+			UBlueprintGeneratedClass* IteratedClass = Cast<UBlueprintGeneratedClass>(Object->GetClass());
+			UBlueprintGeneratedClass* TargetClass = Cast<UBlueprintGeneratedClass>(InTargetClass);
+
+			// At the moment, only allow for ndisplay configs because I'm unsure of the repercussions of enabling this change for all BP generated classes.
+			static const FName NDisplayGeneratedClass = "DisplayClusterBlueprintGeneratedClass";
+			
+#if WITH_EDITOR
+			return IteratedClass && TargetClass
+				&& InTargetClass->GetClass()->GetFName() == NDisplayGeneratedClass
+				&& IteratedClass->ClassGeneratedBy && TargetClass->ClassGeneratedBy
+				&& IteratedClass->ClassGeneratedBy->GetClass() == TargetClass->ClassGeneratedBy->GetClass();
+#else
+			return false;
+#endif
+		};
+
+		auto IsObjectOfClass = [&](UObject* Object)
 		{
 			return Object
-				&& IsValidActorForRebinding(CastChecked<AActor>(Object), PresetWorld)
-				&& Object->IsA(InTargetClass);
+				&& UE::RemoteControlBinding::IsValidActorForRebinding(CastChecked<AActor>(Object), PresetWorld)
+				&& (Object->IsA(InTargetClass) || WasGeneratedBySameBlueprintClass(Object));
 		};
 
 		if (PresetWorld)
diff --git a/Engine/Plugins/VirtualProduction/RemoteControl/Source/RemoteControl/Public/RemoteControlBinding.h b/Engine/Plugins/VirtualProduction/RemoteControl/Source/RemoteControl/Public/RemoteControlBinding.h
--- a/Engine/Plugins/VirtualProduction/RemoteControl/Source/RemoteControl/Public/RemoteControlBinding.h
+++ b/Engine/Plugins/VirtualProduction/RemoteControl/Source/RemoteControl/Public/RemoteControlBinding.h
@@ -12,6 +12,22 @@
 #include "UObject/UnrealType.h"
 #include "RemoteControlBinding.generated.h"
 
+namespace UE::RemoteControlBinding
+{
+	/**
+	 * @return whether an object is a valid rebinding target.
+	 */
+	bool REMOTECONTROL_API IsValidObjectForRebinding(UObject* InObject, UWorld* PresetWorld);
+	/**
+	 * @return whether an actor is a valid rebinding target.
+	 */
+	bool REMOTECONTROL_API IsValidActorForRebinding(AActor* InActor, UWorld* PresetWorld);
+	/**
+	 * @return whether a subobject is a valid rebinding target.
+	 */
+	bool REMOTECONTROL_API IsValidSubObjectForRebinding(UObject* InComponent, UWorld* PresetWorld);
+}
+
 /**
  * Acts as a bridge between an exposed property/function and an object to act on.
  */
diff --git a/Engine/Plugins/VirtualProduction/RemoteControl/Source/RemoteControlUI/Private/UI/SRCPanelExposedEntity.cpp b/Engine/Plugins/VirtualProduction/RemoteControl/Source/RemoteControlUI/Private/UI/SRCPanelExposedEntity.cpp
--- a/Engine/Plugins/VirtualProduction/RemoteControl/Source/RemoteControlUI/Private/UI/SRCPanelExposedEntity.cpp
+++ b/Engine/Plugins/VirtualProduction/RemoteControl/Source/RemoteControlUI/Private/UI/SRCPanelExposedEntity.cpp
@@ -3,9 +3,11 @@
 #include "SRCPanelExposedEntity.h"
 
 #include "ActorTreeItem.h"
+#include "Components/BillboardComponent.h"
 #include "Commands/RemoteControlCommands.h"
 #include "Editor.h"
 #include "EditorFontGlyphs.h"
+#include "EngineUtils.h"
 #include "Engine/Selection.h"
 #include "Engine/Classes/Components/ActorComponent.h"
 #include "Framework/Application/SlateApplication.h"
@@ -25,6 +27,7 @@
 #include "Styling/RemoteControlStyles.h"
 #include "Styling/SlateIconFinder.h"
 #include "Modules/ModuleManager.h"
+#include "UObject/UObjectIterator.h"
 #include "Widgets/Input/SButton.h"
 #include "Widgets/Input/SComboButton.h"
 #include "Widgets/Input/SCheckBox.h"
@@ -35,6 +38,33 @@
 
 #define LOCTEXT_NAMESPACE "RemoteControlPanel"
 
+namespace RebindingUtils
+{
+	TMap<AActor*, TArray<UObject*>> GetLevelSubObjectsOfClass(UClass* TargetClass , UWorld* PresetWorld)
+	{
+		TMap<AActor*, TArray<UObject*>> ObjectMap;
+
+		for (TActorIterator<AActor> It(PresetWorld, AActor::StaticClass(), EActorIteratorFlags::SkipPendingKill); It; ++It)
+		{
+			if (UE::RemoteControlBinding::IsValidActorForRebinding(*It, PresetWorld))
+			{
+				TArray<UObject*> SubObjects;
+				GetObjectsWithOuter(*It, SubObjects);
+
+				for (UObject* SubObject : SubObjects)
+				{
+					if (SubObject->IsA(TargetClass) && UE::RemoteControlBinding::IsValidSubObjectForRebinding(SubObject, PresetWorld))
+					{
+						ObjectMap.FindOrAdd(*It).Add(SubObject);
+					}
+				}
+			}
+		}
+
+		return ObjectMap;
+	}
+}
+
 TSharedPtr<FRemoteControlEntity> SRCPanelExposedEntity::GetEntity() const
 {
 	if (Preset.IsValid())
@@ -79,6 +109,14 @@ TSharedPtr<SWidget> SRCPanelExposedEntity::GetContextMenu()
 				CreateRebindComponentMenuContent(SubMenuBuilder);
 			}));
 
+	MenuBuilder.AddSubMenu(
+		LOCTEXT("EntityRebindSubObjectSubmenuLabel", "Rebind SubObject"),
+		LOCTEXT("EntityRebindSubObjectSubmenuToolTip", "Pick a subobject to rebind this exposed entity."),
+		FNewMenuDelegate::CreateLambda([this](FMenuBuilder& SubMenuBuilder)
+			{
+				CreateRebindSubObjectMenuContent(SubMenuBuilder);
+			}));
+
 	MenuBuilder.AddSubMenu(
 		LOCTEXT("EntityRebindAllUnderActorSubmenuLabel", "Rebind all properties for this actor"),
 		LOCTEXT("EntityRebindAllUnderActorSubmenuToolTip", "Pick an actor to rebind."),
@@ -154,6 +192,54 @@ void SRCPanelExposedEntity::CreateRebindComponentMenuContent(FMenuBuilder& SubMe
 	}
 }
 
+void SRCPanelExposedEntity::CreateRebindSubObjectMenuContent(FMenuBuilder& SubMenuBuilder)
+{
+	TInlineComponentArray<UActorComponent*> ComponentArray;
+
+	if (TSharedPtr<FRemoteControlEntity> Entity = GetEntity())
+	{
+		constexpr bool bAllowPie = false;
+		TMap<AActor*, TArray<UObject*>> GroupedObjects = RebindingUtils::GetLevelSubObjectsOfClass(Entity->GetSupportedBindingClass(), Preset->GetWorld(bAllowPie));
+		for (const TPair<AActor*, TArray<UObject*>>& ActorsAndSubobjects : GroupedObjects)
+		{
+			SubMenuBuilder.BeginSection(NAME_None, FText::FromString(ActorsAndSubobjects.Key->GetActorLabel()));
+			for (UObject* SubObject : ActorsAndSubobjects.Value)
+			{
+				FString EntryLabel = SubObject->GetName();
+				const FString ToolTip = SubObject->GetPathName();
+					
+				// Special case for NDisplay properties.
+				static const FName NDisplayViewportClassName = "DisplayClusterConfigurationViewport";
+				if (SubObject->GetClass()->GetFName() == NDisplayViewportClassName)
+				{
+					static const FTopLevelAssetPath DisplayClusterConfigurationPath = FTopLevelAssetPath(TEXT("/Script/DisplayClusterConfiguration"), TEXT("DisplayClusterConfigurationClusterNode"));
+					if (UClass* ClusterNodeClass = FindObject<UClass>(DisplayClusterConfigurationPath);
+						UObject * ClusterNodeOuter = SubObject->GetTypedOuter(ClusterNodeClass))
+					{
+						EntryLabel = SubObject->GetPathName(ClusterNodeOuter->GetOuter());
+					}
+				}
+
+				SubMenuBuilder.AddMenuEntry(
+					FText::FromString(EntryLabel),
+					FText::FromString(ToolTip),
+					FSlateIconFinder::FindIconForClass(SubObject->GetClass(), TEXT("SCS.Component")),
+					FUIAction(
+						FExecuteAction::CreateLambda([Entity, SubObject]
+						{
+							if (Entity)
+							{
+								Entity->BindObject(SubObject);
+							}
+						}),
+						FCanExecuteAction())
+				);
+			}
+			SubMenuBuilder.EndSection();
+		}
+	}
+}
+
 TSharedRef<SWidget> SRCPanelExposedEntity::CreateRebindAllPropertiesForActorMenuContent()
 {
 	FSceneOutlinerModule& SceneOutlinerModule = FModuleManager::Get().LoadModuleChecked<FSceneOutlinerModule>("SceneOutliner");
@@ -213,7 +299,7 @@ TSharedRef<SWidget> SRCPanelExposedEntity::CreateRebindMenuContent()
 	Options.Filters->AddFilterPredicate<FActorTreeItem>(FActorTreeItem::FFilterPredicate::CreateSP(this, &SRCPanelExposedEntity::IsActorSelectable));
 	constexpr bool bAllowPIE = false;
 	UWorld* PresetWorld = URemoteControlPreset::GetWorld(Preset.Get(), bAllowPIE);
-
+	
 	return SNew(SBox)
 		.MaxDesiredHeight(400.0f)
 		.WidthOverride(300.0f)
@@ -266,7 +352,7 @@ bool SRCPanelExposedEntity::IsActorSelectable(const AActor* Actor) const
 {
 	if (TSharedPtr<FRemoteControlEntity> Entity = GetEntity())
 	{
-		if (Entity->GetBindings().Num())
+		if (Entity->GetBindings().Num() && Entity->GetBindings()[0].IsValid())
 		{
 			// Don't show what it's already bound to.
 			if (UObject* Component = Entity->GetBindings()[0]->Resolve())
@@ -434,4 +520,5 @@ bool SRCPanelExposedEntity::ShouldUseRebindingContext() const
 	return false;
 }
 
+
 #undef LOCTEXT_NAMESPACE
diff --git a/Engine/Plugins/VirtualProduction/RemoteControl/Source/RemoteControlUI/Private/UI/SRCPanelExposedEntity.h b/Engine/Plugins/VirtualProduction/RemoteControl/Source/RemoteControlUI/Private/UI/SRCPanelExposedEntity.h
--- a/Engine/Plugins/VirtualProduction/RemoteControl/Source/RemoteControlUI/Private/UI/SRCPanelExposedEntity.h
+++ b/Engine/Plugins/VirtualProduction/RemoteControl/Source/RemoteControlUI/Private/UI/SRCPanelExposedEntity.h
@@ -10,11 +10,11 @@
 
 class AActor;
 class FMenuBuilder;
+struct FRCPanelStyle;
 struct FRemoteControlEntity;
 class SInlineEditableTextBlock;
 class SWidget;
 class URemoteControlPreset;
-struct FRCPanelStyle;
 
 /**
  * @Note if you inherit from this struct, you must call SRCPanelExposedEntity::Initialize.
@@ -68,6 +68,9 @@ private:
 	/** Create the content of the rebind component button. */
 	void CreateRebindComponentMenuContent(FMenuBuilder& SubMenuBuilder);
 
+	/** Create the content of the rebind subobject button. */
+	void CreateRebindSubObjectMenuContent(FMenuBuilder& SubMenuBuilder);
+
 	/** Create the content for the menu used to rebind all properties for the actor that owns this entity. */
 	TSharedRef<SWidget> CreateRebindAllPropertiesForActorMenuContent();
 
