UI Style improvements for Remote Control Logic

diff --git a/Engine/Plugins/VirtualProduction/RemoteControl/Source/RemoteControlLogic/Private/Behaviour/Builtin/Path/RCBehaviourSetAssetByPathNode.cpp b/Engine/Plugins/VirtualProduction/RemoteControl/Source/RemoteControlLogic/Private/Behaviour/Builtin/Path/RCBehaviourSetAssetByPathNode.cpp
--- a/Engine/Plugins/VirtualProduction/RemoteControl/Source/RemoteControlLogic/Private/Behaviour/Builtin/Path/RCBehaviourSetAssetByPathNode.cpp
+++ b/Engine/Plugins/VirtualProduction/RemoteControl/Source/RemoteControlLogic/Private/Behaviour/Builtin/Path/RCBehaviourSetAssetByPathNode.cpp
@@ -10,7 +10,7 @@
 
 URCBehaviourSetAssetByPathNode::URCBehaviourSetAssetByPathNode()
 {
-	DisplayName = NSLOCTEXT("Remote Control Behaviour", "Behaviour Name - Set Asset By Path", "Set Asset By Path");
+	DisplayName = NSLOCTEXT("Remote Control Behaviour", "Behaviour Name - Set Asset By Path", "Path");
 	BehaviorDescription = NSLOCTEXT("Remote Control Behaviour", "Behaviour Desc - Set Asset By Path", "Triggers an event which sets an object based on the selected Exposed Entity.");
 }
 
diff --git a/Engine/Plugins/VirtualProduction/RemoteControl/Source/RemoteControlLogic/Private/Behaviour/Builtin/RCBehaviourOnValueChangedNode.cpp b/Engine/Plugins/VirtualProduction/RemoteControl/Source/RemoteControlLogic/Private/Behaviour/Builtin/RCBehaviourOnValueChangedNode.cpp
--- a/Engine/Plugins/VirtualProduction/RemoteControl/Source/RemoteControlLogic/Private/Behaviour/Builtin/RCBehaviourOnValueChangedNode.cpp
+++ b/Engine/Plugins/VirtualProduction/RemoteControl/Source/RemoteControlLogic/Private/Behaviour/Builtin/RCBehaviourOnValueChangedNode.cpp
@@ -6,7 +6,7 @@
 
 URCBehaviourOnValueChangedNode::URCBehaviourOnValueChangedNode()
 {
-	DisplayName = NSLOCTEXT("Remote Control Behaviour", "Behavior Name - On Value Changed", "On Value Changed");
+	DisplayName = NSLOCTEXT("Remote Control Behaviour", "Behavior Name - On Value Changed", "On Modify");
 	BehaviorDescription = NSLOCTEXT("Remote Control Behaviour", "Behavior Desc - On Value Changed", "Triggers an event when the associated property is modified");
 }
 
diff --git a/Engine/Plugins/VirtualProduction/RemoteControl/Source/RemoteControlLogic/Private/Behaviour/Builtin/RangeMap/RCBehaviourRangeMapNode.cpp b/Engine/Plugins/VirtualProduction/RemoteControl/Source/RemoteControlLogic/Private/Behaviour/Builtin/RangeMap/RCBehaviourRangeMapNode.cpp
--- a/Engine/Plugins/VirtualProduction/RemoteControl/Source/RemoteControlLogic/Private/Behaviour/Builtin/RangeMap/RCBehaviourRangeMapNode.cpp
+++ b/Engine/Plugins/VirtualProduction/RemoteControl/Source/RemoteControlLogic/Private/Behaviour/Builtin/RangeMap/RCBehaviourRangeMapNode.cpp
@@ -8,7 +8,7 @@
 
 URCBehaviourRangeMapNode::URCBehaviourRangeMapNode()
 {
-	DisplayName = NSLOCTEXT("Remote Control Behaviour", "Behavior Name - Range Mapping", "Range Mapping");
+	DisplayName = NSLOCTEXT("Remote Control Behaviour", "Behavior Name - Range Mapping", "Range");
 	BehaviorDescription = NSLOCTEXT("Remote Control Behaviour", "Behavior Desc - Range Mapping", "Behaviour lerps property values between a normalized 0-1 step. Users define min/max values and tolerance threshold.");
 }
 
diff --git a/Engine/Plugins/VirtualProduction/RemoteControl/Source/RemoteControlUI/Private/Styling/RemoteControlStyles.cpp b/Engine/Plugins/VirtualProduction/RemoteControl/Source/RemoteControlUI/Private/Styling/RemoteControlStyles.cpp
--- a/Engine/Plugins/VirtualProduction/RemoteControl/Source/RemoteControlUI/Private/Styling/RemoteControlStyles.cpp
+++ b/Engine/Plugins/VirtualProduction/RemoteControl/Source/RemoteControlUI/Private/Styling/RemoteControlStyles.cpp
@@ -9,6 +9,7 @@ FRCPanelStyle::FRCPanelStyle()
 	, PanelPadding(0.f)
 	, PanelTextStyle(FTextBlockStyle::GetDefault())
 	, SectionHeaderTextStyle(FTextBlockStyle::GetDefault())
+	, SplitterHandleSize(4.f)
 {
 }
 
diff --git a/Engine/Plugins/VirtualProduction/RemoteControl/Source/RemoteControlUI/Private/Styling/RemoteControlStyles.h b/Engine/Plugins/VirtualProduction/RemoteControl/Source/RemoteControlUI/Private/Styling/RemoteControlStyles.h
--- a/Engine/Plugins/VirtualProduction/RemoteControl/Source/RemoteControlUI/Private/Styling/RemoteControlStyles.h
+++ b/Engine/Plugins/VirtualProduction/RemoteControl/Source/RemoteControlUI/Private/Styling/RemoteControlStyles.h
@@ -112,4 +112,9 @@ struct FRCPanelStyle : public FSlateWidgetStyle
 	UPROPERTY(EditAnywhere, Category = Appearance)
 		FCheckBoxStyle ToggleButtonStyle;
 	FRCPanelStyle& SetToggleButtonStyle(const FCheckBoxStyle& InToggleButtonStyle) { ToggleButtonStyle = InToggleButtonStyle; return *this; }
+
+	/** Splitter handle size*/
+	UPROPERTY(EditAnywhere, Category = Appearance)
+	float SplitterHandleSize;
+	FRCPanelStyle& SetSplitterHandleSize(const float InSplitterHandleSize) { SplitterHandleSize = InSplitterHandleSize; return *this; }
 };
diff --git a/Engine/Plugins/VirtualProduction/RemoteControl/Source/RemoteControlUI/Private/UI/Action/RCActionModel.cpp b/Engine/Plugins/VirtualProduction/RemoteControl/Source/RemoteControlUI/Private/UI/Action/RCActionModel.cpp
--- a/Engine/Plugins/VirtualProduction/RemoteControl/Source/RemoteControlUI/Private/UI/Action/RCActionModel.cpp
+++ b/Engine/Plugins/VirtualProduction/RemoteControl/Source/RemoteControlUI/Private/UI/Action/RCActionModel.cpp
@@ -80,7 +80,7 @@ FRCActionModel::FRCActionModel(URCAction* InAction, const TSharedPtr<class FRCBe
 {
 	BehaviourItemWeakPtr = InBehaviourItem;
 
-	RCPanelStyle = &FRemoteControlPanelStyle::Get()->GetWidgetStyle<FRCPanelStyle>("RemoteControlPanel.MinorPanel");
+	RCPanelStyle = &FRemoteControlPanelStyle::Get()->GetWidgetStyle<FRCPanelStyle>("RemoteControlPanel.LogicControllersPanel");
 }
 
 TSharedRef<SWidget> FRCActionModel::GetNameWidget() const
diff --git a/Engine/Plugins/VirtualProduction/RemoteControl/Source/RemoteControlUI/Private/UI/Action/SRCActionPanel.cpp b/Engine/Plugins/VirtualProduction/RemoteControl/Source/RemoteControlUI/Private/UI/Action/SRCActionPanel.cpp
--- a/Engine/Plugins/VirtualProduction/RemoteControl/Source/RemoteControlUI/Private/UI/Action/SRCActionPanel.cpp
+++ b/Engine/Plugins/VirtualProduction/RemoteControl/Source/RemoteControlUI/Private/UI/Action/SRCActionPanel.cpp
@@ -95,6 +95,7 @@ void SRCActionPanel::UpdateWrappedWidget(TSharedPtr<FRCBehaviourModel> InBehavio
 		// Action Dock Panel
 		TSharedPtr<SRCMinorPanel> ActionDockPanel = SNew(SRCMinorPanel)
 			.HeaderLabel(LOCTEXT("ActionsLabel", "Actions"))
+			.EnableFooter(true)
 			[
 				ActionPanelList.ToSharedRef()
 			];
@@ -168,36 +169,21 @@ void SRCActionPanel::UpdateWrappedWidget(TSharedPtr<FRCBehaviourModel> InBehavio
 
 		ActionDockPanel->AddHeaderToolbarItem(EToolbar::Left, AddNewActionButton);
 		ActionDockPanel->AddHeaderToolbarItem(EToolbar::Right, AddAllActionsButton);
-		ActionDockPanel->AddHeaderToolbarItem(EToolbar::Right, DeleteSelectedActionButton);
+		ActionDockPanel->AddFooterToolbarItem(EToolbar::Right, DeleteSelectedActionButton);
 
 		// Header Dock Panel
-		TSharedPtr<SRCMinorPanel> HeaderDockPanel = SNew(SRCMinorPanel)
-			.HeaderLabel(LOCTEXT("HeaderDockLabel", ""))
+		TSharedPtr<SRCMinorPanel> BehaviourDetailsPanel = SNew(SRCMinorPanel)
+			.EnableHeader(false)
 			[
 				SAssignNew(BehaviourDetailsWidget, SRCBehaviourDetails, SharedThis(this), InBehaviourItem.ToSharedRef())
 			];
-		
-		// Toggle Behaviour Button
-		const bool bIsChecked = InBehaviourItem->IsBehaviourEnabled();
-		SAssignNew(ToggleBehaviourButton, SCheckBox)
-			.AddMetaData<FTagMetaData>(FTagMetaData(TEXT("Toggle Behavior")))
-			.ToolTipText(LOCTEXT("EditModeTooltip", "Enable/Disable this Behaviour.\nWhen a behaviour is disabled its Actions will not be processed when the Controller value changes"))
-			.HAlign(HAlign_Center)
-			.Style(&RCPanelStyle->ToggleButtonStyle)
-			.ForegroundColor(FSlateColor::UseForeground())
-			.IsChecked(bIsChecked)
-			.OnCheckStateChanged(this, &SRCActionPanel::OnToggleEnableBehaviour);
-
-		RefreshIsBehaviourEnabled(bIsChecked);
-
-		HeaderDockPanel->AddHeaderToolbarItem(EToolbar::Left, ToggleBehaviourButton.ToSharedRef());
 
 		TSharedRef<SRCMajorPanel> ActionsPanel = SNew(SRCMajorPanel)
 			.EnableFooter(false)
 			.EnableHeader(false)
 			.ChildOrientation(Orient_Vertical);
 
-		ActionsPanel->AddPanel(HeaderDockPanel.ToSharedRef(), 0.5f);
+		ActionsPanel->AddPanel(BehaviourDetailsPanel.ToSharedRef(), 0.5f);
 		ActionsPanel->AddPanel(ActionDockPanel.ToSharedRef(), 0.5f);
 
 		WrappedBoxWidget->SetContent(ActionsPanel);
@@ -223,27 +209,18 @@ FReply SRCActionPanel::OnClickOverrideBlueprintButton()
 	return FReply::Handled();
 }
 
-void SRCActionPanel::OnToggleEnableBehaviour(ECheckBoxState State)
-{
-	SetIsBehaviourEnabled(State == ECheckBoxState::Checked);
-}
-
 void SRCActionPanel::SetIsBehaviourEnabled(const bool bIsEnabled)
 {
-	if (TSharedPtr<FRCBehaviourModel> BehaviourItem = GetSelectedBehaviourItem())
+	if (BehaviourDetailsWidget)
 	{
-		BehaviourItem->SetIsBehaviourEnabled(bIsEnabled);
-
-		RefreshIsBehaviourEnabled(bIsEnabled);
+		BehaviourDetailsWidget->SetIsBehaviourEnabled(bIsEnabled);
 	}
 }
 
 void SRCActionPanel::RefreshIsBehaviourEnabled(const bool bIsEnabled)
 {
-	if (ToggleBehaviourButton && BehaviourDetailsWidget && ActionPanelList)
+	if (ActionPanelList)
 	{
-		ToggleBehaviourButton->SetIsChecked(bIsEnabled);
-		BehaviourDetailsWidget->SetEnabled(bIsEnabled);
 		ActionPanelList->SetEnabled(bIsEnabled);
 	}
 }
diff --git a/Engine/Plugins/VirtualProduction/RemoteControl/Source/RemoteControlUI/Private/UI/Action/SRCActionPanel.h b/Engine/Plugins/VirtualProduction/RemoteControl/Source/RemoteControlUI/Private/UI/Action/SRCActionPanel.h
--- a/Engine/Plugins/VirtualProduction/RemoteControl/Source/RemoteControlUI/Private/UI/Action/SRCActionPanel.h
+++ b/Engine/Plugins/VirtualProduction/RemoteControl/Source/RemoteControlUI/Private/UI/Action/SRCActionPanel.h
@@ -74,14 +74,17 @@ public:
 		bAddActionMenuNeedsRefresh = true;
 	}
 
+	/** Set the Enabled state of our parent Behaviour */
+	void SetIsBehaviourEnabled(const bool bIsEnabled);
+
 	/** Fetches the currently selected behaviour item*/
 	TSharedPtr<FRCBehaviourModel> GetSelectedBehaviourItem()
 	{
 		return SelectedBehaviourItemWeakPtr.Pin();
 	}
 
-	/** Set the Enabled state of our parent Behaviour */
-	void SetIsBehaviourEnabled(const bool bIsEnabled);
+	/** Refreshes the UI widgets enabled state depending on whether the parent behaviour is currently enabled */
+	void RefreshIsBehaviourEnabled(const bool bIsEnabled);
 
 protected:
 	/** Warns user before deleting a selected panel item. */
@@ -106,12 +109,6 @@ private:
 	/** Handles click event for Open Behaviour Blueprint button*/
 	FReply OnClickOverrideBlueprintButton();
 
-	/** Handles click event for toggling Enable/Disable behaviour*/
-	void OnToggleEnableBehaviour(ECheckBoxState State);
-
-	/** Refreshes the UI widgets enabled state depending on whether the parent behaviour is currently enabled */
-	void RefreshIsBehaviourEnabled(const bool bIsEnabled);
-
 	/**
 	 * Builds a menu containing the list of all possible Actions
 	 * These are derived from the list of Exposed entities of the Remote Control Preset associated with us.
@@ -151,9 +148,6 @@ private:
 	/** Behaviour specific details widget*/
 	TSharedPtr<SRCBehaviourDetails> BehaviourDetailsWidget;
 
-	/** Toggle Behaviour button*/
-	TSharedPtr<SCheckBox> ToggleBehaviourButton;
-
 	/** Helper widget for behavior details. */
 	static TSharedPtr<SBox> NoneSelectedWidget;
 
diff --git a/Engine/Plugins/VirtualProduction/RemoteControl/Source/RemoteControlUI/Private/UI/Action/SRCBehaviourDetails.cpp b/Engine/Plugins/VirtualProduction/RemoteControl/Source/RemoteControlUI/Private/UI/Action/SRCBehaviourDetails.cpp
--- a/Engine/Plugins/VirtualProduction/RemoteControl/Source/RemoteControlUI/Private/UI/Action/SRCBehaviourDetails.cpp
+++ b/Engine/Plugins/VirtualProduction/RemoteControl/Source/RemoteControlUI/Private/UI/Action/SRCBehaviourDetails.cpp
@@ -23,8 +23,6 @@
 
 #define LOCTEXT_NAMESPACE "SRCBehaviourDetails"
 
-BEGIN_SLATE_FUNCTION_BUILD_OPTIMIZATION
-
 void SRCBehaviourDetails::Construct(const FArguments& InArgs, TSharedRef<SRCActionPanel> InActionPanel, TSharedRef<FRCBehaviourModel> InBehaviourItem)
 {
 	RCPanelStyle = &FRemoteControlPanelStyle::Get()->GetWidgetStyle<FRCPanelStyle>("RemoteControlPanel.MinorPanel");
@@ -33,7 +31,7 @@ void SRCBehaviourDetails::Construct(const FArguments& InArgs, TSharedRef<SRCActi
 	BehaviourItemWeakPtr = InBehaviourItem;
 
 	URCBehaviour* Behaviour = BehaviourItemWeakPtr.Pin()->GetBehaviour();
-	const FText BehaviourDisplayName = Behaviour->GetDisplayName().ToUpper();
+	const FText BehaviourDisplayName = Behaviour->GetDisplayName();
 	const FText BehaviourDescription = Behaviour->GetBehaviorDescription();
 
 	const FSlateFontInfo& FontBehaviorDesc = FRemoteControlPanelStyle::Get()->GetFontStyle("RemoteControlPanel.Behaviours.BehaviorDescription");
@@ -46,7 +44,9 @@ void SRCBehaviourDetails::Construct(const FArguments& InArgs, TSharedRef<SRCActi
 		TypeDisplayName = FName::NameToDisplayString(UE::RCUIHelpers::GetFieldClassDisplayName(Behaviour->ControllerWeakPtr->GetProperty()).ToString(), false);
 	}
 
-	TSharedRef<SWidget> BehaviourDetailsWidget = InBehaviourItem->GetBehaviourDetailsWidget();
+	BehaviourDetailsWidget = InBehaviourItem->GetBehaviourDetailsWidget();
+
+	const bool bIsChecked = InBehaviourItem->IsBehaviourEnabled();
 	
 	ChildSlot
 		.Padding(RCPanelStyle->PanelPadding)
@@ -55,7 +55,7 @@ void SRCBehaviourDetails::Construct(const FArguments& InArgs, TSharedRef<SRCActi
 
 			+ SVerticalBox::Slot()
 			.HAlign(HAlign_Fill)
-			.Padding(4.f)
+			.Padding(6.f)
 			.AutoHeight()
 			[
 				SNew(SHorizontalBox)
@@ -67,43 +67,31 @@ void SRCBehaviourDetails::Construct(const FArguments& InArgs, TSharedRef<SRCActi
 				.FillWidth(1.f)
 				.Padding(4.f)
 				[
-					SNew(STextBlock)
+					SAssignNew(BehaviourTitleWidget, STextBlock)
 					.ColorAndOpacity(FLinearColor::White)
 					.Text(BehaviourDisplayName)
 					.TextStyle(&RCPanelStyle->SectionHeaderTextStyle)
+					.ToolTipText(BehaviourDescription)
 				]
 
+				// Toggle Behaviour Button
 				+ SHorizontalBox::Slot()
 				.VAlign(VAlign_Center)
 				.AutoWidth()
 				.Padding(4.f, 0.f)
 				[
-					SNew(SBorder)
-					.BorderImage(FAppStyle::Get().GetBrush("SettingsEditor.CheckoutWarningBorder"))
-					.Padding(FMargin(8.f, 4.f))
-					.BorderBackgroundColor(TypeColor)
-					[
-						SNew(STextBlock)
-						.ColorAndOpacity(FLinearColor::Black)
-						.Text(FText::FromString(TypeDisplayName))
-						.TextStyle(&RCPanelStyle->PanelTextStyle)
-					]
+					SNew(SCheckBox)
+					.AddMetaData<FTagMetaData>(FTagMetaData(TEXT("Toggle Behavior")))
+					.ToolTipText(LOCTEXT("EditModeTooltip", "Enable/Disable this Behaviour.\nWhen a behaviour is disabled its Actions will not be processed when the Controller value changes"))
+					.HAlign(HAlign_Center)
+					.Style(&RCPanelStyle->ToggleButtonStyle)
+					.ForegroundColor(FSlateColor::UseForeground())
+					.IsChecked_Lambda([this]() { return BehaviourItemWeakPtr.IsValid() && BehaviourItemWeakPtr.Pin()->IsBehaviourEnabled() ? ECheckBoxState::Checked : ECheckBoxState::Unchecked; })
+					.OnCheckStateChanged(this, &SRCBehaviourDetails::OnToggleEnableBehaviour)
 				]
 			]
-
-			+ SVerticalBox::Slot()
-			.Padding(8.f, 4.f)
-			.VAlign(VAlign_Top)
-			.FillHeight(1.f)
-			[
-				SNew(STextBlock)
-				.ColorAndOpacity(FSlateColor::UseForeground())
-				.Text(BehaviourDescription)
-				.TextStyle(&RCPanelStyle->PanelTextStyle)
-				.AutoWrapText(true) // note: use of autowidth on this slot prevents wrapping from working!
-			]
 			
-			// Border (separating Header and "If Value Is" panel)
+			// Border (separating Header and Behaviour specific details panel)
 			+ SVerticalBox::Slot()
 			.Padding(2.f, 4.f)
 			.AutoHeight()
@@ -112,7 +100,7 @@ void SRCBehaviourDetails::Construct(const FArguments& InArgs, TSharedRef<SRCActi
 				.SeparatorImage(FAppStyle::Get().GetBrush("Separator"))
 				.Thickness(2.f)
 				.Orientation(EOrientation::Orient_Horizontal)
-				.Visibility_Lambda([BehaviourDetailsWidget]() { return BehaviourDetailsWidget != SNullWidget::NullWidget ? EVisibility::Visible : EVisibility::Collapsed; })
+				.Visibility_Lambda([this]() { return this->BehaviourDetailsWidget.IsValid() ? EVisibility::Visible : EVisibility::Collapsed; })
 			]
 
 			// Behaviour Specific Details Panel
@@ -120,7 +108,7 @@ void SRCBehaviourDetails::Construct(const FArguments& InArgs, TSharedRef<SRCActi
 			.Padding(8.f, 4.f)
 			.AutoHeight()
 			[
-				BehaviourDetailsWidget
+				BehaviourDetailsWidget.ToSharedRef()
 			]
 
 			// Spacer to fill the gap.
@@ -131,8 +119,37 @@ void SRCBehaviourDetails::Construct(const FArguments& InArgs, TSharedRef<SRCActi
 				SNew(SSpacer)
 			]
 		];
+
+	RefreshIsBehaviourEnabled(bIsChecked);
+}
+
+void SRCBehaviourDetails::OnToggleEnableBehaviour(ECheckBoxState State)
+{
+	SetIsBehaviourEnabled(State == ECheckBoxState::Checked);
 }
 
-END_SLATE_FUNCTION_BUILD_OPTIMIZATION
+void SRCBehaviourDetails::SetIsBehaviourEnabled(const bool bIsEnabled)
+{
+	if (TSharedPtr<FRCBehaviourModel> BehaviourItem = BehaviourItemWeakPtr.Pin())
+	{
+		BehaviourItem->SetIsBehaviourEnabled(bIsEnabled);
+
+		RefreshIsBehaviourEnabled(bIsEnabled);
+	}
+}
+
+void SRCBehaviourDetails::RefreshIsBehaviourEnabled(const bool bIsEnabled)
+{
+	if(BehaviourDetailsWidget && BehaviourTitleWidget)
+	{
+		BehaviourDetailsWidget->SetEnabled(bIsEnabled);
+		BehaviourTitleWidget->SetEnabled(bIsEnabled);
+	}
+
+	if (TSharedPtr<SRCActionPanel> ActionPanel = ActionPanelWeakPtr.Pin())
+	{
+		ActionPanel->RefreshIsBehaviourEnabled(bIsEnabled);
+	}
+}
 
 #undef LOCTEXT_NAMESPACE
\ No newline at end of file
diff --git a/Engine/Plugins/VirtualProduction/RemoteControl/Source/RemoteControlUI/Private/UI/Action/SRCBehaviourDetails.h b/Engine/Plugins/VirtualProduction/RemoteControl/Source/RemoteControlUI/Private/UI/Action/SRCBehaviourDetails.h
--- a/Engine/Plugins/VirtualProduction/RemoteControl/Source/RemoteControlUI/Private/UI/Action/SRCBehaviourDetails.h
+++ b/Engine/Plugins/VirtualProduction/RemoteControl/Source/RemoteControlUI/Private/UI/Action/SRCBehaviourDetails.h
@@ -6,6 +6,7 @@
 
 struct FRCPanelStyle;
 class SRCActionPanel;
+class STextBlock;
 class FRCBehaviourModel;
 
 /**
@@ -23,7 +24,16 @@ public:
 	/** Constructs this widget with InArgs */
 	void Construct(const FArguments& InArgs, TSharedRef<SRCActionPanel> InActionPanel, TSharedRef<FRCBehaviourModel> InBehaviourItem);
 
+	/** Set the Enabled state of our parent Behaviour */
+	void SetIsBehaviourEnabled(const bool bIsEnabled);
+
 private:
+	/** Refreshes the UI widgets enabled state depending on whether the parent behaviour is currently enabled */
+	void RefreshIsBehaviourEnabled(const bool bIsEnabled);
+
+	/** Handles click event for toggling Enable/Disable behaviour*/
+	void OnToggleEnableBehaviour(ECheckBoxState State);
+
 	/** The parent Action panel UI widget holding this Header*/
 	TWeakPtr<SRCActionPanel> ActionPanelWeakPtr;
 	
@@ -32,4 +42,10 @@ private:
 
 	/** Panel Style reference. */
 	const FRCPanelStyle* RCPanelStyle;
+
+	/** Behaviour specific details panel*/
+	TSharedPtr<SWidget> BehaviourDetailsWidget;
+
+	/** Behaviour Title text widget*/
+	TSharedPtr<STextBlock> BehaviourTitleWidget;
 };
diff --git a/Engine/Plugins/VirtualProduction/RemoteControl/Source/RemoteControlUI/Private/UI/Behaviour/SRCBehaviourPanel.cpp b/Engine/Plugins/VirtualProduction/RemoteControl/Source/RemoteControlUI/Private/UI/Behaviour/SRCBehaviourPanel.cpp
--- a/Engine/Plugins/VirtualProduction/RemoteControl/Source/RemoteControlUI/Private/UI/Behaviour/SRCBehaviourPanel.cpp
+++ b/Engine/Plugins/VirtualProduction/RemoteControl/Source/RemoteControlUI/Private/UI/Behaviour/SRCBehaviourPanel.cpp
@@ -28,18 +28,17 @@
 #define LOCTEXT_NAMESPACE "SRCBehaviourPanel"
 
 TSharedPtr<SBox> SRCBehaviourPanel::NoneSelectedWidget = SNew(SBox)
-			.Padding(0.f)
+			.Padding(10.f)
 			.VAlign(VAlign_Center)
 			.HAlign(HAlign_Center)
 			[
 				SNew(STextBlock)
-				.Text(LOCTEXT("NoneSelected", "Select a controller to view its behaviors."))
+				.Text(LOCTEXT("NoneSelected", "Select a controller\nto view its behaviors."))
 				.TextStyle(&FAppStyle::GetWidgetStyle<FTextBlockStyle>("NormalText"))
 				.Justification(ETextJustify::Center)
+				.AutoWrapText(true)
 			];
 
-BEGIN_SLATE_FUNCTION_BUILD_OPTIMIZATION
-
 void SRCBehaviourPanel::Construct(const FArguments& InArgs, const TSharedRef<SRemoteControlPanel>& InPanel)
 {
 	SRCLogicPanelBase::Construct(SRCLogicPanelBase::FArguments(), InPanel);
@@ -59,8 +58,6 @@ void SRCBehaviourPanel::Construct(const FArguments& InArgs, const TSharedRef<SRe
 	InPanel->OnControllerSelectionChanged.AddSP(this, &SRCBehaviourPanel::OnControllerSelectionChanged);
 }
 
-END_SLATE_FUNCTION_BUILD_OPTIMIZATION
-
 void SRCBehaviourPanel::Shutdown()
 {
 	NoneSelectedWidget.Reset();
@@ -81,6 +78,7 @@ void SRCBehaviourPanel::UpdateWrappedWidget(TSharedPtr<FRCControllerModel> InCon
 		// Behaviour Dock Panel
 		TSharedPtr<SRCMinorPanel> BehaviourDockPanel = SNew(SRCMinorPanel)
 			.HeaderLabel(LOCTEXT("BehavioursLabel", "Behavior"))
+			.EnableFooter(true)
 			[
 				SAssignNew(BehaviourPanelList, SRCBehaviourPanelList, SharedThis(this), InControllerItem, RemoteControlPanel.ToSharedRef())
 			];
@@ -133,7 +131,7 @@ void SRCBehaviourPanel::UpdateWrappedWidget(TSharedPtr<FRCControllerModel> InCon
 			];
 
 		BehaviourDockPanel->AddHeaderToolbarItem(EToolbar::Left, AddNewBehaviourButton);
-		BehaviourDockPanel->AddHeaderToolbarItem(EToolbar::Right, DeleteSelectedBehaviourButton);
+		BehaviourDockPanel->AddFooterToolbarItem(EToolbar::Right, DeleteSelectedBehaviourButton);
 
 		WrappedBoxWidget->SetContent(BehaviourDockPanel.ToSharedRef());
 	}
diff --git a/Engine/Plugins/VirtualProduction/RemoteControl/Source/RemoteControlUI/Private/UI/Controller/RCControllerModel.cpp b/Engine/Plugins/VirtualProduction/RemoteControl/Source/RemoteControlUI/Private/UI/Controller/RCControllerModel.cpp
--- a/Engine/Plugins/VirtualProduction/RemoteControl/Source/RemoteControlUI/Private/UI/Controller/RCControllerModel.cpp
+++ b/Engine/Plugins/VirtualProduction/RemoteControl/Source/RemoteControlUI/Private/UI/Controller/RCControllerModel.cpp
@@ -33,11 +33,10 @@ TSharedRef<SWidget> FRCControllerModel::GetWidget() const
 	const FNodeWidgets NodeWidgets = DetailTreeNodeWeakPtr.Pin()->CreateNodeWidgets();
 
 	const TSharedRef<SHorizontalBox> FieldWidget = SNew(SHorizontalBox);
-	
 	if (NodeWidgets.ValueWidget)
 	{
 		FieldWidget->AddSlot()
-			.Padding(FMargin(3.0f, 2.0f))
+			.Padding(FMargin(10.0f, 2.0f))
 			.HAlign(HAlign_Right)
 			.AutoWidth()
 			[
@@ -47,7 +46,7 @@ TSharedRef<SWidget> FRCControllerModel::GetWidget() const
 	else if (NodeWidgets.WholeRowWidget)
 	{
 		FieldWidget->AddSlot()
-			.Padding(FMargin(3.0f, 2.0f))
+			.Padding(FMargin(10.0f, 2.0f))
 			.AutoWidth()
 			[
 				NodeWidgets.WholeRowWidget.ToSharedRef()
@@ -59,7 +58,10 @@ TSharedRef<SWidget> FRCControllerModel::GetWidget() const
 
 TSharedRef<SWidget> FRCControllerModel::GetNameWidget() const
 {
-	return ControllerNameTextBox.ToSharedRef();
+	return SNew(SBox).Padding(10.f, 2.f)
+		[
+			ControllerNameTextBox.ToSharedRef()
+		];
 }
 
 URCVirtualPropertyBase* FRCControllerModel::GetVirtualProperty() const
diff --git a/Engine/Plugins/VirtualProduction/RemoteControl/Source/RemoteControlUI/Private/UI/Controller/SRCControllerPanel.cpp b/Engine/Plugins/VirtualProduction/RemoteControl/Source/RemoteControlUI/Private/UI/Controller/SRCControllerPanel.cpp
--- a/Engine/Plugins/VirtualProduction/RemoteControl/Source/RemoteControlUI/Private/UI/Controller/SRCControllerPanel.cpp
+++ b/Engine/Plugins/VirtualProduction/RemoteControl/Source/RemoteControlUI/Private/UI/Controller/SRCControllerPanel.cpp
@@ -40,6 +40,7 @@ void SRCControllerPanel::Construct(const FArguments& InArgs, const TSharedRef<SR
 	// Controller Dock Panel
 	TSharedPtr<SRCMinorPanel> ControllerDockPanel = SNew(SRCMinorPanel)
 		.HeaderLabel(LOCTEXT("ControllersLabel", "Controller"))
+		.EnableFooter(true)
 		[
 			SAssignNew(ControllerPanelList, SRCControllerPanelList, SharedThis(this), InPanel)
 		];
@@ -93,7 +94,7 @@ void SRCControllerPanel::Construct(const FArguments& InArgs, const TSharedRef<SR
 		];
 
 	ControllerDockPanel->AddHeaderToolbarItem(EToolbar::Left, AddNewControllerButton);
-	ControllerDockPanel->AddHeaderToolbarItem(EToolbar::Right, DeleteSelectedControllerButton);
+	ControllerDockPanel->AddFooterToolbarItem(EToolbar::Right, DeleteSelectedControllerButton);
 
 	ChildSlot
 		.Padding(RCPanelStyle->PanelPadding)
diff --git a/Engine/Plugins/VirtualProduction/RemoteControl/Source/RemoteControlUI/Private/UI/Controller/SRCControllerPanelList.cpp b/Engine/Plugins/VirtualProduction/RemoteControl/Source/RemoteControlUI/Private/UI/Controller/SRCControllerPanelList.cpp
--- a/Engine/Plugins/VirtualProduction/RemoteControl/Source/RemoteControlUI/Private/UI/Controller/SRCControllerPanelList.cpp
+++ b/Engine/Plugins/VirtualProduction/RemoteControl/Source/RemoteControlUI/Private/UI/Controller/SRCControllerPanelList.cpp
@@ -264,7 +264,7 @@ void SRCControllerPanelList::Construct(const FArguments& InArgs, const TSharedRe
 	
 	ControllerPanelWeakPtr = InControllerPanel;
 	
-	RCPanelStyle = &FRemoteControlPanelStyle::Get()->GetWidgetStyle<FRCPanelStyle>("RemoteControlPanel.MinorPanel");
+	RCPanelStyle = &FRemoteControlPanelStyle::Get()->GetWidgetStyle<FRCPanelStyle>("RemoteControlPanel.LogicControllersPanel");
 
 	ListView = SNew(SListView<TSharedPtr<FRCControllerModel>>)
 		.ListItemsSource(&ControllerItems)
@@ -278,22 +278,22 @@ void SRCControllerPanelList::Construct(const FArguments& InArgs, const TSharedRe
 
 			+ SHeaderRow::Column(UE::RCControllerPanelList::Columns::TypeColor)
 			.DefaultLabel(LOCTEXT("ControllerNameColumnName", ""))
-			.FillWidth(0.03f)
+			.FixedWidth(30)
 			.HeaderContentPadding(RCPanelStyle->HeaderRowPadding)
 
 			+ SHeaderRow::Column(UE::RCControllerPanelList::Columns::DragHandle)
 			.DefaultLabel(FText::GetEmpty())
-			.FillWidth(0.07f)
+			.FixedWidth(30)
 			.HeaderContentPadding(RCPanelStyle->HeaderRowPadding)
 
 			+ SHeaderRow::Column(UE::RCControllerPanelList::Columns::Name)
 			.DefaultLabel(LOCTEXT("ControllerNameColumnName", "Name"))
-			.FillWidth(0.3f)
+			.FillWidth(0.15f)
 			.HeaderContentPadding(RCPanelStyle->HeaderRowPadding)
 
 			+ SHeaderRow::Column(UE::RCControllerPanelList::Columns::Value)
 			.DefaultLabel(LOCTEXT("ControllerValueColumnName", "Input"))
-			.FillWidth(0.6f)
+			.FillWidth(0.75f)
 			.HeaderContentPadding(RCPanelStyle->HeaderRowPadding)
 		);
 
@@ -392,7 +392,7 @@ TSharedRef<ITableRow> SRCControllerPanelList::OnGenerateWidgetForList(TSharedPtr
 
 	return SNew(SControllerRowType, OwnerTable, InItem.ToSharedRef(), SharedThis(this))
 		.Style(&RCPanelStyle->TableRowStyle)
-		.Padding(FMargin(3.f));
+		.Padding(FMargin(4.5f));
 }
 
 void SRCControllerPanelList::OnTreeSelectionChanged(TSharedPtr<FRCControllerModel> InItem, ESelectInfo::Type)
diff --git a/Engine/Plugins/VirtualProduction/RemoteControl/Source/RemoteControlUI/Private/UI/Panels/SRCDockPanel.cpp b/Engine/Plugins/VirtualProduction/RemoteControl/Source/RemoteControlUI/Private/UI/Panels/SRCDockPanel.cpp
--- a/Engine/Plugins/VirtualProduction/RemoteControl/Source/RemoteControlUI/Private/UI/Panels/SRCDockPanel.cpp
+++ b/Engine/Plugins/VirtualProduction/RemoteControl/Source/RemoteControlUI/Private/UI/Panels/SRCDockPanel.cpp
@@ -22,9 +22,14 @@ void SRCDockPanel::PrivateRegisterAttributes(FSlateAttributeInitializer& Attribu
 {
 }
 
+static const FRCPanelStyle* GetPanelStyle()
+{
+	return &FRemoteControlPanelStyle::Get()->GetWidgetStyle<FRCPanelStyle>("RemoteControlPanel.MinorPanel");
+}
+
 void SRCDockPanel::Construct(const FArguments& InArgs)
 {
-	RCPanelStyle = &FRemoteControlPanelStyle::Get()->GetWidgetStyle<FRCPanelStyle>("RemoteControlPanel.MinorPanel");
+	RCPanelStyle = GetPanelStyle();
 
 	bIsFooterEnabled = false;
 	bIsHeaderEnabled = false;
@@ -53,6 +58,9 @@ BEGIN_SLATE_FUNCTION_BUILD_OPTIMIZATION
 
 void SRCMajorPanel::Construct(const SRCMajorPanel::FArguments& InArgs)
 {
+	RCPanelStyle = GetPanelStyle();
+	const float SplitterHandleSize = RCPanelStyle->SplitterHandleSize;
+
 	SRCDockPanel::Construct(SRCDockPanel::FArguments()
 		[
 			SAssignNew(ContentPanel, SSplitter)
@@ -64,6 +72,8 @@ void SRCMajorPanel::Construct(const SRCMajorPanel::FArguments& InArgs)
 			[
 				SAssignNew(Children, SSplitter)
 				.Orientation(InArgs._ChildOrientation)
+				.HitDetectionSplitterHandleSize(SplitterHandleSize)
+				.PhysicalSplitterHandleSize(SplitterHandleSize)
 			]
 		]
 	);
@@ -132,12 +142,16 @@ void SRCMajorPanel::Construct(const SRCMajorPanel::FArguments& InArgs)
 
 END_SLATE_FUNCTION_BUILD_OPTIMIZATION
 
-void SRCMajorPanel::AddPanel(TSharedRef<SWidget> InContent, const float InDesiredSize)
+void SRCMajorPanel::AddPanel(TSharedRef<SWidget> InContent, const float InDesiredSize, const bool bResizable/* = true */)
 {
+	const SSplitter::ESizeRule SizeRule = InDesiredSize == 0.f ? SSplitter::ESizeRule::SizeToContent : SSplitter::ESizeRule::FractionOfParent;
+
 	if (Children.IsValid())
 	{
 		Children->AddSlot()
+			.SizeRule(SizeRule)
 			.Value(InDesiredSize)
+			.Resizable(bResizable)
 			[
 				InContent
 			];
diff --git a/Engine/Plugins/VirtualProduction/RemoteControl/Source/RemoteControlUI/Private/UI/Panels/SRCDockPanel.h b/Engine/Plugins/VirtualProduction/RemoteControl/Source/RemoteControlUI/Private/UI/Panels/SRCDockPanel.h
--- a/Engine/Plugins/VirtualProduction/RemoteControl/Source/RemoteControlUI/Private/UI/Panels/SRCDockPanel.h
+++ b/Engine/Plugins/VirtualProduction/RemoteControl/Source/RemoteControlUI/Private/UI/Panels/SRCDockPanel.h
@@ -91,7 +91,7 @@ public:
 	 *
 	 * @param InPanel The panel widget.
 	 */
-	virtual void AddPanel(TSharedRef<SWidget> InPanel, const float InDesiredSize);
+	virtual void AddPanel(TSharedRef<SWidget> InPanel, const float InDesiredSize, const bool bResizable = true);
 
 	/**
 	 * Gets the content of this panel
diff --git a/Engine/Plugins/VirtualProduction/RemoteControl/Source/RemoteControlUI/Private/UI/RemoteControlPanelStyle.cpp b/Engine/Plugins/VirtualProduction/RemoteControl/Source/RemoteControlUI/Private/UI/RemoteControlPanelStyle.cpp
--- a/Engine/Plugins/VirtualProduction/RemoteControl/Source/RemoteControlUI/Private/UI/RemoteControlPanelStyle.cpp
+++ b/Engine/Plugins/VirtualProduction/RemoteControl/Source/RemoteControlUI/Private/UI/RemoteControlPanelStyle.cpp
@@ -164,7 +164,6 @@ void FRemoteControlPanelStyle::SetupPanelStyles(TSharedRef<FSlateStyleSet> InSty
 	FHeaderRowStyle HeaderRowStyle = AppStyle.GetWidgetStyle<FHeaderRowStyle>("PropertyTable.HeaderRow");
 	HeaderRowStyle.BackgroundBrush = FSlateColorBrush(FStyleColors::Background);
 	HeaderRowStyle.HorizontalSeparatorBrush = FSlateColorBrush(FStyleColors::Background);
-	HeaderRowStyle.HorizontalSeparatorThickness = 2.f;
 
 	// Table Row Style
 	FTableRowStyle TableRowStyle = AppStyle.GetWidgetStyle<FTableRowStyle>("DetailsView.TreeView.TableRow");
@@ -217,6 +216,18 @@ void FRemoteControlPanelStyle::SetupPanelStyles(TSharedRef<FSlateStyleSet> InSty
 
 	StyleSet->Set("RemoteControlPanel.MinorPanel", RCMinorPanelStyle);
 
+	FRCPanelStyle RCLogicControllersPanelStyle = RCMinorPanelStyle;
+	RCLogicControllersPanelStyle.SetHeaderRowPadding(FMargin(15.f, 2.f));
+
+	// Table Row Style
+	FTableRowStyle ControllersTableRowStyle = RCLogicControllersPanelStyle.TableRowStyle;
+	ControllersTableRowStyle.SetEvenRowBackgroundBrush(FSlateColorBrush(FStyleColors::WindowBorder));
+	ControllersTableRowStyle.SetOddRowBackgroundBrush(FSlateColorBrush(FStyleColors::WindowBorder));
+
+	RCLogicControllersPanelStyle.SetTableRowStyle(ControllersTableRowStyle);
+
+	StyleSet->Set("RemoteControlPanel.LogicControllersPanel", RCLogicControllersPanelStyle);
+
 	// Behaviour Panel style
 	FRCPanelStyle RCBehaviourPanelStyle = RCMinorPanelStyle;
 
diff --git a/Engine/Plugins/VirtualProduction/RemoteControl/Source/RemoteControlUI/Private/UI/SRemoteControlPanel.cpp b/Engine/Plugins/VirtualProduction/RemoteControl/Source/RemoteControlUI/Private/UI/SRemoteControlPanel.cpp
--- a/Engine/Plugins/VirtualProduction/RemoteControl/Source/RemoteControlUI/Private/UI/SRemoteControlPanel.cpp
+++ b/Engine/Plugins/VirtualProduction/RemoteControl/Source/RemoteControlUI/Private/UI/SRemoteControlPanel.cpp
@@ -537,7 +537,7 @@ void SRemoteControlPanel::Construct(const FArguments& InArgs, URemoteControlPres
 	// Controllers with Behaviours panel
 	TSharedRef<SRCMajorPanel> ControllersAndBehavioursPanel = SNew(SRCMajorPanel)
 		.EnableHeader(false)
-		.EnableFooter(true);
+		.EnableFooter(false);
 
 	ControllerPanel = SNew(SRCControllerPanel, SharedThis(this))
 		.LiveMode_Lambda([this]() { return bIsInLiveMode; })
@@ -546,33 +546,23 @@ void SRemoteControlPanel::Construct(const FArguments& InArgs, URemoteControlPres
 	BehaviourPanel = SNew(SRCBehaviourPanel, SharedThis(this))
 		.Visibility_Lambda([this] { return !bIsInLiveMode && bIsLogicPanelEnabled && (ActivePanel == ERCPanels::RCP_Properties || ActivePanel == ERCPanels::RCP_None) ? EVisibility::Visible : EVisibility::Collapsed; });
 
-	ControllersAndBehavioursPanel->AddPanel(ControllerPanel.ToSharedRef(), 0.5f);
-		
-	ControllersAndBehavioursPanel->AddPanel(BehaviourPanel.ToSharedRef(), 0.5f);
-
-	// TODO : Replace this with actual Favorites Panel.
-	TSharedRef<SWidget> FavoritesItem = SNew(SBox)
-		.AddMetaData<FTagMetaData>(FTagMetaData(TEXT("Favorites")))
-		.HeightOverride(30.f)
-		.VAlign(VAlign_Center)
-		[
-			SNew(STextBlock)
-			.Text(LOCTEXT("FavoritesLabel", "Favorites"))
-			.TextStyle(&RCPanelStyle->PanelTextStyle)
-		];
-
-	ControllersAndBehavioursPanel->AddFooterToolItem(EToolbar::Left, FavoritesItem);
+	ControllersAndBehavioursPanel->AddPanel(ControllerPanel.ToSharedRef(), 0.8f);
+	
+	constexpr bool bResizable = false;
+	ControllersAndBehavioursPanel->AddPanel(BehaviourPanel.ToSharedRef(), 0.f, bResizable);
 
 	// Actions Panel.
 	ActionPanel = SNew(SRCActionPanel, SharedThis(this))
 		.Visibility_Lambda([this] { return !bIsInLiveMode && bIsLogicPanelEnabled && (ActivePanel == ERCPanels::RCP_Properties || ActivePanel == ERCPanels::RCP_None) ? EVisibility::Visible : EVisibility::Collapsed; });
 
-	LogicPanel->AddPanel(ControllersAndBehavioursPanel, 0.67f);
-	LogicPanel->AddPanel(ActionPanel.ToSharedRef(), 0.33f);
+	LogicPanel->AddPanel(ControllersAndBehavioursPanel, 0.6f);
+	LogicPanel->AddPanel(ActionPanel.ToSharedRef(), 0.4f);
 
 	// Make 2 Columns with Panel Drawer + Main Panel
 	TSharedRef<SSplitter> ContentPanel = SNew(SSplitter)
 		.Orientation(Orient_Horizontal)
+		.HitDetectionSplitterHandleSize(RCPanelStyle->SplitterHandleSize)
+		.PhysicalSplitterHandleSize(RCPanelStyle->SplitterHandleSize)
 	
 		// Panel Drawer
 		+SSplitter::Slot()
