Remote Control Logic - Set Asset By Path - UI Button Order Change

Change order of UI Buttons within SetAssetByPath Logic Behaviour based on feedback

* Auto-Insert Array at Index 0 to prevent user confusion.
* Add Conditional to prevent deletion of Array.

diff --git a/Engine/Plugins/VirtualProduction/RemoteControl/Source/RemoteControlLogic/Private/Behaviour/Builtin/Path/RCSetAssetByPathBehaviour.cpp b/Engine/Plugins/VirtualProduction/RemoteControl/Source/RemoteControlLogic/Private/Behaviour/Builtin/Path/RCSetAssetByPathBehaviour.cpp
--- a/Engine/Plugins/VirtualProduction/RemoteControl/Source/RemoteControlLogic/Private/Behaviour/Builtin/Path/RCSetAssetByPathBehaviour.cpp
+++ b/Engine/Plugins/VirtualProduction/RemoteControl/Source/RemoteControlLogic/Private/Behaviour/Builtin/Path/RCSetAssetByPathBehaviour.cpp
@@ -46,6 +46,24 @@ bool URCSetAssetByPathBehaviour::SetAssetByPath(const FString& AssetPath, const
 	{
 		return false;
 	}
+	
+	URemoteControlPreset* Preset = Controller->PresetWeakPtr.Get();
+	if (!Preset)
+	{
+		return false;
+	}
+	
+	const TSharedPtr<FRemoteControlProperty> RemoteControlProperty = Preset->GetExposedEntity<FRemoteControlProperty>(TargetEntityId).Pin();
+	if (!RemoteControlProperty)
+	{
+		return false;
+	}
+
+	FProperty* Property = RemoteControlProperty->GetProperty();
+	if (!Property)
+	{
+		return false;
+	}
 
 	if (!bInternal)
 	{
@@ -53,8 +71,19 @@ bool URCSetAssetByPathBehaviour::SetAssetByPath(const FString& AssetPath, const
 		{
 			return FText::FromString(FString("Path Behaviour attempts to set Asset to %s") + *AssetPath + *ControllerString);
 		});
+
+#if WITH_EDITOR
+		FEditPropertyChain PropertyChain;
+		PropertyChain.AddHead(Property);
+		RemoteControlProperty->GetBoundObject()->PreEditChange(PropertyChain);
+#endif
 		if (SetExternalAsset(AssetPath + ControllerString))
 		{
+#if WITH_EDITOR
+			FPropertyChangedEvent ChangedEvent(Property, EPropertyChangeType::ValueSet);
+			RemoteControlProperty->GetBoundObject()->PostEditChangeProperty(ChangedEvent);
+#endif
+			
 			return true;
 		}
 
@@ -69,7 +98,7 @@ bool URCSetAssetByPathBehaviour::SetAssetByPath(const FString& AssetPath, const
 
 		return false;
 	}
-	
+
 	FSoftObjectPath MainObjectRef(AssetPath + ControllerString);
 	
 	if (UObject* Object = MainObjectRef.TryLoad())
@@ -78,7 +107,23 @@ bool URCSetAssetByPathBehaviour::SetAssetByPath(const FString& AssetPath, const
 		{
 			return FText::FromString(FString("Path Behaviour attempts to set Asset to %s") + *MainObjectRef.GetAssetPathString());
 		});
-		return SetInternalAsset(Object);
+
+#if WITH_EDITOR
+		FEditPropertyChain PropertyChain;
+		PropertyChain.AddHead(Property);
+		RemoteControlProperty->GetBoundObject()->PreEditChange(PropertyChain);
+#endif
+		if (SetInternalAsset(Object))
+		{
+#if WITH_EDITOR
+			FPropertyChangedEvent ChangedEvent(Property, EPropertyChangeType::ValueSet);
+			RemoteControlProperty->GetBoundObject()->PostEditChangeProperty(ChangedEvent);
+#endif
+
+			return true;
+		}
+
+		return false;
 	}
 	
 	if (!DefaultString.IsEmpty())
@@ -89,7 +134,21 @@ bool URCSetAssetByPathBehaviour::SetAssetByPath(const FString& AssetPath, const
 		{
 			return FText::FromString(FString("Path Behaviour attempts to set Asset to %s") + *DefaultObjectRef.GetAssetPathString());
 		});
-		return SetInternalAsset(DefaultObjectRef.TryLoad());
+
+#if WITH_EDITOR
+		FEditPropertyChain PropertyChain;
+		PropertyChain.AddHead(Property);
+		RemoteControlProperty->GetBoundObject()->PreEditChange(PropertyChain);
+#endif
+		if (SetInternalAsset(DefaultObjectRef.TryLoad()))
+		{
+#if WITH_EDITOR
+			FPropertyChangedEvent ChangedEvent(Property, EPropertyChangeType::ValueSet);
+			RemoteControlProperty->GetBoundObject()->PostEditChangeProperty(ChangedEvent);
+#endif
+			
+			return true;
+		}
 	}
 	
 	return false;
@@ -418,3 +477,12 @@ void URCSetAssetByPathBehaviour::UpdateTargetEntity()
 	
 	TargetEntity = ExposedEntity.Pin();
 }
+
+void URCSetAssetByPathBehaviour::RefreshPathArray()
+{
+	if (PathStruct.PathArray.Num() < 1)
+	{
+		PathStruct.PathArray.Add("");
+	}
+}
+
diff --git a/Engine/Plugins/VirtualProduction/RemoteControl/Source/RemoteControlLogic/Public/Behaviour/Builtin/Path/RCSetAssetByPathBehaviour.h b/Engine/Plugins/VirtualProduction/RemoteControl/Source/RemoteControlLogic/Public/Behaviour/Builtin/Path/RCSetAssetByPathBehaviour.h
--- a/Engine/Plugins/VirtualProduction/RemoteControl/Source/RemoteControlLogic/Public/Behaviour/Builtin/Path/RCSetAssetByPathBehaviour.h
+++ b/Engine/Plugins/VirtualProduction/RemoteControl/Source/RemoteControlLogic/Public/Behaviour/Builtin/Path/RCSetAssetByPathBehaviour.h
@@ -27,7 +27,11 @@ struct FRCSetAssetPath
 {
 	GENERATED_BODY()
 
-	FRCSetAssetPath() = default;
+	FRCSetAssetPath()
+	{
+		// Add Empty index 0 when creating 
+		PathArray.Add("");
+	}
 	
 	/** An Array of Strings holding the Path of an Asset, seperated in several String. Will concated back together later. */
 	UPROPERTY()
@@ -65,6 +69,9 @@ public:
 
 	/** Auxialiary Function to apply to update the Target Texture */
 	void UpdateTargetEntity();
+
+	/** Called whenever a change has occured in the Slate */
+	void RefreshPathArray();
 	
 public:
 	/** Pointer to property container */
diff --git a/Engine/Plugins/VirtualProduction/RemoteControl/Source/RemoteControlUI/Private/UI/Behaviour/Builtin/Path/RCBehaviourSetAssetByPathModel.cpp b/Engine/Plugins/VirtualProduction/RemoteControl/Source/RemoteControlUI/Private/UI/Behaviour/Builtin/Path/RCBehaviourSetAssetByPathModel.cpp
--- a/Engine/Plugins/VirtualProduction/RemoteControl/Source/RemoteControlUI/Private/UI/Behaviour/Builtin/Path/RCBehaviourSetAssetByPathModel.cpp
+++ b/Engine/Plugins/VirtualProduction/RemoteControl/Source/RemoteControlUI/Private/UI/Behaviour/Builtin/Path/RCBehaviourSetAssetByPathModel.cpp
@@ -208,9 +208,13 @@ void FRCSetAssetByPathBehaviourModel::RegenerateWeakPtrInternal()
 
 void FRCSetAssetByPathBehaviourModel::RefreshPathAndPreview()
 {
-	TSharedPtr<SVerticalBox> FieldArrayWidget = SNew(SVerticalBox);
+	const TSharedPtr<SVerticalBox> FieldArrayWidget = SNew(SVerticalBox);
 
 	RegenerateWeakPtrInternal();
+	if (URCSetAssetByPathBehaviour* Behaviour = Cast<URCSetAssetByPathBehaviour>(GetBehaviour()))
+	{
+		Behaviour->RefreshPathArray();
+	}
 	
 	for (const TSharedPtr<IDetailTreeNode>& DetailTreeNodeArray : DetailTreeNodeWeakPtrArray)
 	{
@@ -379,7 +383,7 @@ FText FRCSetAssetByPathBehaviourModel::GetSelectedEntityText() const
 }
 
 
-void FRCSetAssetByPathBehaviourModel::CreateAssetPathFromSelection()
+void FRCSetAssetByPathBehaviourModel::CreateAssetPathFromSelection() const
 {
 	URCSetAssetByPathBehaviour* SetAssetByPathBehaviour = Cast<URCSetAssetByPathBehaviour>(GetBehaviour());
 	if (!SetAssetByPathBehaviour)
@@ -387,10 +391,6 @@ void FRCSetAssetByPathBehaviourModel::CreateAssetPathFromSelection()
 		return;
 	}
 	
-	// Clear it, in case it is an already used one.
-	TArray<FString>& PathArray = SetAssetByPathBehaviour->PathStruct.PathArray;
-	PathArray.Empty();
-	
 	TArray<FAssetData> AssetData;
 	GEditor->GetContentBrowserSelections(AssetData);
 	if (AssetData.Num() != 1)
@@ -398,6 +398,10 @@ void FRCSetAssetByPathBehaviourModel::CreateAssetPathFromSelection()
 		// Only works with one selected Content
 		return;
 	}
+	
+	// Clear it, in case it is an already used one.
+	TArray<FString>& PathArray = SetAssetByPathBehaviour->PathStruct.PathArray;
+	PathArray.Empty();
 
 	int32 IndexOfLast;
 	const UObject* SelectedAsset = AssetData[0].GetAsset();
diff --git a/Engine/Plugins/VirtualProduction/RemoteControl/Source/RemoteControlUI/Private/UI/Behaviour/Builtin/Path/RCBehaviourSetAssetByPathModel.h b/Engine/Plugins/VirtualProduction/RemoteControl/Source/RemoteControlUI/Private/UI/Behaviour/Builtin/Path/RCBehaviourSetAssetByPathModel.h
--- a/Engine/Plugins/VirtualProduction/RemoteControl/Source/RemoteControlUI/Private/UI/Behaviour/Builtin/Path/RCBehaviourSetAssetByPathModel.h
+++ b/Engine/Plugins/VirtualProduction/RemoteControl/Source/RemoteControlUI/Private/UI/Behaviour/Builtin/Path/RCBehaviourSetAssetByPathModel.h
@@ -78,5 +78,5 @@ private:
 	void RefreshPathAndPreview();
 
 	/** Creates a asset path, based on the selected Asset in the content browser. */
-	void CreateAssetPathFromSelection();
+	void CreateAssetPathFromSelection() const;
 };
diff --git a/Engine/Plugins/VirtualProduction/RemoteControl/Source/RemoteControlUI/Private/UI/Behaviour/Builtin/Path/SRCBehaviourSetAssetByPath.cpp b/Engine/Plugins/VirtualProduction/RemoteControl/Source/RemoteControlUI/Private/UI/Behaviour/Builtin/Path/SRCBehaviourSetAssetByPath.cpp
--- a/Engine/Plugins/VirtualProduction/RemoteControl/Source/RemoteControlUI/Private/UI/Behaviour/Builtin/Path/SRCBehaviourSetAssetByPath.cpp
+++ b/Engine/Plugins/VirtualProduction/RemoteControl/Source/RemoteControlUI/Private/UI/Behaviour/Builtin/Path/SRCBehaviourSetAssetByPath.cpp
@@ -20,11 +20,11 @@ BEGIN_SLATE_FUNCTION_BUILD_OPTIMIZATION
 void SRCBehaviourSetAssetByPath::Construct(const FArguments& InArgs, TSharedRef<FRCSetAssetByPathBehaviourModel> InBehaviourItem)
 {
 	URCSetAssetByPathBehaviour* PathBehaviour = Cast<URCSetAssetByPathBehaviour>(InBehaviourItem->GetBehaviour());
-	TSharedRef<SHorizontalBox> InternalExternalSwitchWidget = SNew(SHorizontalBox);
-	TSharedRef<SHorizontalBox> SelectedClassWidget = SNew(SHorizontalBox)
+	const TSharedRef<SHorizontalBox> SelectedClassWidget = SNew(SHorizontalBox);
+	const TSharedRef<SHorizontalBox> InternalExternalSwitchWidget = SNew(SHorizontalBox)
 		.Visibility_Lambda([PathBehaviour]()
 		{
-			return PathBehaviour->bInternal ? EVisibility::Visible : EVisibility::Collapsed;
+			return PathBehaviour->AssetClass->GetFName() == NAME_Vector ? EVisibility::Visible : EVisibility::Collapsed;
 		});
 
 	if (!PathBehaviour)
@@ -36,6 +36,30 @@ void SRCBehaviourSetAssetByPath::Construct(const FArguments& InArgs, TSharedRef<
 		return;
 	}
 		
+	
+	for (UClass* SupportedClass : PathBehaviour->GetSupportedClasses())
+	{
+		SelectedClassWidget->AddSlot()
+		.HAlign(HAlign_Right)
+		.VAlign(VAlign_Center)
+		.Padding(FMargin(0.0f,0.0f))
+		.FillWidth(1.0f)
+		.AutoWidth()
+		[
+			SNew(SButton)
+			.Text(FText::FromString(SupportedClass->GetName()))
+			.ButtonColorAndOpacity_Lambda([PathBehaviour, SupportedClass]()
+			{
+				FSlateColor Color = PathBehaviour->AssetClass == SupportedClass ? FAppStyle::Get().GetSlateColor("Colors.Select") : FAppStyle::Get().GetSlateColor("Colors.White");
+				return Color;
+			})
+			.OnPressed_Lambda([PathBehaviour, SupportedClass]()
+			{
+				PathBehaviour->AssetClass = SupportedClass;
+			})
+		];
+	}
+		
 	InternalExternalSwitchWidget->AddSlot()
 	.HAlign(HAlign_Right)
 	.VAlign(VAlign_Center)
@@ -76,40 +100,17 @@ void SRCBehaviourSetAssetByPath::Construct(const FArguments& InArgs, TSharedRef<
 		})
 	];
 	
-	for (UClass* SupportedClass : PathBehaviour->GetSupportedClasses())
-	{
-		SelectedClassWidget->AddSlot()
-		.HAlign(HAlign_Right)
-		.VAlign(VAlign_Center)
-		.Padding(FMargin(0.0f,0.0f))
-		.FillWidth(1.0f)
-		.AutoWidth()
-		[
-			SNew(SButton)
-			.Text(FText::FromString(SupportedClass->GetName()))
-			.ButtonColorAndOpacity_Lambda([PathBehaviour, SupportedClass]()
-			{
-				FSlateColor Color = PathBehaviour->AssetClass == SupportedClass ? FAppStyle::Get().GetSlateColor("Colors.Select") : FAppStyle::Get().GetSlateColor("Colors.White");
-				return Color;
-			})
-			.OnPressed_Lambda([PathBehaviour, SupportedClass]()
-			{
-				PathBehaviour->AssetClass = SupportedClass;
-			})
-		];
-	}
-	
 	ChildSlot
 	[
 		SNew(SVerticalBox)
 		+ SVerticalBox::Slot()
 		[
-			InternalExternalSwitchWidget
+			SelectedClassWidget
 		]
 		+ SVerticalBox::Slot()
 		.Padding(0.f,12.f)
 		[
-			SelectedClassWidget
+			InternalExternalSwitchWidget
 		]
 		+ SVerticalBox::Slot()
 		.AutoHeight()
