Auto Bind scrubbing (Revised approach)

Notable changes from previous approach:
* Uses Property Bag infrastructure for serialization of metadata
* No longer using SliderExponent
* Using Delta and LinearDeltaSensitivity for value independent scaling of slider values. Mouse motion should be the sole driver of value changes now.

diff --git a/Engine/Plugins/VirtualProduction/RemoteControl/Source/RemoteControl/Private/RCVirtualPropertyContainer.cpp b/Engine/Plugins/VirtualProduction/RemoteControl/Source/RemoteControl/Private/RCVirtualPropertyContainer.cpp
--- a/Engine/Plugins/VirtualProduction/RemoteControl/Source/RemoteControl/Private/RCVirtualPropertyContainer.cpp
+++ b/Engine/Plugins/VirtualProduction/RemoteControl/Source/RemoteControl/Private/RCVirtualPropertyContainer.cpp
@@ -19,16 +19,20 @@ void URCVirtualPropertyContainerBase::AddVirtualProperty(URCVirtualPropertyBase*
 	{
 		// Add property to Set
 		VirtualProperties.Add(InVirtualProperty);
-
-		// Initialize
-		InVirtualProperty->Init();
 	}
 }
 
-URCVirtualPropertyInContainer* URCVirtualPropertyContainerBase::AddProperty(const FName& InPropertyName, TSubclassOf<URCVirtualPropertyInContainer> InPropertyClass, const EPropertyBagPropertyType InValueType, UObject* InValueTypeObject)
+URCVirtualPropertyInContainer* URCVirtualPropertyContainerBase::AddProperty(const FName& InPropertyName, TSubclassOf<URCVirtualPropertyInContainer> InPropertyClass, const EPropertyBagPropertyType InValueType, UObject* InValueTypeObject, TArray<FPropertyBagPropertyDescMetaData> MetaData/* = TArray<FPropertyBagPropertyDescMetaData>()*/)
 {
 	const FName PropertyName = GenerateUniquePropertyName(InPropertyName, InValueType, InValueTypeObject, this);
-	Bag.AddProperty(PropertyName, InValueType, InValueTypeObject);
+
+	FPropertyBagPropertyDesc PropertyBagDesc = FPropertyBagPropertyDesc(PropertyName, InValueType, InValueTypeObject);
+
+#if WITH_EDITORONLY_DATA
+	PropertyBagDesc.MetaData = MetaData;
+#endif
+
+	Bag.AddProperties({ PropertyBagDesc });
 
 	// Ensure that the property has been successfully added to the Bag
 	const FPropertyBagPropertyDesc* BagPropertyDesc = Bag.FindPropertyDescByName(PropertyName);
@@ -147,6 +151,11 @@ URCVirtualPropertyBase* URCVirtualPropertyContainerBase::GetVirtualProperty(cons
 {
 	for (URCVirtualPropertyBase* VirtualProperty : VirtualProperties)
 	{
+		if (!ensure(VirtualProperty))
+		{
+			continue;
+		}
+
 		if (VirtualProperty->PropertyName == InPropertyName)
 		{
 			return VirtualProperty;
diff --git a/Engine/Plugins/VirtualProduction/RemoteControl/Source/RemoteControl/Public/RCVirtualPropertyContainer.h b/Engine/Plugins/VirtualProduction/RemoteControl/Source/RemoteControl/Public/RCVirtualPropertyContainer.h
--- a/Engine/Plugins/VirtualProduction/RemoteControl/Source/RemoteControl/Public/RCVirtualPropertyContainer.h
+++ b/Engine/Plugins/VirtualProduction/RemoteControl/Source/RemoteControl/Public/RCVirtualPropertyContainer.h
@@ -32,10 +32,11 @@ public:
 	 * @param InPropertyClass				Class of the Virtual Property 
 	 * @param InValueType					Property Type
 	 * @param InValueTypeObject				Property Type object if exists
+	 * @param MetaData                      Property Metadata used by Slate widgets (eg: Delta, LinearDeltaSensitivity, etc)
 	 *
 	 * @return Virtual Property Object
 	 */
-	virtual URCVirtualPropertyInContainer* AddProperty(const FName& InPropertyName, TSubclassOf<URCVirtualPropertyInContainer> InPropertyClass, const EPropertyBagPropertyType InValueType, UObject* InValueTypeObject = nullptr);
+	virtual URCVirtualPropertyInContainer* AddProperty(const FName& InPropertyName, TSubclassOf<URCVirtualPropertyInContainer> InPropertyClass, const EPropertyBagPropertyType InValueType, UObject* InValueTypeObject = nullptr, TArray<FPropertyBagPropertyDescMetaData> MetaData = TArray<FPropertyBagPropertyDescMetaData>());
 
 	/**
 	 * Duplicates a property from given FProperty.
diff --git a/Engine/Plugins/VirtualProduction/RemoteControl/Source/RemoteControlLogic/Private/Controller/RCController.cpp b/Engine/Plugins/VirtualProduction/RemoteControl/Source/RemoteControlLogic/Private/Controller/RCController.cpp
--- a/Engine/Plugins/VirtualProduction/RemoteControl/Source/RemoteControlLogic/Private/Controller/RCController.cpp
+++ b/Engine/Plugins/VirtualProduction/RemoteControl/Source/RemoteControlLogic/Private/Controller/RCController.cpp
@@ -10,29 +10,6 @@
 
 #define LOCTEXT_NAMESPACE "RCController"
 
-void URCController::Init()
-{
-#if WITH_EDITORONLY_DATA
-	// Setup metadata
-	// Note: In the future this can be set directly by the user from UI
-	if (IsVectorType())
-	{
-		if (FProperty* Property = GetProperty())
-		{
-			// Provides users with a vector scrubbing experience that approximates the SceneComponent's Location Vector widget
-			Property->SetMetaData(TEXT("SliderExponent"), FString::Printf(TEXT("%f"), SliderExponent));
-		}
-	}
-#endif
-}
-
-void URCController::PostLoad()
-{
-	Super::PostLoad();
-
-	Init();
-}
-
 #if WITH_EDITOR
 void URCController::PostEditUndo()
 {
diff --git a/Engine/Plugins/VirtualProduction/RemoteControl/Source/RemoteControlLogic/Private/Controller/RCControllerContainer.cpp b/Engine/Plugins/VirtualProduction/RemoteControl/Source/RemoteControlLogic/Private/Controller/RCControllerContainer.cpp
--- a/Engine/Plugins/VirtualProduction/RemoteControl/Source/RemoteControlLogic/Private/Controller/RCControllerContainer.cpp
+++ b/Engine/Plugins/VirtualProduction/RemoteControl/Source/RemoteControlLogic/Private/Controller/RCControllerContainer.cpp
@@ -1,10 +1,32 @@
 ﻿// Copyright Epic Games, Inc. All Rights Reserved.
 
 #include "Controller/RCControllerContainer.h"
+
 #include "Controller/RCController.h"
+#include "RCVirtualProperty.h"
+#include "Templates/SubclassOf.h"
 
-#if WITH_EDITOR
+URCVirtualPropertyInContainer* URCControllerContainer::AddProperty(const FName& InPropertyName, TSubclassOf<URCVirtualPropertyInContainer> InPropertyClass, const EPropertyBagPropertyType InValueType, UObject* InValueTypeObject /*= nullptr*/, TArray<FPropertyBagPropertyDescMetaData> MetaData /*= TArray<FPropertyBagPropertyDescMetaData>()*/)
+{
+	// Vector Controllers
+	if (InValueType == EPropertyBagPropertyType::Struct)
+	{
+		if (InValueTypeObject == TBaseStructure<FVector>::Get())
+		{
+			MetaData.Add(FPropertyBagPropertyDescMetaData(FName("Delta"), FString::Printf(TEXT("%f"), VectorSliderDelta)));
+			MetaData.Add(FPropertyBagPropertyDescMetaData(FName("LinearDeltaSensitivity"), FString::Printf(TEXT("%f"), VectorLinearDeltaSensitivity)));
+		}
+		else if (InValueTypeObject == TBaseStructure<FRotator>::Get())
+		{
+			MetaData.Add(FPropertyBagPropertyDescMetaData(FName("Delta"), FString::Printf(TEXT("%f"), RotatorSliderDelta)));
+			MetaData.Add(FPropertyBagPropertyDescMetaData(FName("LinearDeltaSensitivity"), FString::Printf(TEXT("%f"), RotatorLinearDeltaSensitivity)));
+		}
+	}
+
+	return Super::AddProperty(InPropertyName, InPropertyClass, InValueType, InValueTypeObject, MetaData);
+}
 
+#if WITH_EDITOR
 URCController* URCControllerContainer::GetControllerFromChangeEvent(const FPropertyChangedEvent& Event)
 {
 	if (const FProperty* FinalProperty = (Event.Property == Event.MemberProperty) ? Event.Property : Event.MemberProperty)
diff --git a/Engine/Plugins/VirtualProduction/RemoteControl/Source/RemoteControlLogic/Public/Controller/RCController.h b/Engine/Plugins/VirtualProduction/RemoteControl/Source/RemoteControlLogic/Public/Controller/RCController.h
--- a/Engine/Plugins/VirtualProduction/RemoteControl/Source/RemoteControlLogic/Public/Controller/RCController.h
+++ b/Engine/Plugins/VirtualProduction/RemoteControl/Source/RemoteControlLogic/Public/Controller/RCController.h
@@ -19,11 +19,6 @@ class REMOTECONTROLLOGIC_API URCController : public URCVirtualPropertyInContaine
 	GENERATED_BODY()
 
 public:
-	/** Initialization routine. Called after the parent container has setup data for this property */
-	virtual void Init();
-
-	/** Called post serialization. Used to reinitialize Controller metadata after a Remote Control Preset asset is loaded*/
-	virtual void PostLoad() override;
 
 #if WITH_EDITOR
 	/** Called after applying a transaction to the object. Used to broadcast Undo related container changes to UI Used to broadcast Undo related container changes to UI */
@@ -64,9 +59,4 @@ public:
 	/** Set of the behaviours */
 	UPROPERTY()
 	TSet<TObjectPtr<URCBehaviour>> Behaviours;
-
-	/** Slider exponent for Vector Controllers. Determines how fast the value change travels while scrubbing the X/Y/Z sliders in UI
-	* The default value approximates the behaviour of SceneComponent Location Vectors (see SNumericVectorInputBox in ComponentTransformDetails.cpp) */
-	UPROPERTY()
-	float SliderExponent = 0.4f;
 };
diff --git a/Engine/Plugins/VirtualProduction/RemoteControl/Source/RemoteControlLogic/Public/Controller/RCControllerContainer.h b/Engine/Plugins/VirtualProduction/RemoteControl/Source/RemoteControlLogic/Public/Controller/RCControllerContainer.h
--- a/Engine/Plugins/VirtualProduction/RemoteControl/Source/RemoteControlLogic/Public/Controller/RCControllerContainer.h
+++ b/Engine/Plugins/VirtualProduction/RemoteControl/Source/RemoteControlLogic/Public/Controller/RCControllerContainer.h
@@ -15,8 +15,29 @@ UCLASS()
 class REMOTECONTROLLOGIC_API URCControllerContainer : public URCVirtualPropertyContainerBase
 {
 	GENERATED_BODY()
+
+public:
+	/** Adds a new Controller to this container */
+	virtual URCVirtualPropertyInContainer* AddProperty(const FName& InPropertyName, TSubclassOf<URCVirtualPropertyInContainer> InPropertyClass, const EPropertyBagPropertyType InValueType, UObject* InValueTypeObject = nullptr, TArray<FPropertyBagPropertyDescMetaData> MetaData = TArray<FPropertyBagPropertyDescMetaData>()) override;
 	
 private:
+	/** Slider exponent for Vector Controller sliders. Determines how fast the value change travels while scrubbing the X/Y/Z sliders in UI
+	* The default value approximates the behaviour of SceneComponent Location Vectors (see SNumericVectorInputBox in ComponentTransformDetails.cpp) */
+	UPROPERTY()
+	float VectorSliderDelta = 5.f;
+
+	/** Linear Delta Sensitivity for Vector Controller sliders. Multiplied inversely with mouse movement to determine actual slider change */
+	UPROPERTY()
+	float VectorLinearDeltaSensitivity = 10.f;
+
+	/** Slider exponent for Rotator Controller sliders.*/
+	UPROPERTY()
+	float RotatorSliderDelta = 10.f;
+
+	/** Linear Delta Sensitivity for Rotator Controller sliders*/
+	UPROPERTY()
+	float RotatorLinearDeltaSensitivity = 10.f;
+
 #if WITH_EDITOR
 	/** Fetches the Controller underlying a given a Property Changed Event*/
 	URCController* GetControllerFromChangeEvent(const FPropertyChangedEvent& Event);
