"Edit Step" option for Range Map Behaviour

diff --git a/Engine/Plugins/VirtualProduction/RemoteControl/Source/RemoteControlLogic/Private/Behaviour/Builtin/RangeMap/RCRangeMapBehaviour.cpp b/Engine/Plugins/VirtualProduction/RemoteControl/Source/RemoteControlLogic/Private/Behaviour/Builtin/RangeMap/RCRangeMapBehaviour.cpp
--- a/Engine/Plugins/VirtualProduction/RemoteControl/Source/RemoteControlLogic/Private/Behaviour/Builtin/RangeMap/RCRangeMapBehaviour.cpp
+++ b/Engine/Plugins/VirtualProduction/RemoteControl/Source/RemoteControlLogic/Private/Behaviour/Builtin/RangeMap/RCRangeMapBehaviour.cpp
@@ -28,6 +28,8 @@ namespace UE::RCRangeMapBehaviour
 URCRangeMapBehaviour::URCRangeMapBehaviour()
 {
 	PropertyContainer = CreateDefaultSubobject<URCVirtualPropertyContainerBase>(FName("VirtualPropertyContainer"));
+
+	bExecuteBehavioursDuringPreChange = true;
 }
 
 void URCRangeMapBehaviour::Initialize()
@@ -43,7 +45,6 @@ void URCRangeMapBehaviour::Initialize()
 	ContainerPtr->AddProperty(UE::RCRangeMapBehaviour::MinValue, URCController::StaticClass(), EPropertyBagPropertyType::Double);
 	ContainerPtr->AddProperty(UE::RCRangeMapBehaviour::MaxValue, URCController::StaticClass(), EPropertyBagPropertyType::Double);
 	ContainerPtr->AddProperty(UE::RCRangeMapBehaviour::Threshold, URCController::StaticClass(), EPropertyBagPropertyType::Double);
-	ContainerPtr->AddProperty(UE::RCRangeMapBehaviour::Step, URCController::StaticClass(), EPropertyBagPropertyType::Double);
 
 	Super::Initialize();
 }
@@ -124,6 +125,27 @@ bool URCRangeMapBehaviour::GetNearestActionByThreshold(TTuple<URCAction*, bool>&
 	return true;
 }
 
+
+bool FRCRangeMapStep::GetStepValue(double& OutValue) const
+{
+	if (StepValueProperty)
+	{
+		return StepValueProperty->GetValueDouble(OutValue);
+	}
+
+	return false;
+}
+
+bool URCRangeMapBehaviour::GetStepValueForAction(const URCAction* InAction, double& OutValue)
+{
+	if (FRCRangeMapStep* StepData = RangeMapActionContainer.Find(InAction))
+	{
+		return StepData->GetStepValue(OutValue);
+	}
+
+	return false;
+}
+
 void URCRangeMapBehaviour::GetLerpActions(TMap<FGuid, TArray<URCAction*>>& OutNumericActionsByField)
 {
 	for (URCAction* Action : ActionContainer->GetActions())
@@ -145,10 +167,16 @@ void URCRangeMapBehaviour::GetLerpActions(TMap<FGuid, TArray<URCAction*>>& OutNu
 
 		Algo::Sort(ArrayToSort, [this](const URCAction* ActionA, const URCAction* ActionB)
 		{
-			const double A = RangeMapActionContainer.Find(ActionA)->StepValue;
-			const double B = RangeMapActionContainer.Find(ActionB)->StepValue;
+			double A, B;
+			bool bRes1 = GetStepValueForAction(ActionA, A);
+			bool bRes2 = GetStepValueForAction(ActionA, B);
+
+			if (ensure(bRes1 && bRes2))
+			{
+				return A > B;
+			}
 
-			return A > B;
+			return false;
 		});
 	}
 }
@@ -219,9 +247,17 @@ void URCRangeMapBehaviour::Execute()
 			return;
 		}
 
+		double MinRangeStepValue, MaxRangeStepValue;
+		bool bRes1 = MinRangeMapStep->GetStepValue(MinRangeStepValue);
+		bool bRes2 = MaxRangeMapStep->GetStepValue(MaxRangeStepValue);
+		if (!ensure(bRes1 && bRes2))
+		{
+			continue; // ideally this should never happen
+		}
+
 		// Denormalize and Map them based on our Min and Max Value
-		const double MappedMinStep = UKismetMathLibrary::Lerp(MinValue, MaxValue, MinRangeMapStep->StepValue);
-		const double MappedMaxStep = UKismetMathLibrary::Lerp(MinValue, MaxValue, MaxRangeMapStep->StepValue);
+		const double MappedMinStep = UKismetMathLibrary::Lerp(MinValue, MaxValue, MinRangeStepValue);
+		const double MappedMaxStep = UKismetMathLibrary::Lerp(MinValue, MaxValue, MaxRangeStepValue);
 
 		// Normalize our Controller based on the new MappedMinStep and MappedMaxStep Range.
 		double CustomNormalizedStepValue = UKismetMathLibrary::NormalizeToRange(ControllerFloatValue, MappedMinStep, MappedMaxStep);
@@ -296,7 +332,7 @@ URCAction* URCRangeMapBehaviour::DuplicateAction(URCAction* InAction, URCBehavio
 
 			if(ensureAlways(ActionValueProperty))
 			{
-				FRCRangeMapStep NewData(RangeMapStepData->StepValue, ActionValueProperty);
+				FRCRangeMapStep NewData(RangeMapStepData->StepValueProperty, ActionValueProperty);
 				InBehaviourRangeMap->RangeMapActionContainer.Add(NewAction, NewData);
 			}
 		}
@@ -307,26 +343,20 @@ URCAction* URCRangeMapBehaviour::DuplicateAction(URCAction* InAction, URCBehavio
 
 URCAction* URCRangeMapBehaviour::AddAction(const TSharedRef<const FRemoteControlField> InRemoteControlField)
 {
-	double InStepValue = 0.0;
-	URCVirtualPropertyBase* StepVirtualProperty = PropertyContainer->GetVirtualProperty(UE::RCRangeMapBehaviour::Step);
-	StepVirtualProperty->GetValueDouble(InStepValue);
-	
-	const TRCActionUniquenessTest ActionUniquenessTest = [this, InRemoteControlField, InStepValue](const TSet<TObjectPtr<URCAction>>& Actions)
-	{
-		return IsActionUnique(InRemoteControlField, InStepValue, Actions);
-	};
+	// The new workflow for Rangemap involves the user creating the Step action first (with default step value) and then clicking on Edit to set the value
+	// For this reason no uniqueness test can be performed at this point
+	const TRCActionUniquenessTest PassthroughUniquenessTest = [](const TSet<TObjectPtr<URCAction>>& Actions) { return true;	};
 	
-	return ActionContainer->AddAction(ActionUniquenessTest, InRemoteControlField);
+	return ActionContainer->AddAction(PassthroughUniquenessTest, InRemoteControlField);
 }
 
 void URCRangeMapBehaviour::OnActionAdded(URCAction* Action, URCVirtualPropertySelfContainer* InPropertyValue)
 {
-	double InStepValue = 0.0;
-	URCVirtualPropertyBase* StepVirtualProperty = PropertyContainer->GetVirtualProperty(UE::RCRangeMapBehaviour::Step);
-	StepVirtualProperty->GetValueDouble(InStepValue);
-	StepVirtualProperty->SetValueDouble(0.0);
+	// Virtual Property for Step Input
+	URCVirtualPropertySelfContainer* StepInputProperty = NewObject<URCVirtualPropertySelfContainer>(this);
+	StepInputProperty->AddProperty(UE::RCRangeMapBehaviour::Step, EPropertyBagPropertyType::Double);
 	
-	FRCRangeMapStep RangeMapStep = FRCRangeMapStep(InStepValue, InPropertyValue);
+	FRCRangeMapStep RangeMapStep = FRCRangeMapStep(StepInputProperty, InPropertyValue);
 
 	RangeMapActionContainer.Add(Action, MoveTemp(RangeMapStep));
 }
@@ -361,8 +391,15 @@ bool URCRangeMapBehaviour::GetRangeValuePairsForLerp(TMap<FGuid, TTuple<URCActio
 				continue;
 			}
 
+			double StepValue;
+			bool bRes = RangeMap->GetStepValue(StepValue);
+			if (!bRes)
+			{
+				continue;
+			}
+
 			// Deal with the nullptr
-			if (NormalizedControllerValue >= RangeMap->StepValue && (!MinRangeMap || NumericActionTuple.Value.Last() != Action))
+			if (NormalizedControllerValue >= StepValue && (!MinRangeMap || NumericActionTuple.Value.Last() != Action))
 			{
 				MinAction = Action;
 				MinRangeMap = RangeMap;
@@ -370,7 +407,7 @@ bool URCRangeMapBehaviour::GetRangeValuePairsForLerp(TMap<FGuid, TTuple<URCActio
 				continue;
 			}
 
-			if (NormalizedControllerValue <= RangeMap->StepValue)
+			if (NormalizedControllerValue <= StepValue)
 			{
 				MaxAction = Action;
 				MaxRangeMap = RangeMap;
@@ -397,14 +434,17 @@ TMap<double, URCAction*> URCRangeMapBehaviour::GetNonLerpActions()
 	
 	for (URCAction* Action : ActionContainer->GetActions())
 	{
-		
 		if (!IsSupportedActionLerpType(Action))
 		{
 			FRCRangeMapStep* RangeMapStep = RangeMapActionContainer.Find(Action);
 
 			if (RangeMapStep)
 			{
-				ActionMap.Add(RangeMapStep->StepValue, Action);
+				double StepValue;
+				if (RangeMapStep->GetStepValue(StepValue))
+				{
+					ActionMap.Add(StepValue, Action);
+				}
 			}
 		}
 	}
@@ -476,15 +516,17 @@ bool URCRangeMapBehaviour::IsActionUnique(const TSharedRef<const FRemoteControlF
 		if (const FRCRangeMapStep* RangeMapStep = this->RangeMapActionContainer.Find(Action))
 		{
 			// Check if the Step itself is already used
-			const double ActionStepValue = RangeMapStep->StepValue;
-
-			// Only Important bit is StepValue and FieldId
-			if (Action->ExposedFieldId == FieldId && FMath::IsNearlyEqual(ActionStepValue, InStepValue))
+			double ActionStepValue;
+			if(RangeMapStep->GetStepValue(ActionStepValue))
 			{
-				return false; // Not Unique.
+				// Only Important bit is StepValue and FieldId
+				if (Action->ExposedFieldId == FieldId && FMath::IsNearlyEqual(ActionStepValue, InStepValue))
+				{
+					return false; // Not Unique.
+				}
 			}
 		}
 	}
 	
 	return true;
-}
+}
\ No newline at end of file
diff --git a/Engine/Plugins/VirtualProduction/RemoteControl/Source/RemoteControlLogic/Public/Behaviour/Builtin/RangeMap/RCRangeMapBehaviour.h b/Engine/Plugins/VirtualProduction/RemoteControl/Source/RemoteControlLogic/Public/Behaviour/Builtin/RangeMap/RCRangeMapBehaviour.h
--- a/Engine/Plugins/VirtualProduction/RemoteControl/Source/RemoteControlLogic/Public/Behaviour/Builtin/RangeMap/RCRangeMapBehaviour.h
+++ b/Engine/Plugins/VirtualProduction/RemoteControl/Source/RemoteControlLogic/Public/Behaviour/Builtin/RangeMap/RCRangeMapBehaviour.h
@@ -10,29 +10,26 @@ class URCAction;
 class URCVirtualPropertyContainerBase;
 
 USTRUCT()
-struct FRCRangeMapStep
+struct REMOTECONTROLLOGIC_API FRCRangeMapStep
 {
 	GENERATED_BODY()
 
 	FRCRangeMapStep() {}
-	FRCRangeMapStep(const double InStepValue, URCVirtualPropertySelfContainer* InPropertyValue)
-		: StepValue(InStepValue)
+	FRCRangeMapStep(URCVirtualPropertySelfContainer* InStepValueProperty, URCVirtualPropertySelfContainer* InPropertyValue)
+		: StepValueProperty(InStepValueProperty), PropertyValue(InPropertyValue)
 	{
-		if (!InPropertyValue)
-		{
-			return;
-		}
-		
-		PropertyValue = InPropertyValue;
 	}
 
 	/** The Value which we use represent the action based on a normalized step. */
 	UPROPERTY()
-	double StepValue = 0.0;
+	TObjectPtr<URCVirtualPropertySelfContainer> StepValueProperty;
 
 	/** The Property this Action holds and will be used for calculations for the lerp. */
 	UPROPERTY()
 	TObjectPtr<URCVirtualPropertySelfContainer> PropertyValue;
+
+	/** Returns the Step value from the virtual property, safely*/
+	bool GetStepValue(double& OutValue) const;
 };
 
 /**
@@ -77,6 +74,9 @@ public:
 	/** Checks whether or not one of the following fields can be added to the AddAction Menu. Makes sure they're unique. */
 	bool CanHaveActionForField(const TSharedPtr<FRemoteControlField> InRemoteControlField) const override;
 
+	/** Returns the Step value associated with a given Action*/
+	bool GetStepValueForAction(const URCAction* InAction, double& OutValue);
+
 private:
 	/** Minimum Value which the Range has */
 	double MinValue;
@@ -91,6 +91,7 @@ private:
 	float ControllerFloatValue;
 	
 private:
+
 	/** Boolean Function to help differentiate Custom Actions */
 	bool IsSupportedActionLerpType(TObjectPtr<URCAction> InAction) const;
 
diff --git a/Engine/Plugins/VirtualProduction/RemoteControl/Source/RemoteControlUI/Private/UI/Action/RCActionModel.cpp b/Engine/Plugins/VirtualProduction/RemoteControl/Source/RemoteControlUI/Private/UI/Action/RCActionModel.cpp
--- a/Engine/Plugins/VirtualProduction/RemoteControl/Source/RemoteControlUI/Private/UI/Action/RCActionModel.cpp
+++ b/Engine/Plugins/VirtualProduction/RemoteControl/Source/RemoteControlUI/Private/UI/Action/RCActionModel.cpp
@@ -92,7 +92,8 @@ TSharedRef<SWidget> FRCActionModel::GetNameWidget() const
 			if (const TSharedPtr<FRemoteControlField> RemoteControlField = Preset->GetExposedEntity<FRemoteControlField>(Action->ExposedFieldId).Pin())
 			{
 				return SNew(SBox)
-					.Padding(FMargin(6.f))
+					.Padding(FMargin(8.f))
+					.VAlign(EVerticalAlignment::VAlign_Center)
 					[
 						SNew(STextBlock).Text(FText::FromName(RemoteControlField->GetLabel()))
 					];
diff --git a/Engine/Plugins/VirtualProduction/RemoteControl/Source/RemoteControlUI/Private/UI/Action/RangeMap/RCActionRangeMapModel.cpp b/Engine/Plugins/VirtualProduction/RemoteControl/Source/RemoteControlUI/Private/UI/Action/RangeMap/RCActionRangeMapModel.cpp
--- a/Engine/Plugins/VirtualProduction/RemoteControl/Source/RemoteControlUI/Private/UI/Action/RangeMap/RCActionRangeMapModel.cpp
+++ b/Engine/Plugins/VirtualProduction/RemoteControl/Source/RemoteControlUI/Private/UI/Action/RangeMap/RCActionRangeMapModel.cpp
@@ -11,6 +11,7 @@
 #include "RemoteControlField.h"
 #include "RemoteControlPreset.h"
 #include "Styling/RemoteControlStyles.h"
+#include "UI/Action/SRCVirtualPropertyWidget.h"
 #include "UI/Behaviour/Builtin/RangeMap/RCBehaviourRangeMapModel.h"
 #include "UI/RCUIHelpers.h"
 #include "UI/RemoteControlPanelStyle.h"
@@ -127,7 +128,7 @@ TSharedPtr<FRCActionRangeMapModel> FRCActionRangeMapModel::GetModelByActionType(
 	return nullptr;
 }
 
-TSharedRef<SWidget> FRCActionRangeMapModel::GetStepWidget() const
+TSharedRef<SWidget> FRCActionRangeMapModel::GetStepWidget()
 {
 	if (TSharedPtr<FRCRangeMapBehaviourModel> BehaviourItem = StaticCastSharedPtr<FRCRangeMapBehaviourModel>(BehaviourItemWeakPtr.Pin()))
 	{
@@ -135,13 +136,11 @@ TSharedRef<SWidget> FRCActionRangeMapModel::GetStepWidget() const
 		{
 			if (FRCRangeMapStep* RangeStep = Behaviour->RangeMapActionContainer.Find(GetAction()))
 			{
-				const FText StepAsText = FText::FromString(FString::SanitizeFloat(RangeStep->StepValue));
-				
-				return SNew(SBox)
-					.Padding(6.0f)
-					[
-						SNew(STextBlock).Text(StepAsText)
-					];
+				if (ensure(RangeStep->StepValueProperty))
+				{
+					return SAssignNew(EditableVirtualPropertyWidget, SRCVirtualPropertyWidget, RangeStep->StepValueProperty)
+						.OnGenerateWidget(this, &FRCActionRangeMapModel::OnGenerateStepWidget);
+				}
 			}
 		}
 	}
@@ -149,6 +148,36 @@ TSharedRef<SWidget> FRCActionRangeMapModel::GetStepWidget() const
 	return SNullWidget::NullWidget;
 }
 
+TSharedRef<SWidget> FRCActionRangeMapModel::OnGenerateStepWidget(class URCVirtualPropertySelfContainer* StepValueProperty) const
+{
+	if (!ensure(StepValueProperty))
+	{
+		return SNullWidget::NullWidget;
+	}
+
+	if (TSharedPtr<FRCRangeMapBehaviourModel> BehaviourItem = StaticCastSharedPtr<FRCRangeMapBehaviourModel>(BehaviourItemWeakPtr.Pin()))
+	{
+		if (URCRangeMapBehaviour* Behaviour = Cast<URCRangeMapBehaviour>(BehaviourItem->GetBehaviour()))
+		{
+			if (const FRCRangeMapStep* RangeStep = Behaviour->RangeMapActionContainer.Find(GetAction()))
+			{
+				double StepValue;
+				if(RangeStep->GetStepValue(StepValue))
+				{
+					const FText StepAsText = FText::FromString(FString::SanitizeFloat(StepValue));
+
+					return SNew(SBox)
+						[
+							SNew(STextBlock).Text(StepAsText)
+						];
+				}
+			}
+		}
+	}
+
+	return SNullWidget::NullWidget;
+}
+
 TSharedRef<ITableRow> FRCActionRangeMapModel::OnGenerateWidgetForList(TSharedPtr<FRCActionRangeMapModel> InItem, const TSharedRef<STableViewBase>& OwnerTable)
 {
 	typedef UE::RCActionPanelRangeMapList::SActionItemListRow ActionRowType;
diff --git a/Engine/Plugins/VirtualProduction/RemoteControl/Source/RemoteControlUI/Private/UI/Action/RangeMap/RCActionRangeMapModel.h b/Engine/Plugins/VirtualProduction/RemoteControl/Source/RemoteControlUI/Private/UI/Action/RangeMap/RCActionRangeMapModel.h
--- a/Engine/Plugins/VirtualProduction/RemoteControl/Source/RemoteControlUI/Private/UI/Action/RangeMap/RCActionRangeMapModel.h
+++ b/Engine/Plugins/VirtualProduction/RemoteControl/Source/RemoteControlUI/Private/UI/Action/RangeMap/RCActionRangeMapModel.h
@@ -22,7 +22,7 @@ public:
 	/**
 	 * Simple Widget representing the Normalized Value or Step Value for the Range and this particular Value.
 	 */
-	virtual TSharedRef<SWidget> GetStepWidget() const;
+	virtual TSharedRef<SWidget> GetStepWidget();
 
 	/** OnGenerateRow delegate for the Actions List View*/
 	TSharedRef<ITableRow> OnGenerateWidgetForList(TSharedPtr<FRCActionRangeMapModel> InItem, const TSharedRef<STableViewBase>& OwnerTable);
@@ -32,6 +32,9 @@ public:
 
 	/** Chooses the appropriate Action model for the current class and field type*/
 	static TSharedPtr<FRCActionRangeMapModel> GetModelByActionType(URCAction* InAction, const TSharedPtr<class FRCBehaviourModel> InBehaviourItem, const TSharedPtr<SRemoteControlPanel> InRemoteControlPanel);
+
+private:
+	TSharedRef<SWidget> OnGenerateStepWidget(class URCVirtualPropertySelfContainer* InComparand) const;
 };
 
 /*
diff --git a/Engine/Plugins/VirtualProduction/RemoteControl/Source/RemoteControlUI/Private/UI/Behaviour/Builtin/RangeMap/RCBehaviourRangeMapModel.cpp b/Engine/Plugins/VirtualProduction/RemoteControl/Source/RemoteControlUI/Private/UI/Behaviour/Builtin/RangeMap/RCBehaviourRangeMapModel.cpp
--- a/Engine/Plugins/VirtualProduction/RemoteControl/Source/RemoteControlUI/Private/UI/Behaviour/Builtin/RangeMap/RCBehaviourRangeMapModel.cpp
+++ b/Engine/Plugins/VirtualProduction/RemoteControl/Source/RemoteControlUI/Private/UI/Behaviour/Builtin/RangeMap/RCBehaviourRangeMapModel.cpp
@@ -74,9 +74,7 @@ TSharedRef<SWidget> FRCRangeMapBehaviourModel::GetPropertyWidget() const
 	}
 	
 	TSharedPtr<IDetailTreeNode> DetailTreeNodeThreshold = (++DetailTreeItr)->Pin();
-	TSharedPtr<IDetailTreeNode> DetailTreeNodeStep = (++DetailTreeItr)->Pin();
 	const FNodeWidgets ThresholdWidget = DetailTreeNodeThreshold->CreateNodeWidgets();
-	const FNodeWidgets StepWidget = DetailTreeNodeStep->CreateNodeWidgets();
 	
 	FieldWidget->AddSlot()
 		.Padding(FMargin(3.0f, 2.0f))
@@ -100,34 +98,6 @@ TSharedRef<SWidget> FRCRangeMapBehaviourModel::GetPropertyWidget() const
 			]
 		];
 
-	FieldWidget->AddSlot()
-		.Padding(FMargin(3.0f, 6.0f))
-		.HAlign(HAlign_Right)
-		.AutoHeight()
-		[
-			SNew(SBorder)
-			[
-				SNew(SHorizontalBox)
-				+SHorizontalBox::Slot()
-				.Padding(FMargin(3.0f, 2.0f))
-				.HAlign(HAlign_Right)
-				.AutoWidth()
-				[
-					SNew(STextBlock)
-					.Text(LOCTEXT("StepInputValue", "Step Value for Action"))
-					.Font(FRemoteControlPanelStyle::Get()->GetFontStyle("RemoteControlPanel.Behaviours.BehaviourDescription"))
-				]
-				+SHorizontalBox::Slot()
-				.Padding(FMargin(3.0f, 2.0f))
-				.HAlign(HAlign_Right)
-				.FillWidth(1.0f)
-				.AutoWidth()
-				[
-					StepWidget.ValueWidget.ToSharedRef()
-				]
-			]
-		];
-
 	return 	SNew(SHorizontalBox)
 		.Clipping(EWidgetClipping::OnDemand)
 		+ SHorizontalBox::Slot()
@@ -161,9 +131,6 @@ URCAction* FRCRangeMapBehaviourModel::AddAction(const TSharedRef<const FRemoteCo
 
 	if (URCRangeMapBehaviour* RangeMapBehaviour = Cast<URCRangeMapBehaviour>(GetBehaviour()))
 	{
-		double StepValue;
-		RangeMapBehaviour->PropertyContainer->GetVirtualProperty(FName("Step"))->GetValueDouble(StepValue);
-
 		const FGuid FieldId = InRemoteControlField->GetId();
 		if (TSharedPtr<FRemoteControlProperty> RemoteProperty = RangeMapBehaviour->ControllerWeakPtr.Get()->PresetWeakPtr.Get()->GetExposedEntity<FRemoteControlProperty>(FieldId).Pin())
 		{
