Regression: Multiple Actions for same field

This was broken by the commit for supporting Material Actions.

We needed to explicitly check if we're dealing with an Array property as well, in addition to the other checks.

diff --git a/Engine/Plugins/VirtualProduction/RemoteControl/Source/RemoteControlLogic/Private/Action/RCActionContainer.cpp b/Engine/Plugins/VirtualProduction/RemoteControl/Source/RemoteControlLogic/Private/Action/RCActionContainer.cpp
--- a/Engine/Plugins/VirtualProduction/RemoteControl/Source/RemoteControlLogic/Private/Action/RCActionContainer.cpp
+++ b/Engine/Plugins/VirtualProduction/RemoteControl/Source/RemoteControlLogic/Private/Action/RCActionContainer.cpp
@@ -69,6 +69,17 @@ URCBehaviour* URCActionContainer::GetParentBehaviour()
 	return Cast<URCBehaviour>(GetOuter());
 }
 
+TArray<const URCPropertyAction*> URCActionContainer::GetPropertyActions() const
+{
+	TArray<const URCPropertyAction*> PropertyActions;
+
+	Algo::TransformIf(Actions, PropertyActions,
+		[](const URCAction* Action) { return Action->IsA(URCPropertyAction::StaticClass()); },
+		[](const URCAction* Action) { return Cast<URCPropertyAction>(Action); });
+
+	return PropertyActions;
+}
+
 URCPropertyAction* URCActionContainer::AddPropertyAction(const TSharedRef<const FRemoteControlProperty> InRemoteControlProperty)
 {
 	// Create new Property
@@ -88,26 +99,25 @@ URCPropertyAction* URCActionContainer::AddPropertyAction(const TSharedRef<const
 
 	if (FProperty* Property = InRemoteControlProperty->GetProperty())
 	{
-		for (URCAction* Action : Actions)
+		// Reusing input property is only applicable for Array properties
+		// Eg: "Override Materials" will be represented by a single virtual property Array with each Action occupying a different index slot.
+		if(Property->IsA(FArrayProperty::StaticClass()))
 		{
-			if (URCPropertyAction* PropertyAction = Cast<URCPropertyAction>(Action))
+			for (const URCPropertyAction* PropertyAction : GetPropertyActions())
 			{
-				if (PropertyAction->PropertySelfContainer)
+				if (FProperty* InputFieldProperty = PropertyAction->GetProperty())
 				{
-					if(FProperty* InputFieldProperty = PropertyAction->PropertySelfContainer->GetProperty())
+					if (InputFieldProperty->GetFName() == Property->GetFName())
 					{
-						if(InputFieldProperty->GetFName() == Property->GetFName())
+						if (const TSharedPtr<const FRemoteControlProperty> RemoteControlPropertyForAction = PropertyAction->GetRemoteControlProperty())
 						{
-							if (const TSharedPtr<const FRemoteControlProperty> RemoteControlPropertyForAction = PropertyAction->GetRemoteControlProperty())
+							if (InRemoteControlProperty->GetBoundObject() == RemoteControlPropertyForAction->GetBoundObject())
 							{
-								if (InRemoteControlProperty->GetBoundObject() == RemoteControlPropertyForAction->GetBoundObject())
-								{
-									// We already have an Action associated with this property, so this must be a container type like Array.
-									// Both the UI widget and Execute will operate on a unique index associated with this Action for this container property
-									NewPropertyAction->PropertySelfContainer = PropertyAction->PropertySelfContainer;
-
-									bFoundMatchingContainer = true;
-								}
+								// We already have an Action associated with this Array container property.
+								// The array container will be reused, both the UI widget and Execute will operate on a unique index associated with the respective Action
+								NewPropertyAction->PropertySelfContainer = PropertyAction->PropertySelfContainer;
+
+								bFoundMatchingContainer = true;
 							}
 						}
 					}
diff --git a/Engine/Plugins/VirtualProduction/RemoteControl/Source/RemoteControlLogic/Private/Action/RCPropertyAction.cpp b/Engine/Plugins/VirtualProduction/RemoteControl/Source/RemoteControlLogic/Private/Action/RCPropertyAction.cpp
--- a/Engine/Plugins/VirtualProduction/RemoteControl/Source/RemoteControlLogic/Private/Action/RCPropertyAction.cpp
+++ b/Engine/Plugins/VirtualProduction/RemoteControl/Source/RemoteControlLogic/Private/Action/RCPropertyAction.cpp
@@ -72,6 +72,16 @@ void URCPropertyAction::Execute() const
 	Super::Execute();
 }
 
+FProperty* URCPropertyAction::GetProperty() const
+{
+	if (ensure(PropertySelfContainer))
+	{
+		return PropertySelfContainer->GetProperty();
+	}
+
+	return nullptr;
+}
+
 FName URCAction::GetExposedFieldLabel() const
 {
 	if (const URemoteControlPreset* Preset = PresetWeakPtr.Get())
diff --git a/Engine/Plugins/VirtualProduction/RemoteControl/Source/RemoteControlLogic/Public/Action/RCActionContainer.h b/Engine/Plugins/VirtualProduction/RemoteControl/Source/RemoteControlLogic/Public/Action/RCActionContainer.h
--- a/Engine/Plugins/VirtualProduction/RemoteControl/Source/RemoteControlLogic/Public/Action/RCActionContainer.h
+++ b/Engine/Plugins/VirtualProduction/RemoteControl/Source/RemoteControlLogic/Public/Action/RCActionContainer.h
@@ -64,6 +64,9 @@ public:
 	/** Retrieves the actions present in this container */
 	const TSet<TObjectPtr<URCAction>>& GetActions() const { return Actions; }
 
+	/** Retrieves all the Property Actions present in this container */
+	TArray<const URCPropertyAction*> GetPropertyActions() const;
+
 #if WITH_EDITOR
 	/** Called after applying a transaction to the object. Used to broadcast Undo related container changes to UI */
 	virtual void PostEditUndo();
diff --git a/Engine/Plugins/VirtualProduction/RemoteControl/Source/RemoteControlLogic/Public/Action/RCPropertyAction.h b/Engine/Plugins/VirtualProduction/RemoteControl/Source/RemoteControlLogic/Public/Action/RCPropertyAction.h
--- a/Engine/Plugins/VirtualProduction/RemoteControl/Source/RemoteControlLogic/Public/Action/RCPropertyAction.h
+++ b/Engine/Plugins/VirtualProduction/RemoteControl/Source/RemoteControlLogic/Public/Action/RCPropertyAction.h
@@ -24,6 +24,9 @@ public:
 	virtual void Execute() const override;
 	//~ End URCAction interface
 
+	/* Returns the FProperty underlying this Action */
+	FProperty* GetProperty() const;
+
 public:
 	/** Virtual Property Container */
 	UPROPERTY()
