Remote Control Logic - RangeMap Behaviour - Implement UI Feedback

- Change from Step to Input
- Remove Threshold and make it an internal value of 0.05

diff --git a/Engine/Plugins/VirtualProduction/RemoteControl/Source/RemoteControlLogic/Private/Behaviour/Builtin/RangeMap/RCBehaviourRangeMapNode.cpp b/Engine/Plugins/VirtualProduction/RemoteControl/Source/RemoteControlLogic/Private/Behaviour/Builtin/RangeMap/RCBehaviourRangeMapNode.cpp
--- a/Engine/Plugins/VirtualProduction/RemoteControl/Source/RemoteControlLogic/Private/Behaviour/Builtin/RangeMap/RCBehaviourRangeMapNode.cpp
+++ b/Engine/Plugins/VirtualProduction/RemoteControl/Source/RemoteControlLogic/Private/Behaviour/Builtin/RangeMap/RCBehaviourRangeMapNode.cpp
@@ -9,7 +9,7 @@
 URCBehaviourRangeMapNode::URCBehaviourRangeMapNode()
 {
 	DisplayName = NSLOCTEXT("Remote Control Behaviour", "Behavior Name - Range Mapping", "Range");
-	BehaviorDescription = NSLOCTEXT("Remote Control Behaviour", "Behavior Desc - Range Mapping", "Behaviour lerps property values between a normalized 0-1 step. Users define min/max values and tolerance threshold.");
+	BehaviorDescription = NSLOCTEXT("Remote Control Behaviour", "Behavior Desc - Range Mapping", "This Behaviours can take a list of Actions, sort them using normalized values between 0..1, and interpolates numerical values whilst executing non-numerical whose distance are less than 0.05 from a given action.");
 }
 
 bool URCBehaviourRangeMapNode::Execute_Implementation(URCBehaviour* InBehaviour) const
diff --git a/Engine/Plugins/VirtualProduction/RemoteControl/Source/RemoteControlLogic/Private/Behaviour/Builtin/RangeMap/RCRangeMapBehaviour.cpp b/Engine/Plugins/VirtualProduction/RemoteControl/Source/RemoteControlLogic/Private/Behaviour/Builtin/RangeMap/RCRangeMapBehaviour.cpp
--- a/Engine/Plugins/VirtualProduction/RemoteControl/Source/RemoteControlLogic/Private/Behaviour/Builtin/RangeMap/RCRangeMapBehaviour.cpp
+++ b/Engine/Plugins/VirtualProduction/RemoteControl/Source/RemoteControlLogic/Private/Behaviour/Builtin/RangeMap/RCRangeMapBehaviour.cpp
@@ -19,10 +19,11 @@
 
 namespace UE::RCRangeMapBehaviour
 {
-	const FName MinValue = TEXT("MinValue");
-	const FName MaxValue = TEXT("MaxValue");
-	const FName Threshold = TEXT("Threshold");
-	const FName Step = TEXT("Step");
+	const FName InputMin = TEXT("InputMin");
+	const FName InputMax = TEXT("InputMax");
+	const FName Input = TEXT("Input");
+	
+	constexpr double Threshold = 0.05;
 }
 
 URCRangeMapBehaviour::URCRangeMapBehaviour()
@@ -34,17 +35,17 @@ URCRangeMapBehaviour::URCRangeMapBehaviour()
 
 void URCRangeMapBehaviour::Initialize()
 {
-	URCController* RCController = ControllerWeakPtr.Get();
+	const URCController* RCController = ControllerWeakPtr.Get();
 	if (!RCController)
 	{
 		return;
 	}
 
-	TWeakObjectPtr<URCVirtualPropertyContainerBase> ContainerPtr = PropertyContainer;
+	const TWeakObjectPtr<URCVirtualPropertyContainerBase> ContainerPtr = PropertyContainer;
 
-	ContainerPtr->AddProperty(UE::RCRangeMapBehaviour::MinValue, URCController::StaticClass(), EPropertyBagPropertyType::Double);
-	ContainerPtr->AddProperty(UE::RCRangeMapBehaviour::MaxValue, URCController::StaticClass(), EPropertyBagPropertyType::Double);
-	ContainerPtr->AddProperty(UE::RCRangeMapBehaviour::Threshold, URCController::StaticClass(), EPropertyBagPropertyType::Double);
+	ContainerPtr->AddProperty(UE::RCRangeMapBehaviour::InputMin, URCController::StaticClass(), EPropertyBagPropertyType::Double);
+	ContainerPtr->AddProperty(UE::RCRangeMapBehaviour::InputMax, URCController::StaticClass(), EPropertyBagPropertyType::Double);
+	ContainerPtr->AddProperty(UE::RCRangeMapBehaviour::Input, URCController::StaticClass(), EPropertyBagPropertyType::Double);
 
 	Super::Initialize();
 }
@@ -57,7 +58,7 @@ bool URCRangeMapBehaviour::IsSupportedActionLerpType(TObjectPtr<URCAction> InAct
 		return false;
 	}
 	
-	TObjectPtr<URCPropertyAction> InPropertyAction = Cast<URCPropertyAction>(InAction);
+	const TObjectPtr<URCPropertyAction> InPropertyAction = Cast<URCPropertyAction>(InAction);
 	if (!InPropertyAction)
 	{
 		return false;
@@ -78,17 +79,20 @@ void URCRangeMapBehaviour::Refresh()
 		return;
 	}
 
+	if (!PropertyContainer->GetVirtualProperty(UE::RCRangeMapBehaviour::InputMin) || !PropertyContainer->GetVirtualProperty(UE::RCRangeMapBehaviour::InputMax))
+	{
+		return;
+	}
 	// Step 0: Get All required Values, getting the current values set within the DetailPanel setting them up for our header.
-	PropertyContainer->GetVirtualProperty(UE::RCRangeMapBehaviour::MinValue)->GetValueDouble(MinValue);
-	PropertyContainer->GetVirtualProperty(UE::RCRangeMapBehaviour::MaxValue)->GetValueDouble(MaxValue);
-	PropertyContainer->GetVirtualProperty(UE::RCRangeMapBehaviour::Threshold)->GetValueDouble(Threshold);
+	PropertyContainer->GetVirtualProperty(UE::RCRangeMapBehaviour::InputMin)->GetValueDouble(InputMin);
+	PropertyContainer->GetVirtualProperty(UE::RCRangeMapBehaviour::InputMax)->GetValueDouble(InputMax);
 
 	Controller->GetValueFloat(ControllerFloatValue);
 	
 	// Step 1: Clamp the Value between Max/Min if necessary
-	if (ControllerFloatValue < MinValue || MaxValue < ControllerFloatValue)
+	if (ControllerFloatValue < InputMin || InputMax < ControllerFloatValue)
 	{
-		ControllerFloatValue = FMath::Clamp(ControllerFloatValue, MinValue, MaxValue);
+		ControllerFloatValue = FMath::Clamp(ControllerFloatValue, InputMin, InputMax);
 		Controller->SetValueFloat(ControllerFloatValue);
 	}
 }
@@ -98,49 +102,50 @@ bool URCRangeMapBehaviour::GetNearestActionByThreshold(TTuple<URCAction*, bool>&
 	// Variables
 	TMap<double, URCAction*> NonLerpActions = GetNonLerpActions();
 	
-	TArray<double> StepActionArray;
-	NonLerpActions.GenerateKeyArray(StepActionArray);
+	TArray<double> InputActionArray;
+	NonLerpActions.GenerateKeyArray(InputActionArray);
 
 	// Step 1: Convert to double for Kismet Operation and Normalize.
-	double NormalizedControllerValue = UKismetMathLibrary::NormalizeToRange(ControllerFloatValue, MinValue, MaxValue);
+	const double NormalizedControllerValue = UKismetMathLibrary::NormalizeToRange(ControllerFloatValue, InputMin, InputMax);
 
-	// Step 2: Go through each of the Action Steps and look for the shortest distance
-	for (int StepIndex = 0; StepIndex < StepActionArray.Num(); StepIndex++)
+	// Step 2: Go through each of the Action Input and look for the shortest distance
+	for (int InputIndex = 0; InputIndex < InputActionArray.Num(); InputIndex++)
 	{
-		double StepValue = StepActionArray[StepIndex];
-		double ValueDifference = FMath::Abs(StepValue - NormalizedControllerValue);
+		const double InputValue = InputActionArray[InputIndex];
+		const double ValueDifference = FMath::Abs(InputValue - NormalizedControllerValue);
 		
-		if (StepActionArray.Num() > 1 && StepIndex > 0)
+		if (InputActionArray.Num() > 1 && InputIndex > 0)
 		{
-			double PrevValueDifference = FMath::Abs(StepActionArray[StepIndex-1] - NormalizedControllerValue);
+			const double PrevValueDifference = FMath::Abs(InputActionArray[InputIndex-1] - NormalizedControllerValue);
 			
-			// Break out early in case previous Step had a lesser difference
+			// Break out early in case previous Input had a lesser difference
 			if (PrevValueDifference <= ValueDifference)
 			{
 				break;
 			}
 		}
-		OutTuple = TTuple<URCAction*, bool>(NonLerpActions[StepValue], ValueDifference <= Threshold);
+		OutTuple = TTuple<URCAction*, bool>(NonLerpActions[InputValue], ValueDifference <= UE::RCRangeMapBehaviour::Threshold);
 	}
+	
 	return true;
 }
 
 
-bool FRCRangeMapStep::GetStepValue(double& OutValue) const
+bool FRCRangeMapInput::GetInputValue(double& OutValue) const
 {
-	if (StepValueProperty)
+	if (InputProperty)
 	{
-		return StepValueProperty->GetValueDouble(OutValue);
+		return InputProperty->GetValueDouble(OutValue);
 	}
 
 	return false;
 }
 
-bool URCRangeMapBehaviour::GetStepValueForAction(const URCAction* InAction, double& OutValue)
+bool URCRangeMapBehaviour::GetValueForAction(const URCAction* InAction, double& OutValue)
 {
-	if (FRCRangeMapStep* StepData = RangeMapActionContainer.Find(InAction))
+	if (FRCRangeMapInput* StepData = RangeMapActionContainer.Find(InAction))
 	{
-		return StepData->GetStepValue(OutValue);
+		return StepData->GetInputValue(OutValue);
 	}
 
 	return false;
@@ -160,7 +165,7 @@ void URCRangeMapBehaviour::GetLerpActions(TMap<FGuid, TArray<URCAction*>>& OutNu
 		}
 	}
 
-	// Step 03: Sort actions using their StepValue
+	// Step 03: Sort actions using their InputValue
 	for (TTuple<FGuid, TArray<URCAction*>>& NumericActionTuple : OutNumericActionsByField)
 	{
 		TArray<URCAction*>& ArrayToSort = NumericActionTuple.Value;
@@ -168,10 +173,10 @@ void URCRangeMapBehaviour::GetLerpActions(TMap<FGuid, TArray<URCAction*>>& OutNu
 		Algo::Sort(ArrayToSort, [this](const URCAction* ActionA, const URCAction* ActionB)
 		{
 			double A, B;
-			bool bRes1 = GetStepValueForAction(ActionA, A);
-			bool bRes2 = GetStepValueForAction(ActionA, B);
+			bool bRes1 = GetValueForAction(ActionA, A);
+			bool bRes2 = GetValueForAction(ActionA, B);
 
-			if (ensure(bRes1 && bRes2))
+			if (bRes1 && bRes2)
 			{
 				return A > B;
 			}
@@ -195,7 +200,7 @@ void URCRangeMapBehaviour::Execute()
 	// Do this beforehand.
 	BehaviourNode->PreExecute(this);
 	
-	URCController* RCController = ControllerWeakPtr.Get();
+	const URCController* RCController = ControllerWeakPtr.Get();
 	if (!RCController)
 	{
 		ensureMsgf(false, TEXT("Remote Control Controller is not available/nullptr."));
@@ -203,7 +208,7 @@ void URCRangeMapBehaviour::Execute()
 	}
 
 	RCController->GetValueFloat(ControllerFloatValue);
-	ControllerFloatValue = FMath::Clamp(ControllerFloatValue, MinValue, MaxValue);
+	ControllerFloatValue = FMath::Clamp(ControllerFloatValue, InputMin, InputMax);
 
 	TTuple<URCAction*, bool> NearestAction;
 	if (GetNearestActionByThreshold(NearestAction) && NearestAction.Value)
@@ -229,7 +234,7 @@ void URCRangeMapBehaviour::Execute()
 		}
 		FGuid CurrentFieldId = LerpAction.Key;
 
-		// Get the StepValues given for the FGuid
+		// Get the InputValue given for the FGuid
 		const URCAction* MinActionOfPair = LerpAction.Value.Key;
 		const URCAction* MaxActionOfPair = LerpAction.Value.Value;
 
@@ -238,42 +243,42 @@ void URCRangeMapBehaviour::Execute()
 			return;
 		}
 
-		const FRCRangeMapStep* MinRangeMapStep = RangeMapActionContainer.Find(MinActionOfPair);
-		const FRCRangeMapStep* MaxRangeMapStep = RangeMapActionContainer.Find(MaxActionOfPair);
+		const FRCRangeMapInput* MinRangeMapInput = RangeMapActionContainer.Find(MinActionOfPair);
+		const FRCRangeMapInput* MaxRangeMapInput = RangeMapActionContainer.Find(MaxActionOfPair);
 
-		if (!(MinRangeMapStep && MaxRangeMapStep))
+		if (!(MinRangeMapInput && MaxRangeMapInput))
 		{
-			ensureMsgf(false, TEXT("%s is an invalid pointer."), MinRangeMapStep ? TEXT("MaxRangeMapStep") : TEXT("MinRangeMapStep"));
+			ensureMsgf(false, TEXT("%s is an invalid pointer."), MinRangeMapInput ? TEXT("MaxRangeMapInput") : TEXT("MinRangeMapInput"));
 			return;
 		}
 
-		double MinRangeStepValue, MaxRangeStepValue;
-		bool bRes1 = MinRangeMapStep->GetStepValue(MinRangeStepValue);
-		bool bRes2 = MaxRangeMapStep->GetStepValue(MaxRangeStepValue);
-		if (!ensure(bRes1 && bRes2))
+		double MinRangeValue, MaxRangeValue;
+		bool bRes1 = MinRangeMapInput->GetInputValue(MinRangeValue);
+		bool bRes2 = MaxRangeMapInput->GetInputValue(MaxRangeValue);
+		if (!(bRes1 && bRes2))
 		{
 			continue; // ideally this should never happen
 		}
 
 		// Denormalize and Map them based on our Min and Max Value
-		const double MappedMinStep = UKismetMathLibrary::Lerp(MinValue, MaxValue, MinRangeStepValue);
-		const double MappedMaxStep = UKismetMathLibrary::Lerp(MinValue, MaxValue, MaxRangeStepValue);
+		const double MappedMinInput = UKismetMathLibrary::Lerp(InputMin, InputMax, MinRangeValue);
+		const double MappedMaxInput = UKismetMathLibrary::Lerp(InputMin, InputMax, MaxRangeValue);
 
-		// Normalize our Controller based on the new MappedMinStep and MappedMaxStep Range.
-		double CustomNormalizedStepValue = UKismetMathLibrary::NormalizeToRange(ControllerFloatValue, MappedMinStep, MappedMaxStep);
+		// Normalize our Controller based on the new MappedMinInput and MappedMaxInput Range.
+		double CustomNormalizedInputValue = UKismetMathLibrary::NormalizeToRange(ControllerFloatValue, MappedMinInput, MappedMaxInput);
 
 		// Apply value based on type
-		switch(MinRangeMapStep->PropertyValue->GetValueType())
+		switch(MinRangeMapInput->PropertyValue->GetValueType())
 		{
-		case EPropertyBagPropertyType::Double:
+			case EPropertyBagPropertyType::Double:
 			{
 				double MinRangeDouble;
 				double MaxRangeDouble;
 
-				MinRangeMapStep->PropertyValue->GetValueDouble(MinRangeDouble);
-				MaxRangeMapStep->PropertyValue->GetValueDouble(MaxRangeDouble);
+				MinRangeMapInput->PropertyValue->GetValueDouble(MinRangeDouble);
+				MaxRangeMapInput->PropertyValue->GetValueDouble(MaxRangeDouble);
 
-				const double ResultDouble = UKismetMathLibrary::Lerp(MinRangeDouble, MaxRangeDouble, CustomNormalizedStepValue);
+				const double ResultDouble = UKismetMathLibrary::Lerp(MinRangeDouble, MaxRangeDouble, CustomNormalizedInputValue);
 				if (const TSharedPtr<FRemoteControlProperty> RCProperty = Preset->GetExposedEntity<FRemoteControlProperty>(CurrentFieldId).Pin())
 				{
 					RCProperty->GetPropertyHandle()->SetValue(ResultDouble);
@@ -281,15 +286,15 @@ void URCRangeMapBehaviour::Execute()
 				
 				break;
 			}
-		case EPropertyBagPropertyType::Float:
+			case EPropertyBagPropertyType::Float:
 			{
 				float MinRangeFloat;
 				float MaxRangeFloat;
 
-				MinRangeMapStep->PropertyValue->GetValueFloat(MinRangeFloat);
-				MaxRangeMapStep->PropertyValue->GetValueFloat(MaxRangeFloat);
+				MinRangeMapInput->PropertyValue->GetValueFloat(MinRangeFloat);
+				MaxRangeMapInput->PropertyValue->GetValueFloat(MaxRangeFloat);
 
-				const float ResultFloat = UKismetMathLibrary::Lerp(MinRangeFloat, MaxRangeFloat, CustomNormalizedStepValue);
+				const float ResultFloat = UKismetMathLibrary::Lerp(MinRangeFloat, MaxRangeFloat, CustomNormalizedInputValue);
 				if (const TSharedPtr<FRemoteControlProperty> RCProperty = Preset->GetExposedEntity<FRemoteControlProperty>(CurrentFieldId).Pin())
 				{
 					RCProperty->GetPropertyHandle()->SetValue(ResultFloat);
@@ -297,13 +302,13 @@ void URCRangeMapBehaviour::Execute()
 				
 				break;
 			}
-		case EPropertyBagPropertyType::Struct:
+			case EPropertyBagPropertyType::Struct:
 			{
-				ApplyLerpOnStruct(MinRangeMapStep, MaxRangeMapStep, CustomNormalizedStepValue, CurrentFieldId);
+				ApplyLerpOnStruct(MinRangeMapInput, MaxRangeMapInput, CustomNormalizedInputValue, CurrentFieldId);
 				
 				break;
 			}
-		default:
+			default:
 			// Shouldn't happen.
 			ensureMsgf(false, TEXT("Unsupported Lerp Type."));
 		}
@@ -322,17 +327,17 @@ URCAction* URCRangeMapBehaviour::DuplicateAction(URCAction* InAction, URCBehavio
 	if (ensure(NewAction))
 	{
 		// Copy action specific data residing in Range Map Behaviour:
-		if (FRCRangeMapStep* RangeMapStepData = RangeMapActionContainer.Find(InAction))
+		if (const FRCRangeMapInput* RangeMapInputData = RangeMapActionContainer.Find(InAction))
 		{
 			URCVirtualPropertySelfContainer* ActionValueProperty = nullptr;
-			if (URCPropertyAction* PropertyAction = Cast<URCPropertyAction>(NewAction))
+			if (const URCPropertyAction* PropertyAction = Cast<URCPropertyAction>(NewAction))
 			{
 				ActionValueProperty = PropertyAction->PropertySelfContainer;
 			}
 
 			if(ensureAlways(ActionValueProperty))
 			{
-				FRCRangeMapStep NewData(RangeMapStepData->StepValueProperty, ActionValueProperty);
+				const FRCRangeMapInput NewData(RangeMapInputData->InputProperty, ActionValueProperty);
 				InBehaviourRangeMap->RangeMapActionContainer.Add(NewAction, NewData);
 			}
 		}
@@ -353,17 +358,17 @@ URCAction* URCRangeMapBehaviour::AddAction(const TSharedRef<const FRemoteControl
 void URCRangeMapBehaviour::OnActionAdded(URCAction* Action, URCVirtualPropertySelfContainer* InPropertyValue)
 {
 	// Virtual Property for Step Input
-	URCVirtualPropertySelfContainer* StepInputProperty = NewObject<URCVirtualPropertySelfContainer>(this);
-	StepInputProperty->AddProperty(UE::RCRangeMapBehaviour::Step, EPropertyBagPropertyType::Double);
+	URCVirtualPropertySelfContainer* InputProperty = NewObject<URCVirtualPropertySelfContainer>(this);
+	InputProperty->AddProperty(UE::RCRangeMapBehaviour::Input, EPropertyBagPropertyType::Double);
 	
-	FRCRangeMapStep RangeMapStep = FRCRangeMapStep(StepInputProperty, InPropertyValue);
+	FRCRangeMapInput RangeMapInput = FRCRangeMapInput(InputProperty, InPropertyValue);
 
-	RangeMapActionContainer.Add(Action, MoveTemp(RangeMapStep));
+	RangeMapActionContainer.Add(Action, MoveTemp(RangeMapInput));
 }
 
 bool URCRangeMapBehaviour::GetRangeValuePairsForLerp(TMap<FGuid, TTuple<URCAction*, URCAction*>>& OutPairs)
 {
-	const double NormalizedControllerValue = UKismetMathLibrary::NormalizeToRange(ControllerFloatValue, MinValue, MaxValue);
+	const double NormalizedControllerValue = UKismetMathLibrary::NormalizeToRange(ControllerFloatValue, InputMin, InputMax);
 
 	TMap<FGuid, TArray<URCAction*>> LerpActionMap;
 	GetLerpActions(LerpActionMap);
@@ -377,29 +382,29 @@ bool URCRangeMapBehaviour::GetRangeValuePairsForLerp(TMap<FGuid, TTuple<URCActio
 		}
 
 		// Intermediary calculations of Lerp for in-between Lerp.
-		FRCRangeMapStep* MaxRangeMap = nullptr;
-		FRCRangeMapStep* MinRangeMap = nullptr;
+		FRCRangeMapInput* MaxRangeMap = nullptr;
+		FRCRangeMapInput* MinRangeMap = nullptr;
 		URCAction* MaxAction = nullptr;
 		URCAction* MinAction = nullptr;
 		
 		for (URCAction* Action : NumericActionTuple.Value)
 		{
-			FRCRangeMapStep* RangeMap = RangeMapActionContainer.Find(Action);
+			FRCRangeMapInput* RangeMap = RangeMapActionContainer.Find(Action);
 			if (!RangeMap)
 			{
 				// Shouldn't happen, but skip if it does
 				continue;
 			}
 
-			double StepValue;
-			bool bRes = RangeMap->GetStepValue(StepValue);
+			double Value;
+			bool bRes = RangeMap->GetInputValue(Value);
 			if (!bRes)
 			{
 				continue;
 			}
 
 			// Deal with the nullptr
-			if (NormalizedControllerValue >= StepValue && (!MinRangeMap || NumericActionTuple.Value.Last() != Action))
+			if (NormalizedControllerValue >= Value && (!MinRangeMap || NumericActionTuple.Value.Last() != Action))
 			{
 				MinAction = Action;
 				MinRangeMap = RangeMap;
@@ -407,7 +412,7 @@ bool URCRangeMapBehaviour::GetRangeValuePairsForLerp(TMap<FGuid, TTuple<URCActio
 				continue;
 			}
 
-			if (NormalizedControllerValue <= StepValue)
+			if (NormalizedControllerValue <= Value)
 			{
 				MaxAction = Action;
 				MaxRangeMap = RangeMap;
@@ -436,12 +441,10 @@ TMap<double, URCAction*> URCRangeMapBehaviour::GetNonLerpActions()
 	{
 		if (!IsSupportedActionLerpType(Action))
 		{
-			FRCRangeMapStep* RangeMapStep = RangeMapActionContainer.Find(Action);
-
-			if (RangeMapStep)
+			if (const FRCRangeMapInput* RangeMapInput = RangeMapActionContainer.Find(Action))
 			{
 				double StepValue;
-				if (RangeMapStep->GetStepValue(StepValue))
+				if (RangeMapInput->GetInputValue(StepValue))
 				{
 					ActionMap.Add(StepValue, Action);
 				}
@@ -458,9 +461,9 @@ bool URCRangeMapBehaviour::CanHaveActionForField(const TSharedPtr<FRemoteControl
 	return true;
 }
 
-void URCRangeMapBehaviour::ApplyLerpOnStruct(const FRCRangeMapStep* MinRangeStep, const FRCRangeMapStep* MaxRangeStep, const double& StepAlpha, const FGuid& FieldId)
+void URCRangeMapBehaviour::ApplyLerpOnStruct(const FRCRangeMapInput* MinRangeInput, const FRCRangeMapInput* MaxRangeInput, const double& InputAlpha, const FGuid& FieldId)
 {
-	URCController* RCController = ControllerWeakPtr.Get();
+	const URCController* RCController = ControllerWeakPtr.Get();
 	if (!RCController)
 	{
 		ensureMsgf(false, TEXT("Remote Controller is invalid or unavailable."));
@@ -474,53 +477,53 @@ void URCRangeMapBehaviour::ApplyLerpOnStruct(const FRCRangeMapStep* MinRangeStep
 		return;
 	}
 	
-	// Check whether its a Rotator or FVector based on the MinRangeStep
-	if (MinRangeStep->PropertyValue->IsVectorType())
+	// Check whether its a Rotator or FVector based on the MinRangeInput
+	if (MinRangeInput->PropertyValue->IsVectorType())
 	{
 		FVector MinRangeVector;
 		FVector MaxRangeVector;
 
-		MinRangeStep->PropertyValue->GetValueVector(MinRangeVector);
-		MaxRangeStep->PropertyValue->GetValueVector(MaxRangeVector);
+		MinRangeInput->PropertyValue->GetValueVector(MinRangeVector);
+		MaxRangeInput->PropertyValue->GetValueVector(MaxRangeVector);
 
-		FVector ResultVector = UKismetMathLibrary::VLerp(MinRangeVector, MaxRangeVector, StepAlpha);
-		if (TSharedPtr<FRemoteControlProperty> RCProperty = Preset->GetExposedEntity<FRemoteControlProperty>(FieldId).Pin())
+		const FVector ResultVector = UKismetMathLibrary::VLerp(MinRangeVector, MaxRangeVector, InputAlpha);
+		if (const TSharedPtr<FRemoteControlProperty> RCProperty = Preset->GetExposedEntity<FRemoteControlProperty>(FieldId).Pin())
 		{
 			RCProperty->GetPropertyHandle()->SetValue(ResultVector);
 		}
 	}
 
-	if (MinRangeStep->PropertyValue->IsRotatorType())
+	if (MinRangeInput->PropertyValue->IsRotatorType())
 	{
 		FRotator MinRangeRotator;
 		FRotator MaxRangeRotator;
 
-		MinRangeStep->PropertyValue->GetValueRotator(MinRangeRotator);
-		MaxRangeStep->PropertyValue->GetValueRotator(MaxRangeRotator);
+		MinRangeInput->PropertyValue->GetValueRotator(MinRangeRotator);
+		MaxRangeInput->PropertyValue->GetValueRotator(MaxRangeRotator);
 
-		FRotator ResultVector = UKismetMathLibrary::RLerp(MinRangeRotator, MaxRangeRotator, StepAlpha, true);
-		if (TSharedPtr<FRemoteControlProperty> RCProperty = Preset->GetExposedEntity<FRemoteControlProperty>(FieldId).Pin())
+		const FRotator ResultVector = UKismetMathLibrary::RLerp(MinRangeRotator, MaxRangeRotator, InputAlpha, true);
+		if (const TSharedPtr<FRemoteControlProperty> RCProperty = Preset->GetExposedEntity<FRemoteControlProperty>(FieldId).Pin())
 		{
-			 RCProperty->GetPropertyHandle()->SetValue(ResultVector);
+			RCProperty->GetPropertyHandle()->SetValue(ResultVector);
 		}
 	}
 }
 
-bool URCRangeMapBehaviour::IsActionUnique(const TSharedRef<const FRemoteControlField> InRemoteControlField, const double& InStepValue, const TSet<TObjectPtr<URCAction>>& InActions)
+bool URCRangeMapBehaviour::IsActionUnique(const TSharedRef<const FRemoteControlField> InRemoteControlField, const double& InValue, const TSet<TObjectPtr<URCAction>>& InActions)
 {
 	const FGuid FieldId = InRemoteControlField->GetId();
 
 	for (const URCAction* Action : InActions)
 	{
-		// Check whether or not we have a RangeStep.
-		if (const FRCRangeMapStep* RangeMapStep = this->RangeMapActionContainer.Find(Action))
+		// Check whether or not we have a RangeInput.
+		if (const FRCRangeMapInput* RangeMapInput = this->RangeMapActionContainer.Find(Action))
 		{
 			// Check if the Step itself is already used
 			double ActionStepValue;
-			if(RangeMapStep->GetStepValue(ActionStepValue))
+			if(RangeMapInput->GetInputValue(ActionStepValue))
 			{
 				// Only Important bit is StepValue and FieldId
-				if (Action->ExposedFieldId == FieldId && FMath::IsNearlyEqual(ActionStepValue, InStepValue))
+				if (Action->ExposedFieldId == FieldId && FMath::IsNearlyEqual(ActionStepValue, InValue))
 				{
 					return false; // Not Unique.
 				}
diff --git a/Engine/Plugins/VirtualProduction/RemoteControl/Source/RemoteControlLogic/Public/Behaviour/Builtin/RangeMap/RCRangeMapBehaviour.h b/Engine/Plugins/VirtualProduction/RemoteControl/Source/RemoteControlLogic/Public/Behaviour/Builtin/RangeMap/RCRangeMapBehaviour.h
--- a/Engine/Plugins/VirtualProduction/RemoteControl/Source/RemoteControlLogic/Public/Behaviour/Builtin/RangeMap/RCRangeMapBehaviour.h
+++ b/Engine/Plugins/VirtualProduction/RemoteControl/Source/RemoteControlLogic/Public/Behaviour/Builtin/RangeMap/RCRangeMapBehaviour.h
@@ -10,26 +10,26 @@ class URCAction;
 class URCVirtualPropertyContainerBase;
 
 USTRUCT()
-struct REMOTECONTROLLOGIC_API FRCRangeMapStep
+struct REMOTECONTROLLOGIC_API FRCRangeMapInput
 {
 	GENERATED_BODY()
 
-	FRCRangeMapStep() {}
-	FRCRangeMapStep(URCVirtualPropertySelfContainer* InStepValueProperty, URCVirtualPropertySelfContainer* InPropertyValue)
-		: StepValueProperty(InStepValueProperty), PropertyValue(InPropertyValue)
+	FRCRangeMapInput() {}
+	FRCRangeMapInput(URCVirtualPropertySelfContainer* InSelfContainer, URCVirtualPropertySelfContainer* InPropertyValue)
+		: InputProperty(InSelfContainer), PropertyValue(InPropertyValue)
 	{
 	}
 
 	/** The Value which we use represent the action based on a normalized step. */
 	UPROPERTY()
-	TObjectPtr<URCVirtualPropertySelfContainer> StepValueProperty;
+	TObjectPtr<URCVirtualPropertySelfContainer> InputProperty;
 
 	/** The Property this Action holds and will be used for calculations for the lerp. */
 	UPROPERTY()
 	TObjectPtr<URCVirtualPropertySelfContainer> PropertyValue;
 
 	/** Returns the Step value from the virtual property, safely*/
-	bool GetStepValue(double& OutValue) const;
+	bool GetInputValue(double& OutValue) const;
 };
 
 /**
@@ -45,9 +45,9 @@ public:
 	UPROPERTY()
 	TObjectPtr<URCVirtualPropertyContainerBase> PropertyContainer;
 
-	/** Container holding all FRCRangeMapSteps correlating to each Action in the ActionContainer. */
+	/** Container holding all FRCRangeMapInputs correlating to each Action in the ActionContainer. */
 	UPROPERTY()
-	TMap<TObjectPtr<URCAction>, FRCRangeMapStep> RangeMapActionContainer;
+	TMap<TObjectPtr<URCAction>, FRCRangeMapInput> RangeMapActionContainer;
 
 public:
 	URCRangeMapBehaviour();
@@ -65,33 +65,29 @@ public:
 	/** Refresh function being called whenever either the Controller or the Properties of the Behaviour Details Panel change */
 	void Refresh();
 	
-	/** Called whenever a new action is added into the ActionContainer. Will add a FRCRangeMapStep into RangeMapActionContainer corresponding to the Action. */
+	/** Called whenever a new action is added into the ActionContainer. Will add a FRCRangeMapInput into RangeMapActionContainer corresponding to the Action. */
 	void OnActionAdded(URCAction* Action, URCVirtualPropertySelfContainer* InPropertyValue);
 
 	/** Adds an Action to the ActionContainer, whilst making sure that it is unique. */
-	URCAction* AddAction(const TSharedRef<const FRemoteControlField> InRemoteControlField) override;
+	virtual URCAction* AddAction(const TSharedRef<const FRemoteControlField> InRemoteControlField) override;
 
 	/** Checks whether or not one of the following fields can be added to the AddAction Menu. Makes sure they're unique. */
-	bool CanHaveActionForField(const TSharedPtr<FRemoteControlField> InRemoteControlField) const override;
+	virtual bool CanHaveActionForField(const TSharedPtr<FRemoteControlField> InRemoteControlField) const override;
 
 	/** Returns the Step value associated with a given Action*/
-	bool GetStepValueForAction(const URCAction* InAction, double& OutValue);
+	bool GetValueForAction(const URCAction* InAction, double& OutValue);
 
 private:
 	/** Minimum Value which the Range has */
-	double MinValue;
+	double InputMin;
 
 	/** Maximum Value which the Range has */
-	double MaxValue;
-
-	/** Threshold the Behaviour should adhere to for rounding up and down. */
-	double Threshold;
+	double InputMax;
 
 	/** Controller Value of Type Float, which is used for the purpose of readability. */
 	float ControllerFloatValue;
 	
 private:
-
 	/** Boolean Function to help differentiate Custom Actions */
 	bool IsSupportedActionLerpType(TObjectPtr<URCAction> InAction) const;
 
@@ -108,8 +104,8 @@ private:
 	TMap<double, URCAction*> GetNonLerpActions();
 
 	/** Auxiliary Function used to calculate and apply the Lerp on Structs, like FVectors or FRotators. */
-	void ApplyLerpOnStruct(const FRCRangeMapStep* MinRangeStep, const FRCRangeMapStep* MaxRangeStep, const double& StepAlpha, const FGuid& FieldId);
+	void ApplyLerpOnStruct(const FRCRangeMapInput* MinRangeInput, const FRCRangeMapInput* MaxRangeInput, const double& InputAlpha, const FGuid& FieldId);
 
 	/** Returns whether or not an Action created with a given RemoteControlField can be considered unique. */
-	bool IsActionUnique(const TSharedRef<const FRemoteControlField> InRemoteControlField, const double& InStepValue, const TSet<TObjectPtr<URCAction>>& InActions);
+	bool IsActionUnique(const TSharedRef<const FRemoteControlField> InRemoteControlField, const double& InValue, const TSet<TObjectPtr<URCAction>>& InActions);
 };
diff --git a/Engine/Plugins/VirtualProduction/RemoteControl/Source/RemoteControlUI/Private/UI/Action/RangeMap/RCActionRangeMapModel.cpp b/Engine/Plugins/VirtualProduction/RemoteControl/Source/RemoteControlUI/Private/UI/Action/RangeMap/RCActionRangeMapModel.cpp
--- a/Engine/Plugins/VirtualProduction/RemoteControl/Source/RemoteControlUI/Private/UI/Action/RangeMap/RCActionRangeMapModel.cpp
+++ b/Engine/Plugins/VirtualProduction/RemoteControl/Source/RemoteControlUI/Private/UI/Action/RangeMap/RCActionRangeMapModel.cpp
@@ -30,7 +30,7 @@ namespace UE::RCActionPanelRangeMapList
 	{
 		const FName TypeColorTag = TEXT("TypeColorTag");
 		const FName DragDropHandle = TEXT("DragDropHandle");
-		const FName Step = TEXT("Step");
+		const FName Input = TEXT("Input");
 		const FName Name = TEXT("Name");
 		const FName Value = TEXT("Value");
 	}
@@ -53,9 +53,9 @@ namespace UE::RCActionPanelRangeMapList
 			{
 				return ActionItem->GetTypeColorTagWidget();
 			}
-			else if (ColumnName == UE::RCActionPanelRangeMapList::Columns::Step)
+			else if (ColumnName == UE::RCActionPanelRangeMapList::Columns::Input)
 			{
-				return ActionItem->GetStepWidget();
+				return ActionItem->GetInputWidget();
 			}
 			else if (ColumnName == UE::RCActionPanelRangeMapList::Columns::Name)
 			{
@@ -93,8 +93,8 @@ TSharedPtr<SHeaderRow> FRCActionRangeMapModel::GetHeaderRow()
 		.FixedWidth(25.f)
 		.HeaderContentPadding(RCPanelStyle->HeaderRowPadding)
 
-		+ SHeaderRow::Column(UE::RCActionPanelRangeMapList::Columns::Step)
-		.DefaultLabel(LOCTEXT("RCActionStepColumnHeader", "Step"))
+		+ SHeaderRow::Column(UE::RCActionPanelRangeMapList::Columns::Input)
+		.DefaultLabel(LOCTEXT("RCActionInputColumnHeader", "Input"))
 		.FillWidth(0.2f)
 		.HeaderContentPadding(RCPanelStyle->HeaderRowPadding)
 
@@ -128,18 +128,18 @@ TSharedPtr<FRCActionRangeMapModel> FRCActionRangeMapModel::GetModelByActionType(
 	return nullptr;
 }
 
-TSharedRef<SWidget> FRCActionRangeMapModel::GetStepWidget()
+TSharedRef<SWidget> FRCActionRangeMapModel::GetInputWidget()
 {
-	if (TSharedPtr<FRCRangeMapBehaviourModel> BehaviourItem = StaticCastSharedPtr<FRCRangeMapBehaviourModel>(BehaviourItemWeakPtr.Pin()))
+	if (const TSharedPtr<FRCRangeMapBehaviourModel> BehaviourItem = StaticCastSharedPtr<FRCRangeMapBehaviourModel>(BehaviourItemWeakPtr.Pin()))
 	{
 		if (URCRangeMapBehaviour* Behaviour = Cast<URCRangeMapBehaviour>(BehaviourItem->GetBehaviour()))
 		{
-			if (FRCRangeMapStep* RangeStep = Behaviour->RangeMapActionContainer.Find(GetAction()))
+			if (const FRCRangeMapInput* RangeMapInput = Behaviour->RangeMapActionContainer.Find(GetAction()))
 			{
-				if (ensure(RangeStep->StepValueProperty))
+				if (ensure(RangeMapInput->InputProperty))
 				{
-					return SAssignNew(EditableVirtualPropertyWidget, SRCVirtualPropertyWidget, RangeStep->StepValueProperty)
-						.OnGenerateWidget(this, &FRCActionRangeMapModel::OnGenerateStepWidget);
+					return SAssignNew(EditableVirtualPropertyWidget, SRCVirtualPropertyWidget, RangeMapInput->InputProperty)
+						.OnGenerateWidget(this, &FRCActionRangeMapModel::OnGenerateInputWidget);
 				}
 			}
 		}
@@ -148,7 +148,7 @@ TSharedRef<SWidget> FRCActionRangeMapModel::GetStepWidget()
 	return SNullWidget::NullWidget;
 }
 
-TSharedRef<SWidget> FRCActionRangeMapModel::OnGenerateStepWidget(class URCVirtualPropertySelfContainer* StepValueProperty) const
+TSharedRef<SWidget> FRCActionRangeMapModel::OnGenerateInputWidget(class URCVirtualPropertySelfContainer* StepValueProperty) const
 {
 	if (!ensure(StepValueProperty))
 	{
@@ -159,12 +159,12 @@ TSharedRef<SWidget> FRCActionRangeMapModel::OnGenerateStepWidget(class URCVirtua
 	{
 		if (URCRangeMapBehaviour* Behaviour = Cast<URCRangeMapBehaviour>(BehaviourItem->GetBehaviour()))
 		{
-			if (const FRCRangeMapStep* RangeStep = Behaviour->RangeMapActionContainer.Find(GetAction()))
+			if (const FRCRangeMapInput* RangeInput = Behaviour->RangeMapActionContainer.Find(GetAction()))
 			{
-				double StepValue;
-				if(RangeStep->GetStepValue(StepValue))
+				double Value;
+				if(RangeInput->GetInputValue(Value))
 				{
-					const FText StepAsText = FText::FromString(FString::SanitizeFloat(StepValue));
+					const FText StepAsText = FText::FromString(FString::SanitizeFloat(Value));
 
 					return SNew(SBox)
 						[
diff --git a/Engine/Plugins/VirtualProduction/RemoteControl/Source/RemoteControlUI/Private/UI/Action/RangeMap/RCActionRangeMapModel.h b/Engine/Plugins/VirtualProduction/RemoteControl/Source/RemoteControlUI/Private/UI/Action/RangeMap/RCActionRangeMapModel.h
--- a/Engine/Plugins/VirtualProduction/RemoteControl/Source/RemoteControlUI/Private/UI/Action/RangeMap/RCActionRangeMapModel.h
+++ b/Engine/Plugins/VirtualProduction/RemoteControl/Source/RemoteControlUI/Private/UI/Action/RangeMap/RCActionRangeMapModel.h
@@ -9,10 +9,10 @@
 /*
  * ~ FRCActionRangeMapModel ~
  *
- * UI Model for representing Actions as Steps for the RangeMap Behaviour.
+ * UI Model for representing Actions as Input for the RangeMap Behaviour.
  * Contains and represents, other than name and value, a step value being a value between 0 to 1.
  *
- * <Step Value> <Name> <Value>
+ * <Input Value> <Name> <Value>
  */
 class FRCActionRangeMapModel : public FRCActionModel
 {
@@ -20,9 +20,9 @@ public:
 	FRCActionRangeMapModel(URCAction* InAction, const TSharedPtr<class FRCBehaviourModel> InBehaviourItem, const TSharedPtr<SRemoteControlPanel> InRemoteControlPanel);
 
 	/**
-	 * Simple Widget representing the Normalized Value or Step Value for the Range and this particular Value.
+	 * Simple Widget representing the Normalized Value or Input Value for the Range and this particular Value.
 	 */
-	virtual TSharedRef<SWidget> GetStepWidget();
+	virtual TSharedRef<SWidget> GetInputWidget();
 
 	/** OnGenerateRow delegate for the Actions List View*/
 	TSharedRef<ITableRow> OnGenerateWidgetForList(TSharedPtr<FRCActionRangeMapModel> InItem, const TSharedRef<STableViewBase>& OwnerTable);
@@ -34,7 +34,7 @@ public:
 	static TSharedPtr<FRCActionRangeMapModel> GetModelByActionType(URCAction* InAction, const TSharedPtr<class FRCBehaviourModel> InBehaviourItem, const TSharedPtr<SRemoteControlPanel> InRemoteControlPanel);
 
 private:
-	TSharedRef<SWidget> OnGenerateStepWidget(class URCVirtualPropertySelfContainer* InComparand) const;
+	TSharedRef<SWidget> OnGenerateInputWidget(class URCVirtualPropertySelfContainer* InComparand) const;
 };
 
 /*
diff --git a/Engine/Plugins/VirtualProduction/RemoteControl/Source/RemoteControlUI/Private/UI/Behaviour/Builtin/RangeMap/RCBehaviourRangeMapModel.cpp b/Engine/Plugins/VirtualProduction/RemoteControl/Source/RemoteControlUI/Private/UI/Behaviour/Builtin/RangeMap/RCBehaviourRangeMapModel.cpp
--- a/Engine/Plugins/VirtualProduction/RemoteControl/Source/RemoteControlUI/Private/UI/Behaviour/Builtin/RangeMap/RCBehaviourRangeMapModel.cpp
+++ b/Engine/Plugins/VirtualProduction/RemoteControl/Source/RemoteControlUI/Private/UI/Behaviour/Builtin/RangeMap/RCBehaviourRangeMapModel.cpp
@@ -59,10 +59,10 @@ TSharedRef<SWidget> FRCRangeMapBehaviourModel::GetPropertyWidget() const
 
 	TSharedPtr<IDetailTreeNode> DetailTreeMinValue = DetailTreeItr->Pin();
 	TSharedPtr<IDetailTreeNode> DetailTreeMaxValue = (++DetailTreeItr)->Pin();
-	FNodeWidgets MinValueWidget = DetailTreeMinValue->CreateNodeWidgets();
-	FNodeWidgets MaxValueWidget = DetailTreeMaxValue->CreateNodeWidgets();
+	FNodeWidgets MinInputWidget = DetailTreeMinValue->CreateNodeWidgets();
+	FNodeWidgets MaxInputWidget = DetailTreeMaxValue->CreateNodeWidgets();
 
-	if (MinValueWidget.ValueWidget && MaxValueWidget.ValueWidget)
+	if (MinInputWidget.ValueWidget && MaxInputWidget.ValueWidget)
 	{
 		FieldWidget->AddSlot()
 			.Padding(FMargin(3.0f, 2.0f))
@@ -73,30 +73,6 @@ TSharedRef<SWidget> FRCRangeMapBehaviourModel::GetPropertyWidget() const
 			];
 	}
 	
-	TSharedPtr<IDetailTreeNode> DetailTreeNodeThreshold = (++DetailTreeItr)->Pin();
-	const FNodeWidgets ThresholdWidget = DetailTreeNodeThreshold->CreateNodeWidgets();
-	
-	FieldWidget->AddSlot()
-		.Padding(FMargin(3.0f, 2.0f))
-		.HAlign(HAlign_Center)
-		.AutoHeight()
-		[
-			SNew(SHorizontalBox)
-			+SHorizontalBox::Slot()
-			.Padding(FMargin(3.0f, 2.0f))
-			.HAlign(HAlign_Right)
-			.AutoWidth()
-			[
-				ThresholdWidget.NameWidget.ToSharedRef()
-			]
-			+SHorizontalBox::Slot()
-			.Padding(FMargin(3.0f, 2.0f))
-			.HAlign(HAlign_Right)
-			.AutoWidth()
-			[
-				ThresholdWidget.ValueWidget.ToSharedRef()
-			]
-		];
 
 	return 	SNew(SHorizontalBox)
 		.Clipping(EWidgetClipping::OnDemand)
@@ -131,10 +107,13 @@ URCAction* FRCRangeMapBehaviourModel::AddAction(const TSharedRef<const FRemoteCo
 
 	if (URCRangeMapBehaviour* RangeMapBehaviour = Cast<URCRangeMapBehaviour>(GetBehaviour()))
 	{
+		double StepValue;
+		RangeMapBehaviour->PropertyContainer->GetVirtualProperty(FName("Input"))->GetValueDouble(StepValue);
+
 		const FGuid FieldId = InRemoteControlField->GetId();
-		if (TSharedPtr<FRemoteControlProperty> RemoteProperty = RangeMapBehaviour->ControllerWeakPtr.Get()->PresetWeakPtr.Get()->GetExposedEntity<FRemoteControlProperty>(FieldId).Pin())
+		if (const TSharedPtr<FRemoteControlProperty> RemoteProperty = RangeMapBehaviour->ControllerWeakPtr.Get()->PresetWeakPtr.Get()->GetExposedEntity<FRemoteControlProperty>(FieldId).Pin())
 		{
-			TObjectPtr<URCVirtualPropertySelfContainer> VirtualPropertySelfContainer = NewObject<URCVirtualPropertySelfContainer>();
+			const TObjectPtr<URCVirtualPropertySelfContainer> VirtualPropertySelfContainer = NewObject<URCVirtualPropertySelfContainer>();
 			VirtualPropertySelfContainer->DuplicateProperty(RemoteProperty->GetProperty()->GetFName(), RemoteProperty->GetProperty());
 			
 			NewAction = RangeMapBehaviour->AddAction(InRemoteControlField);
@@ -146,52 +125,52 @@ URCAction* FRCRangeMapBehaviourModel::AddAction(const TSharedRef<const FRemoteCo
 	return NewAction;
 }
 
-TSharedRef<SWidget> FRCRangeMapBehaviourModel::CreateMinMaxWidget(TSharedPtr<IDetailTreeNode> MinRangeDetailTree, TSharedPtr<IDetailTreeNode> MaxRangeDetailTree) const
+TSharedRef<SWidget> FRCRangeMapBehaviourModel::CreateMinMaxWidget(TSharedPtr<IDetailTreeNode> MinInputDetailTree, TSharedPtr<IDetailTreeNode> MaxInputDetailTree) const
 {
-	const FNodeWidgets MinValueWidget = MinRangeDetailTree->CreateNodeWidgets();
-	const FNodeWidgets MaxValueWidget = MaxRangeDetailTree->CreateNodeWidgets();
+	const FNodeWidgets MinInputWidget = MinInputDetailTree->CreateNodeWidgets();
+	const FNodeWidgets MaxInputWidget = MaxInputDetailTree->CreateNodeWidgets();
 	
 	return SNew(SHorizontalBox)
 		+SHorizontalBox::Slot()
 		.Padding(FMargin(3.0f, 2.0f))
-		.HAlign(HAlign_Right)
+		.HAlign(HAlign_Left)
 		.AutoWidth()
 		[
 			SNew(SVerticalBox)
 			+SVerticalBox::Slot()
-			.Padding(FMargin(3.0f, 2.0f))
-			.HAlign(HAlign_Right)
+			.Padding(FMargin(0.0f, 2.0f))
+			.HAlign(HAlign_Left)
 			.AutoHeight()
 			[
-				MinValueWidget.NameWidget.ToSharedRef()
+				MinInputWidget.NameWidget.ToSharedRef()
 			]
 			+SVerticalBox::Slot()
-			.Padding(FMargin(3.0f, 2.0f))
+			.Padding(FMargin(0.0f, 2.0f))
 			.HAlign(HAlign_Right)
 			.AutoHeight()
 			[
-				MinValueWidget.ValueWidget.ToSharedRef()
+				MinInputWidget.ValueWidget.ToSharedRef()
 			]
 		]
 		+SHorizontalBox::Slot()
 		.Padding(FMargin(3.0f, 2.0f))
-		.HAlign(HAlign_Right)
+		.HAlign(HAlign_Left)
 		.AutoWidth()
 		[
 			SNew(SVerticalBox)
 			+SVerticalBox::Slot()
-			.Padding(FMargin(3.0f, 2.0f))
-			.HAlign(HAlign_Right)
+			.Padding(FMargin(0.0f, 2.0f))
+			.HAlign(HAlign_Left)
 			.AutoHeight()
 			[
-				MaxValueWidget.NameWidget.ToSharedRef()
+				MaxInputWidget.NameWidget.ToSharedRef()
 			]
 			+SVerticalBox::Slot()
-			.Padding(FMargin(3.0f, 2.0f))
+			.Padding(FMargin(0.0f, 2.0f))
 			.HAlign(HAlign_Right)
 			.AutoHeight()
 			[
-				MaxValueWidget.ValueWidget.ToSharedRef()
+				MaxInputWidget.ValueWidget.ToSharedRef()
 			]
 		];
 }
diff --git a/Engine/Plugins/VirtualProduction/RemoteControl/Source/RemoteControlUI/Private/UI/Behaviour/Builtin/RangeMap/RCBehaviourRangeMapModel.h b/Engine/Plugins/VirtualProduction/RemoteControl/Source/RemoteControlUI/Private/UI/Behaviour/Builtin/RangeMap/RCBehaviourRangeMapModel.h
--- a/Engine/Plugins/VirtualProduction/RemoteControl/Source/RemoteControlUI/Private/UI/Behaviour/Builtin/RangeMap/RCBehaviourRangeMapModel.h
+++ b/Engine/Plugins/VirtualProduction/RemoteControl/Source/RemoteControlUI/Private/UI/Behaviour/Builtin/RangeMap/RCBehaviourRangeMapModel.h
@@ -50,5 +50,5 @@ private:
 
 private:
 	/** Creates the horizontal box containing box the Minimum and Maximum Range Widgets */
-	TSharedRef<SWidget> CreateMinMaxWidget(TSharedPtr<IDetailTreeNode> MinRangeDetailTree, TSharedPtr<IDetailTreeNode> MaxRangeDetailTree) const;
+	TSharedRef<SWidget> CreateMinMaxWidget(TSharedPtr<IDetailTreeNode> MinInputDetailTree, TSharedPtr<IDetailTreeNode> MaxInputDetailTree) const;
 };
