Fix list invalidation after re-setting an object property on remote control preset

diff --git a/Engine/Plugins/VirtualProduction/RemoteControl/Source/RemoteControlUI/Private/UI/RCPanelWidgetRegistry.cpp b/Engine/Plugins/VirtualProduction/RemoteControl/Source/RemoteControlUI/Private/UI/RCPanelWidgetRegistry.cpp
--- a/Engine/Plugins/VirtualProduction/RemoteControl/Source/RemoteControlUI/Private/UI/RCPanelWidgetRegistry.cpp
+++ b/Engine/Plugins/VirtualProduction/RemoteControl/Source/RemoteControlUI/Private/UI/RCPanelWidgetRegistry.cpp
@@ -92,12 +92,25 @@ FRCPanelWidgetRegistry::FRCPanelWidgetRegistry()
 	SpecialTreeNodeHandlers.Add(MoveTemp(NewHandler));
 }
 
+FRCPanelWidgetRegistry::~FRCPanelWidgetRegistry()
+{
+	for (const TPair <TWeakObjectPtr<UObject>, TSharedPtr<IPropertyRowGenerator>>& Pair : ObjectToRowGenerator)
+	{
+		if (Pair.Value)
+		{
+			Pair.Value->OnRowsRefreshed().RemoveAll(this);
+		}
+	}
+
+	ObjectToRowGenerator.Reset();
+}
+
 TSharedPtr<IDetailTreeNode> FRCPanelWidgetRegistry::GetObjectTreeNode(UObject* InObject, const FString& InField, ERCFindNodeMethod InFindMethod)
 {
 	const TPair<TWeakObjectPtr<UObject>, FString> CacheKey{InObject, InField};
 	if (TWeakPtr<IDetailTreeNode>* Node = TreeNodeCache.Find(CacheKey))
 	{
-		if (Node->IsValid())
+		if (Node->IsValid() && Node->Pin()->CreatePropertyHandle() && Node->Pin()->CreatePropertyHandle()->GetNumOuterObjects() != 0)
 		{
 			return Node->Pin();
 		}
@@ -243,9 +256,9 @@ TSharedPtr<IDetailTreeNode> FRCPanelWidgetRegistry::FindNDisplayTreeNode(UObject
 		FPropertyRowGeneratorArgs Args;
 		Args.bShouldShowHiddenProperties = false;
 		Generator = FModuleManager::GetModuleChecked<FPropertyEditorModule>("PropertyEditor").CreatePropertyRowGenerator(Args);
-		Generator->SetCustomValidatePropertyNodesFunction(FOnValidatePropertyRowGeneratorNodes::CreateLambda(MoveTemp(ValidationLambda)));
-
 		Generator->SetObjects({ InObject });
+		Generator->SetCustomValidatePropertyNodesFunction(FOnValidatePropertyRowGeneratorNodes::CreateLambda(MoveTemp(ValidationLambda)));
+		Generator->OnRowsRefreshed().AddRaw(this, &FRCPanelWidgetRegistry::OnRowsRefreshed, Generator);
 		ObjectToRowGenerator.Add(WeakObject, Generator);
 	}
 
diff --git a/Engine/Plugins/VirtualProduction/RemoteControl/Source/RemoteControlUI/Private/UI/RCPanelWidgetRegistry.h b/Engine/Plugins/VirtualProduction/RemoteControl/Source/RemoteControlUI/Private/UI/RCPanelWidgetRegistry.h
--- a/Engine/Plugins/VirtualProduction/RemoteControl/Source/RemoteControlUI/Private/UI/RCPanelWidgetRegistry.h
+++ b/Engine/Plugins/VirtualProduction/RemoteControl/Source/RemoteControlUI/Private/UI/RCPanelWidgetRegistry.h
@@ -32,6 +32,7 @@ class FRCPanelWidgetRegistry
 public:
 	
 	FRCPanelWidgetRegistry();
+	~FRCPanelWidgetRegistry();
 
 	/**
 	 * Get a detail tree node for a given object and property.
diff --git a/Engine/Plugins/VirtualProduction/RemoteControl/Source/RemoteControlUI/Private/UI/SRCPanelExposedEntitiesList.cpp b/Engine/Plugins/VirtualProduction/RemoteControl/Source/RemoteControlUI/Private/UI/SRCPanelExposedEntitiesList.cpp
--- a/Engine/Plugins/VirtualProduction/RemoteControl/Source/RemoteControlUI/Private/UI/SRCPanelExposedEntitiesList.cpp
+++ b/Engine/Plugins/VirtualProduction/RemoteControl/Source/RemoteControlUI/Private/UI/SRCPanelExposedEntitiesList.cpp
@@ -174,6 +174,8 @@ void SRCPanelExposedEntitiesList::Construct(const FArguments& InArgs, URemoteCon
 	OnEntityListUpdatedDelegate = InArgs._OnEntityListUpdated;
 	WidgetRegistry = MoveTemp(InWidgetRegistry);
 
+	WidgetRegistry.Pin()->OnObjectRefreshed().AddSP(this, &SRCPanelExposedEntitiesList::OnWidgetRegistryRefreshed);
+
 	ColumnSizeData.LeftColumnWidth = TAttribute<float>(this, &SRCPanelExposedEntitiesList::OnGetLeftColumnWidth);
 	ColumnSizeData.RightColumnWidth = TAttribute<float>(this, &SRCPanelExposedEntitiesList::OnGetRightColumnWidth);
 	ColumnSizeData.OnWidthChanged = SSplitter::FOnSlotResized::CreateSP(this, &SRCPanelExposedEntitiesList::OnSetColumnWidth);
@@ -426,6 +428,11 @@ void SRCPanelExposedEntitiesList::Construct(const FArguments& InArgs, URemoteCon
 
 SRCPanelExposedEntitiesList::~SRCPanelExposedEntitiesList()
 {
+	if (TSharedPtr<FRCPanelWidgetRegistry> Registry = WidgetRegistry.Pin())
+	{
+		Registry->OnObjectRefreshed().RemoveAll(this);
+	}
+
 	UnregisterPresetDelegates();
 	UnregisterEvents();
 }
@@ -1573,9 +1580,15 @@ void SRCPanelExposedEntitiesList::OnProtocolBindingAddedOrRemoved(ERCProtocolBin
 	}
 }
 
+void SRCPanelExposedEntitiesList::OnWidgetRegistryRefreshed(const TArray<UObject*>& Objects)
+{
+	Refresh();
+}
+
 bool FGroupDragEvent::IsDraggedFromSameGroup() const
 {
 	return DragOriginGroup->GetId() == DragTargetGroup->GetId();
 }
 
+
 #undef LOCTEXT_NAMESPACE /* RemoteControlPanelFieldList */
\ No newline at end of file
diff --git a/Engine/Plugins/VirtualProduction/RemoteControl/Source/RemoteControlUI/Private/UI/SRCPanelExposedEntitiesList.h b/Engine/Plugins/VirtualProduction/RemoteControl/Source/RemoteControlUI/Private/UI/SRCPanelExposedEntitiesList.h
--- a/Engine/Plugins/VirtualProduction/RemoteControl/Source/RemoteControlUI/Private/UI/SRCPanelExposedEntitiesList.h
+++ b/Engine/Plugins/VirtualProduction/RemoteControl/Source/RemoteControlUI/Private/UI/SRCPanelExposedEntitiesList.h
@@ -182,6 +182,9 @@ private:
 
 	void OnProtocolBindingAddedOrRemoved(ERCProtocolBinding::Op BindingOperation);
 
+	/** Called when the widget registry is refreshed by an underlying property generator. */
+	void OnWidgetRegistryRefreshed(const TArray<UObject*>& Objects);
+
 private:
 	/** Holds the Groups list view. */
 	TSharedPtr<SListView<TSharedPtr<SRCPanelTreeNode>>> GroupsListView;
diff --git a/Engine/Plugins/VirtualProduction/RemoteControl/Source/RemoteControlUI/Private/UI/SRCPanelExposedField.cpp b/Engine/Plugins/VirtualProduction/RemoteControl/Source/RemoteControlUI/Private/UI/SRCPanelExposedField.cpp
--- a/Engine/Plugins/VirtualProduction/RemoteControl/Source/RemoteControlUI/Private/UI/SRCPanelExposedField.cpp
+++ b/Engine/Plugins/VirtualProduction/RemoteControl/Source/RemoteControlUI/Private/UI/SRCPanelExposedField.cpp
@@ -98,14 +98,6 @@ void SRCPanelExposedField::Construct(const FArguments& InArgs, TWeakPtr<FRemoteC
 	}
 }
 
-SRCPanelExposedField::~SRCPanelExposedField()
-{
-	if (TSharedPtr<FRCPanelWidgetRegistry> Registry = WidgetRegistry.Pin())
-	{
-		Registry->OnObjectRefreshed().RemoveAll(this);
-	}
-}
-
 void SRCPanelExposedField::GetNodeChildren(TArray<TSharedPtr<SRCPanelTreeNode>>& OutChildren) const
 {
 	OutChildren.Append(ChildWidgets);
@@ -335,9 +327,6 @@ TSharedRef<SWidget> SRCPanelExposedField::ConstructWidget()
 		{
 			if (TSharedPtr<FRCPanelWidgetRegistry> Registry = WidgetRegistry.Pin())
 			{
-				Registry->OnObjectRefreshed().RemoveAll(this);
-				Registry->OnObjectRefreshed().AddSP(this, &SRCPanelExposedField::OnObjectsRefreshed);
-
 				if (TSharedPtr<IDetailTreeNode> Node = Registry->GetObjectTreeNode(Objects[0], Field->FieldPathInfo.ToPathPropertyString(), ERCFindNodeMethod::Path))
 				{
 					TArray<TSharedRef<IDetailTreeNode>> ChildNodes;
@@ -356,15 +345,17 @@ TSharedRef<SWidget> SRCPanelExposedField::ConstructWidget()
 					];
 
 					TSharedPtr<IPropertyHandle> PropertyHandle = Node->CreatePropertyHandle();
-					if (PropertyHandle && (PropertyHandle->GetParentHandle()->GetPropertyDisplayName().ToString() == FString("Transform")
+					if (PropertyHandle && PropertyHandle->IsValidHandle() && (PropertyHandle->GetParentHandle()->GetPropertyDisplayName().ToString() == FString("Transform")
 						|| PropertyHandle->GetOuterBaseClass() == UMaterialInstanceDynamic::StaticClass()))
 					{
 						// Set up a Zeroed DefaultValue, in case an ExposedEntity doesn't have a native Default. Needed for certain ResetToDefault cases.
 						void* ValuePtr;
-						Node->CreatePropertyHandle()->GetValueData(ValuePtr);
-						DefaultValue.Reset(new uint8[PropertyHandle->GetProperty()->GetSize()]);
-						Node->CreatePropertyHandle()->GetProperty()->CopyCompleteValue(DefaultValue.Get(), ValuePtr);
-						Node->CreatePropertyHandle()->GetProperty()->ClearValue(DefaultValue.Get());
+						if (PropertyHandle->GetValueData(ValuePtr) == FPropertyAccess::Result::Success)
+						{
+							DefaultValue.Reset(new uint8[PropertyHandle->GetProperty()->GetSize()]);
+							PropertyHandle->GetProperty()->CopyCompleteValue(DefaultValue.Get(), ValuePtr);
+							PropertyHandle->GetProperty()->ClearValue(DefaultValue.Get());
+						}
 
 						TWeakPtr<SRCPanelExposedField> WeakFieldPtr = SharedThis(this);
 						auto IsVisible = [WeakFieldPtr, Node]()
@@ -374,11 +365,12 @@ TSharedRef<SWidget> SRCPanelExposedField::ConstructWidget()
 							if (FieldPtr && PropertyHandle && PropertyHandle->IsValidHandle())
 							{
 								void* DataPtr;
-								PropertyHandle->GetValueData(DataPtr);
-								FProperty* NodeProperty = PropertyHandle->GetProperty();
-								
-								bool bVisible = !NodeProperty->Identical(FieldPtr->DefaultValue.Get(), DataPtr);
-								return bVisible ? EVisibility::Visible : EVisibility::Hidden;
+								if (PropertyHandle->GetValueData(DataPtr) == FPropertyAccess::Result::Success)
+								{
+									FProperty* NodeProperty = PropertyHandle->GetProperty();
+									bool bVisible = !NodeProperty->Identical(FieldPtr->DefaultValue.Get(), DataPtr);
+									return bVisible ? EVisibility::Visible : EVisibility::Hidden;
+								}
 							}
 							return EVisibility::Hidden;
 						};
@@ -419,7 +411,7 @@ TSharedRef<SWidget> SRCPanelExposedField::ConstructWidget()
 
 					return MakeFieldWidget(FieldWidget.ToSharedRef());
 				}
-			}
+			} 
 		}
 
 		return MakeFieldWidget(CreateInvalidWidget());
@@ -572,24 +564,6 @@ FReply SRCPanelExposedField::OnClickFunctionButton()
 	return FReply::Handled();
 }
 
-void SRCPanelExposedField::OnObjectsRefreshed(const TArray<UObject*>& RefreshedObjects)
-{
-	if (TSharedPtr<FRemoteControlField> Field = WeakField.Pin())
-	{
-		TSet<UObject*> BoundObjectsSet;
-		BoundObjectsSet.Append(Field->GetBoundObjects());
-
-		TSet<UObject*> RefreshedObjectsSet;
-		RefreshedObjectsSet.Append(RefreshedObjects);
-
-		const TSet<UObject*> CommonObjects = BoundObjectsSet.Intersect(RefreshedObjectsSet);
-		if (CommonObjects.Num())
-		{
-			Refresh();
-		}
-	}
-}
-
 void SRCPanelFieldChildNode::Construct(const FArguments& InArgs, const TSharedRef<IDetailTreeNode>& InNode, FRCColumnSizeData InColumnSizeData)
 {
 	TArray<TSharedRef<IDetailTreeNode>> ChildNodes;
diff --git a/Engine/Plugins/VirtualProduction/RemoteControl/Source/RemoteControlUI/Private/UI/SRCPanelExposedField.h b/Engine/Plugins/VirtualProduction/RemoteControl/Source/RemoteControlUI/Private/UI/SRCPanelExposedField.h
--- a/Engine/Plugins/VirtualProduction/RemoteControl/Source/RemoteControlUI/Private/UI/SRCPanelExposedField.h
+++ b/Engine/Plugins/VirtualProduction/RemoteControl/Source/RemoteControlUI/Private/UI/SRCPanelExposedField.h
@@ -41,7 +41,6 @@ struct SRCPanelExposedField : public SRCPanelExposedEntity
 	static TSharedPtr<SRCPanelTreeNode> MakeInstance(const FGenerateWidgetArgs& Args);
 
 	void Construct(const FArguments& InArgs, TWeakPtr<FRemoteControlField> Field, FRCColumnSizeData ColumnSizeData, TWeakPtr<FRCPanelWidgetRegistry> InWidgetRegistry);
-	virtual ~SRCPanelExposedField();
 
 	//~ SRCPanelTreeNode Interface 
 	virtual void GetNodeChildren(TArray<TSharedPtr<SRCPanelTreeNode>>& OutChildren) const override;
@@ -91,8 +90,6 @@ private:
 	TSharedRef<SWidget> ConstructCallFunctionButton(bool bIsEnabled = true);
 	/** Handles calling an exposed function.*/
 	FReply OnClickFunctionButton();
-	/** Called when the underlying generator gets refreshed. In that case we must re-fetch the widget. */
-	void OnObjectsRefreshed(const TArray<UObject*>& RefreshedObjects);
 
 private:
 	/** Weak pointer to the underlying RC Field. */
diff --git a/Engine/Plugins/VirtualProduction/RemoteControl/Source/RemoteControlUI/Private/UI/SRemoteControlPanel.cpp b/Engine/Plugins/VirtualProduction/RemoteControl/Source/RemoteControlUI/Private/UI/SRemoteControlPanel.cpp
--- a/Engine/Plugins/VirtualProduction/RemoteControl/Source/RemoteControlUI/Private/UI/SRemoteControlPanel.cpp
+++ b/Engine/Plugins/VirtualProduction/RemoteControl/Source/RemoteControlUI/Private/UI/SRemoteControlPanel.cpp
@@ -1864,9 +1864,38 @@ FReply SRemoteControlPanel::OnClickSettingsButton()
 
 void SRemoteControlPanel::OnMaterialCompiled(UMaterialInterface* MaterialInterface)
 {
+	bool bTriggerRefresh = true;
+
 	// Clear the widget cache on material compiled to make sure we have valid property nodes for IPropertyRowGenerator
-	WidgetRegistry->Clear();
-	Refresh();
+	if (TSharedPtr<SWindow> Window = FSlateApplication::Get().FindWidgetWindow(SharedThis(this)))
+	{
+		const TArray<TSharedRef<SWindow>>& ChildWindows = Window->GetChildWindows();
+		// Selecting a material from the RC window can trigger material recompiles, 
+		// If we refresh the widgets right away we might end up closing the material selection window while the user is 
+		// still selecting a material. Therefore we must wait until that window is closed before refreshing the panel.
+		if (ChildWindows.Num())
+		{
+			bTriggerRefresh = false;
+
+			if (!ChildWindows[0]->GetOnWindowClosedEvent().IsBound())
+			{
+				ChildWindows[0]->GetOnWindowClosedEvent().AddLambda([WeakThis = TWeakPtr<SRemoteControlPanel>(SharedThis(this))](const TSharedRef<SWindow>&)
+				{
+					if (TSharedPtr<SRemoteControlPanel> Panel = WeakThis.Pin())
+					{
+						Panel->WidgetRegistry->Clear();
+						Panel->Refresh();
+					}
+				});
+			}
+		}
+	}
+
+	if (bTriggerRefresh)
+	{
+		WidgetRegistry->Clear();
+		Refresh();
+	}
 }
 
 void SRemoteControlPanel::RegisterDefaultToolBar()
