Remote Control Logic resizing behaviour fixes

diff --git a/Engine/Plugins/VirtualProduction/RemoteControl/Source/RemoteControlUI/Private/UI/Action/SRCActionPanel.cpp b/Engine/Plugins/VirtualProduction/RemoteControl/Source/RemoteControlUI/Private/UI/Action/SRCActionPanel.cpp
--- a/Engine/Plugins/VirtualProduction/RemoteControl/Source/RemoteControlUI/Private/UI/Action/SRCActionPanel.cpp
+++ b/Engine/Plugins/VirtualProduction/RemoteControl/Source/RemoteControlUI/Private/UI/Action/SRCActionPanel.cpp
@@ -183,7 +183,9 @@ void SRCActionPanel::UpdateWrappedWidget(TSharedPtr<FRCBehaviourModel> InBehavio
 			.EnableHeader(false)
 			.ChildOrientation(Orient_Vertical);
 
-		ActionsPanel->AddPanel(BehaviourDetailsPanel.ToSharedRef(), 0.5f);
+		// Panel size of zero forces use of "SizeToContent" ensuring that each Behaviour only takes up as much space as necessary
+		ActionsPanel->AddPanel(BehaviourDetailsPanel.ToSharedRef(), 0.f);
+
 		ActionsPanel->AddPanel(ActionDockPanel.ToSharedRef(), 0.5f);
 
 		WrappedBoxWidget->SetContent(ActionsPanel);
diff --git a/Engine/Plugins/VirtualProduction/RemoteControl/Source/RemoteControlUI/Private/UI/Behaviour/Builtin/Path/RCBehaviourSetAssetByPathModel.cpp b/Engine/Plugins/VirtualProduction/RemoteControl/Source/RemoteControlUI/Private/UI/Behaviour/Builtin/Path/RCBehaviourSetAssetByPathModel.cpp
--- a/Engine/Plugins/VirtualProduction/RemoteControl/Source/RemoteControlUI/Private/UI/Behaviour/Builtin/Path/RCBehaviourSetAssetByPathModel.cpp
+++ b/Engine/Plugins/VirtualProduction/RemoteControl/Source/RemoteControlUI/Private/UI/Behaviour/Builtin/Path/RCBehaviourSetAssetByPathModel.cpp
@@ -16,6 +16,7 @@
 #include "Widgets/Input/SButton.h"
 #include "Widgets/Input/SComboBox.h"
 #include "Widgets/Layout/SBox.h"
+#include "Widgets/Layout/SScrollBox.h"
 #include "Widgets/Text/STextBlock.h"
 
 #define LOCTEXT_NAMESPACE "FRCSetAssetByPathBehaviourModel"
@@ -140,7 +141,15 @@ TSharedRef<SWidget> FRCSetAssetByPathBehaviourModel::GetPropertyWidget()
 		.Padding(FMargin(3.f, 3.f))
 		.AutoHeight()
 		[
-			PathArrayWidget->AsShared()
+			SNew(SBox)
+			.MaxDesiredHeight(150.f)
+			[
+				SNew(SScrollBox)
+				+ SScrollBox::Slot()
+				[
+					PathArrayWidget->AsShared()
+				]
+			]
 		];
 
 	FieldWidget->AddSlot()
