Remote Contorl Logic - RangeMap Behaviour Fix Setter Error

Fix a Remote Control Propert Handle Bug where the Pathing wouldn't work for Structs within the Exposed Object.

diff --git a/Engine/Plugins/VirtualProduction/RemoteControl/Source/RemoteControl/Private/RemoteControlPropertyHandle.cpp b/Engine/Plugins/VirtualProduction/RemoteControl/Source/RemoteControl/Private/RemoteControlPropertyHandle.cpp
--- a/Engine/Plugins/VirtualProduction/RemoteControl/Source/RemoteControl/Private/RemoteControlPropertyHandle.cpp
+++ b/Engine/Plugins/VirtualProduction/RemoteControl/Source/RemoteControl/Private/RemoteControlPropertyHandle.cpp
@@ -446,7 +446,14 @@ FRemoteControlPropertyHandle::FRemoteControlPropertyHandle(const TSharedPtr<FRem
 	FString ThisFieldPath = InProperty->GetName();
 	if (InParentProperty == nullptr || InParentProperty->ArrayDim != 1)
 	{
-		FieldPath = ThisFieldPath;
+		if (const UScriptStruct* Struct = Cast<UScriptStruct>(InProperty->GetOwnerUObject()))
+		{
+			FieldPath = FString::Printf(TEXT("%s.%s"), *Struct->GetName(), *ThisFieldPath);
+		}
+		else
+		{
+			FieldPath = ThisFieldPath;
+		}
 	}
 	else
 	{
diff --git a/Engine/Plugins/VirtualProduction/RemoteControl/Source/RemoteControl/Private/Tests/RemoteControlPropertyHandleTestData.h b/Engine/Plugins/VirtualProduction/RemoteControl/Source/RemoteControl/Private/Tests/RemoteControlPropertyHandleTestData.h
--- a/Engine/Plugins/VirtualProduction/RemoteControl/Source/RemoteControl/Private/Tests/RemoteControlPropertyHandleTestData.h
+++ b/Engine/Plugins/VirtualProduction/RemoteControl/Source/RemoteControl/Private/Tests/RemoteControlPropertyHandleTestData.h
@@ -4,6 +4,8 @@
 #include "CoreMinimal.h"
 #include "RemoteControlPropertyHandleTestData.generated.h"
 
+class UStaticMeshComponent;
+
 UENUM()
 enum class ERemoteControlEnumClass : uint8
 {
@@ -94,6 +96,7 @@ public:
 		}
 
 		StringColorMap.Add(TEXT("mykey"), FColor{1,2,3,4});
+		StaticMeshComponent = CreateDefaultSubobject<UStaticMeshComponent>("StaticMeshComponentStructTest");
 	}
 
 	UPROPERTY(EditAnywhere, Category = "RC")
@@ -120,6 +123,9 @@ public:
 	UPROPERTY(EditAnywhere, Category = "RC")
 	TArray<FVector> ArrayOfVectors;
 
+	UPROPERTY(EditAnywhere, Category = "RC")
+	TObjectPtr<UStaticMeshComponent> StaticMeshComponent;
+
 	UPROPERTY(EditAnywhere, Category = "RC")
 	int8 Int8Value = 5;
 
diff --git a/Engine/Plugins/VirtualProduction/RemoteControl/Source/RemoteControl/Private/Tests/RemoteControlPropertyHandleTests.cpp b/Engine/Plugins/VirtualProduction/RemoteControl/Source/RemoteControl/Private/Tests/RemoteControlPropertyHandleTests.cpp
--- a/Engine/Plugins/VirtualProduction/RemoteControl/Source/RemoteControl/Private/Tests/RemoteControlPropertyHandleTests.cpp
+++ b/Engine/Plugins/VirtualProduction/RemoteControl/Source/RemoteControl/Private/Tests/RemoteControlPropertyHandleTests.cpp
@@ -4,6 +4,7 @@
 #include "Misc/AutomationTest.h"
 #include "UObject/StrongObjectPtr.h"
 
+#include "Components/StaticMeshComponent.h"
 #include "IRemoteControlModule.h"
 #include "IRemoteControlPropertyHandle.h"
 #include "RemoteControlPropertyHandleTestData.h"
@@ -43,6 +44,7 @@ namespace RemoteControlAPIIntegrationTest
 
 		// Test the property handles
 		{
+			IRemoteControlModule::Get().RegisterEmbeddedPreset(Preset.Get());
 			// 1. Find the Remote Control preset
 			URemoteControlPreset* ResolvedPreset = IRemoteControlModule::Get().ResolvePreset(PresetName);
 			if (!ResolvedPreset)
@@ -113,6 +115,7 @@ namespace RemoteControlAPIIntegrationTest
 			Test.TestEqual(TEXT("Value in UObject should be the same as a value from Property Handle."), TestObject->Int16Value, SetInt16Value);
 			Test.TestEqual(TEXT("Value in UObject should be the same as a value from Property Handle."), TestObject->Int32Value, SetInt32Value);
 		}
+		IRemoteControlModule::Get().UnregisterEmbeddedPreset(PresetName);
 	}
 
 	void TestCreatePropertyHandleStruct(FAutomationTestBase& Test)
@@ -128,6 +131,7 @@ namespace RemoteControlAPIIntegrationTest
 
 		// Test the api property handles
 		{
+			IRemoteControlModule::Get().RegisterEmbeddedPreset(Preset.Get());
 			// 1. Find the Remote Control preset
 			URemoteControlPreset* ResolvedPreset = IRemoteControlModule::Get().ResolvePreset(PresetName);
 			if (!ResolvedPreset)
@@ -155,9 +159,9 @@ namespace RemoteControlAPIIntegrationTest
 
 			// 4. Check num Children
 			int32 NumChildren = StructOuterPropertyHandle->GetNumChildren();
-			Test.TestEqual(TEXT("Num children should be equal 4"), 4, NumChildren);
+			Test.TestEqual(TEXT("Num children should be equal 4"), NumChildren, 4);
 			NumChildren = RemoteControlTestStructInnerPropertyHandle->GetNumChildren();
-			Test.TestEqual(TEXT("Num children should be equal 3"), 3, NumChildren);
+			Test.TestEqual(TEXT("Num children should be equal 3"), NumChildren, 3);
 
 			// 5. Get/Set value to the outer struct RemoteControlAPITestObject/RemoteControlTestStructOuter
 			{
@@ -207,6 +211,7 @@ namespace RemoteControlAPIIntegrationTest
 				}
 			}
 		}
+		IRemoteControlModule::Get().UnregisterEmbeddedPreset(PresetName);
 	}
 
 	void TestCreatePropertyHandleArray(FAutomationTestBase& Test)
@@ -777,6 +782,7 @@ namespace RemoteControlAPIIntegrationTest
 		// Test complex path handles
 		{
 			// 1. Find the Remote Control preset
+			IRemoteControlModule::Get().RegisterEmbeddedPreset(Preset.Get());
 			URemoteControlPreset* ResolvedPreset = IRemoteControlModule::Get().ResolvePreset(PresetName);
 			if (!ResolvedPreset)
 			{
@@ -861,6 +867,7 @@ namespace RemoteControlAPIIntegrationTest
 					->InnerSimple.Int32Value,
 				SetInt32Value);
 		}
+		IRemoteControlModule::Get().UnregisterEmbeddedPreset(PresetName);
 	}
 
 	void TestGetPropertyHandleByFieldPath(FAutomationTestBase& Test)
@@ -874,6 +881,8 @@ namespace RemoteControlAPIIntegrationTest
 		TStrongObjectPtr<URemoteControlAPITestObject> TestObject{ NewObject<URemoteControlAPITestObject>() };
 		check(TestObject);
 
+		IRemoteControlModule::Get().RegisterEmbeddedPreset(Preset.Get());
+
 		// Set a test value
 		constexpr int32 Int32ValueTest = 14923;
 
@@ -976,6 +985,72 @@ namespace RemoteControlAPIIntegrationTest
 			IntMap_PropertyHandle_O->SetValue(SetIntValue);
 			Test.TestEqual(TEXT("Value in UObject should be the same as a value from Property Handle."), *TestObject->IntMap.Find(1), SetIntValue);
 		}
+		IRemoteControlModule::Get().UnregisterEmbeddedPreset(PresetName);
+	}
+
+	void TestStructInObjectProperty(FAutomationTestBase& Test)
+	{
+		// Create preset and uobject
+		const TCHAR* PresetName = TEXT("TestComplexPath");
+
+		TStrongObjectPtr<URemoteControlPreset> Preset{ NewObject<URemoteControlPreset>(GetTransientPackage(), PresetName) };
+		check(Preset);
+		
+		TStrongObjectPtr<URemoteControlAPITestObject> TestObject{ NewObject<URemoteControlAPITestObject>() };
+		check(TestObject);
+
+		IRemoteControlModule::Get().RegisterEmbeddedPreset(Preset.Get());
+
+		// Set a test value
+		TestObject->StaticMeshComponent->BodyInstance.LinearDamping = 5.0;
+
+		FProperty* LinearProperty = FBodyInstance::StaticStruct()->FindPropertyByName(PROP_NAME(FBodyInstance, LinearDamping));
+		if(!Test.TestFalse(TEXT("LineaProperty is valid"), LinearProperty == nullptr))
+		{
+			return;
+		}
+
+		FRemoteControlPresetExposeArgs ExposeArgs;
+		ExposeArgs.Label = "LinearDamping";
+		TWeakPtr<FRemoteControlProperty> LinearDampingRC = Preset->ExposeProperty(TestObject->StaticMeshComponent, FRCFieldPathInfo{ "BodyInstance." + LinearProperty->GetName() }, ExposeArgs);
+		
+		{
+			// 1. Find the Remote Control preset
+			URemoteControlPreset* ResolvedPreset = IRemoteControlModule::Get().ResolvePreset(PresetName);
+			if (!ResolvedPreset)
+			{
+				Test.AddError(TEXT("ResolvedPreset not valid"));
+				return;
+			}
+
+			// 2. Find exposed properties by label or ID
+			const TSharedPtr<FRemoteControlProperty> RCLinearDampingProperty = ResolvedPreset->GetExposedEntity<FRemoteControlProperty>(ResolvedPreset->GetExposedEntityId(*LinearProperty->GetName())).Pin();
+			if (!RCLinearDampingProperty.IsValid())
+			{
+				Test.AddError(TEXT("RCLinearDampingProperty not valid"));
+				return;
+			}
+
+			// 3. Get property handle
+			const TSharedPtr<IRemoteControlPropertyHandle> RCLinearDampingPropertyHandle = RCLinearDampingProperty->GetPropertyHandle();
+			if (!RCLinearDampingPropertyHandle.IsValid())
+			{
+				Test.AddError(TEXT("RCLinearDampingPropertyHandle not valid"));
+				return;
+			}
+
+			// 4. Get remote control handles values
+			float LinearDampingValue = 0.0;
+			RCLinearDampingPropertyHandle->GetValue(LinearDampingValue);
+			Test.TestEqual(TEXT("Value in UObject should be the same as a value from Property Handle."), TestObject->StaticMeshComponent->BodyInstance.LinearDamping, LinearDampingValue);
+
+			// 5. Set remote control handles values
+			constexpr float SetLinearDampingValue = 100.0;
+			RCLinearDampingPropertyHandle->SetValue(SetLinearDampingValue);
+			Test.TestEqual(TEXT("Value in UObject should be the same as a value from Property Handle."), TestObject->StaticMeshComponent->BodyInstance.LinearDamping, SetLinearDampingValue);
+		}
+
+		IRemoteControlModule::Get().UnregisterEmbeddedPreset(Preset.Get());
 	}
 }
 
@@ -1001,6 +1076,7 @@ bool FRemoteControlAPIIntegrationTest::RunTest(const FString& Parameters)
 	RemoteControlAPIIntegrationTest::TestLinearColorProperty(*this);
 	RemoteControlAPIIntegrationTest::TestComplexPath(*this);
 	RemoteControlAPIIntegrationTest::TestGetPropertyHandleByFieldPath(*this);
+	RemoteControlAPIIntegrationTest::TestStructInObjectProperty(*this);
 
 	return true;
 }
diff --git a/Engine/Plugins/VirtualProduction/RemoteControl/Source/RemoteControlLogic/Private/Behaviour/Builtin/RangeMap/RCRangeMapBehaviour.cpp b/Engine/Plugins/VirtualProduction/RemoteControl/Source/RemoteControlLogic/Private/Behaviour/Builtin/RangeMap/RCRangeMapBehaviour.cpp
--- a/Engine/Plugins/VirtualProduction/RemoteControl/Source/RemoteControlLogic/Private/Behaviour/Builtin/RangeMap/RCRangeMapBehaviour.cpp
+++ b/Engine/Plugins/VirtualProduction/RemoteControl/Source/RemoteControlLogic/Private/Behaviour/Builtin/RangeMap/RCRangeMapBehaviour.cpp
@@ -174,11 +174,11 @@ void URCRangeMapBehaviour::GetLerpActions(TMap<FGuid, TArray<URCAction*>>& OutNu
 		{
 			double A, B;
 			bool bRes1 = GetValueForAction(ActionA, A);
-			bool bRes2 = GetValueForAction(ActionA, B);
+			bool bRes2 = GetValueForAction(ActionB, B);
 
 			if (bRes1 && bRes2)
 			{
-				return A > B;
+				return A < B;
 			}
 
 			return false;
@@ -387,9 +387,9 @@ bool URCRangeMapBehaviour::GetRangeValuePairsForLerp(TMap<FGuid, TTuple<URCActio
 		URCAction* MaxAction = nullptr;
 		URCAction* MinAction = nullptr;
 		
-		for (URCAction* Action : NumericActionTuple.Value)
+		for (URCAction* CurrentAction : NumericActionTuple.Value)
 		{
-			FRCRangeMapInput* RangeMap = RangeMapActionContainer.Find(Action);
+			FRCRangeMapInput* RangeMap = RangeMapActionContainer.Find(CurrentAction);
 			if (!RangeMap)
 			{
 				// Shouldn't happen, but skip if it does
@@ -404,9 +404,9 @@ bool URCRangeMapBehaviour::GetRangeValuePairsForLerp(TMap<FGuid, TTuple<URCActio
 			}
 
 			// Deal with the nullptr
-			if (NormalizedControllerValue >= Value && (!MinRangeMap || NumericActionTuple.Value.Last() != Action))
+			if (Value <= NormalizedControllerValue && NumericActionTuple.Value.Last() != CurrentAction)
 			{
-				MinAction = Action;
+				MinAction = CurrentAction;
 				MinRangeMap = RangeMap;
 				
 				continue;
@@ -414,7 +414,7 @@ bool URCRangeMapBehaviour::GetRangeValuePairsForLerp(TMap<FGuid, TTuple<URCActio
 
 			if (NormalizedControllerValue <= Value)
 			{
-				MaxAction = Action;
+				MaxAction = CurrentAction;
 				MaxRangeMap = RangeMap;
 
 				continue;
