Remote Control Logic - WebAPI Extension for RC Logc Controllers

Adds the required Callbacks and missing Routes for the WebApp, so that it can control the VirtualProperties/Controllers.
Following changes have been made:
- Add OnAdded/OnChanged/OnRemoved/OnRenamed Callback for Controllers
- Add Websocket Messages from those Callbacks to the Server
- Add Metadata Callbacks and implementation of imoprtant information such as "Widget" Metadata
- Fix up existing calls to Deserializer to work for structs
- Renaming VirtualProperty Calls/Methods within the Preset to Controllers

diff --git a/Engine/Plugins/VirtualProduction/RemoteControl/Source/RemoteControl/Private/RCVirtualProperty.cpp b/Engine/Plugins/VirtualProduction/RemoteControl/Source/RemoteControl/Private/RCVirtualProperty.cpp
--- a/Engine/Plugins/VirtualProduction/RemoteControl/Source/RemoteControl/Private/RCVirtualProperty.cpp
+++ b/Engine/Plugins/VirtualProduction/RemoteControl/Source/RemoteControl/Private/RCVirtualProperty.cpp
@@ -102,6 +102,23 @@ EPropertyBagPropertyType URCVirtualPropertyBase::GetValueType() const
 	return EPropertyBagPropertyType::None;
 }
 
+FString URCVirtualPropertyBase::GetMetadataValue(FName Key) const
+{
+	const FString* MapString = Metadata.Find(Key);
+	return MapString ? *MapString : "";
+}
+
+void URCVirtualPropertyBase::SetMetadataValue(FName Key, FString Data)
+{
+	Metadata.FindOrAdd(Key) = Data;
+}
+
+void URCVirtualPropertyBase::RemoveMetadataValue(FName Key)
+{
+	FString OutString;
+	Metadata.RemoveAndCopyValue(Key, OutString);
+}
+
 const UObject* URCVirtualPropertyBase::GetValueTypeObjectWeakPtr() const
 {
 	if (const FPropertyBagPropertyDesc* PropertyBagPropertyDesc = GetBagPropertyDesc())
@@ -121,17 +138,29 @@ void URCVirtualPropertyBase::SerializeToBackend(IStructSerializerBackend& OutBac
 
 bool URCVirtualPropertyBase::DeserializeFromBackend(IStructDeserializerBackend& InBackend)
 {
+	const UPropertyBag* TypeInfo = GetPropertyBagInstance()->GetPropertyBagStruct();
+	UPropertyBag* TypeInfoNonConst = const_cast<UPropertyBag*>(TypeInfo);
+
+	bool bSuccess = false;
 	FStructDeserializerPolicies Policies;
 	const FProperty* Property = GetProperty();
-	Policies.PropertyFilter = [&Property](const FProperty* CurrentProp, const FProperty* ParentProp)
+	if (Property->GetClass() == FStructProperty::StaticClass())
 	{
-		return CurrentProp == Property;
-	};
-
-	const UPropertyBag* TypeInfo = GetPropertyBagInstance()->GetPropertyBagStruct();
-	UPropertyBag* TypeInfoNonConst = const_cast<UPropertyBag*>(TypeInfo);
+		Policies.PropertyFilter = [&Property](const FProperty* CurrentProp, const FProperty* ParentProp)
+		{
+			return CurrentProp == Property || ParentProp == Property;
+		};
+		bSuccess = FStructDeserializer::Deserialize(GetContainerPtr(), *TypeInfoNonConst, InBackend, Policies);
+	}
+	else
+	{
+		Policies.PropertyFilter = [&Property](const FProperty* CurrentProp, const FProperty* ParentProp)
+		{
+			return CurrentProp == Property;
+		};
+		bSuccess = FStructDeserializer::DeserializeElement(GetContainerPtr(), *TypeInfoNonConst, INDEX_NONE, InBackend, Policies);
+	}
 
-	bool bSuccess = FStructDeserializer::DeserializeElement(GetContainerPtr(), *TypeInfoNonConst, INDEX_NONE, InBackend, Policies);
 	if (bSuccess)
 	{
 		OnModifyPropertyValue();
diff --git a/Engine/Plugins/VirtualProduction/RemoteControl/Source/RemoteControl/Private/RemoteControlModule.cpp b/Engine/Plugins/VirtualProduction/RemoteControl/Source/RemoteControl/Private/RemoteControlModule.cpp
--- a/Engine/Plugins/VirtualProduction/RemoteControl/Source/RemoteControl/Private/RemoteControlModule.cpp
+++ b/Engine/Plugins/VirtualProduction/RemoteControl/Source/RemoteControl/Private/RemoteControlModule.cpp
@@ -1364,7 +1364,7 @@ bool FRemoteControlModule::SetPresetController(const FName PresetName, const FNa
 		return false;
 	}
 
-	URCVirtualPropertyBase* VirtualProperty = Preset->GetVirtualPropertyByDisplayName(ControllerName);
+	URCVirtualPropertyBase* VirtualProperty = Preset->GetControllerByDisplayName(ControllerName);
 	if (!ensure(VirtualProperty))
 	{
 		return false;
diff --git a/Engine/Plugins/VirtualProduction/RemoteControl/Source/RemoteControl/Private/RemoteControlPreset.cpp b/Engine/Plugins/VirtualProduction/RemoteControl/Source/RemoteControl/Private/RemoteControlPreset.cpp
--- a/Engine/Plugins/VirtualProduction/RemoteControl/Source/RemoteControl/Private/RemoteControlPreset.cpp
+++ b/Engine/Plugins/VirtualProduction/RemoteControl/Source/RemoteControl/Private/RemoteControlPreset.cpp
@@ -509,7 +509,7 @@ FName URemoteControlPreset::GenerateUniqueLabel(const FName InDesiredName) const
 	return Registry->GenerateUniqueLabel(InDesiredName);
 }
 
-URCVirtualPropertyBase* URemoteControlPreset::GetVirtualProperty(const FName InPropertyName) const
+URCVirtualPropertyBase* URemoteControlPreset::GetController(const FName InPropertyName) const
 {
 	if (!ensure(ControllerContainer))
 	{
@@ -519,7 +519,7 @@ URCVirtualPropertyBase* URemoteControlPreset::GetVirtualProperty(const FName InP
 	return ControllerContainer->GetVirtualProperty(InPropertyName);
 }
 
-URCVirtualPropertyBase* URemoteControlPreset::GetVirtualProperty(const FGuid& InId) const
+URCVirtualPropertyBase* URemoteControlPreset::GetController(const FGuid& InId) const
 {
 	if (!ensure(ControllerContainer))
 	{
@@ -529,7 +529,12 @@ URCVirtualPropertyBase* URemoteControlPreset::GetVirtualProperty(const FGuid& In
 	return ControllerContainer->GetVirtualProperty(InId);
 }
 
-URCVirtualPropertyBase* URemoteControlPreset::GetVirtualPropertyByDisplayName(const FName InDisplayName) const
+TArray<URCVirtualPropertyBase*> URemoteControlPreset::GetControllers() const
+{
+	return ControllerContainer->VirtualProperties.Array();
+}
+
+URCVirtualPropertyBase* URemoteControlPreset::GetControllerByDisplayName(const FName InDisplayName) const
 {
 	if (!ensure(ControllerContainer))
 	{
@@ -539,7 +544,7 @@ URCVirtualPropertyBase* URemoteControlPreset::GetVirtualPropertyByDisplayName(co
 	return ControllerContainer->GetVirtualPropertyByDisplayName(InDisplayName);
 }
 
-URCVirtualPropertyInContainer* URemoteControlPreset::AddVirtualProperty(TSubclassOf<URCVirtualPropertyInContainer> InPropertyClass, const EPropertyBagPropertyType InValueType, UObject* InValueTypeObject /*= nullptr*/, const FName InPropertyName /*= NAME_None*/)
+URCVirtualPropertyInContainer* URemoteControlPreset::AddController(TSubclassOf<URCVirtualPropertyInContainer> InPropertyClass, const EPropertyBagPropertyType InValueType, UObject* InValueTypeObject /*= nullptr*/, const FName InPropertyName /*= NAME_None*/)
 {
 	if (!ensure(ControllerContainer))
 	{
@@ -557,10 +562,17 @@ URCVirtualPropertyInContainer* URemoteControlPreset::AddVirtualProperty(TSubclas
 	ControllerContainer->Modify();
 #endif
 
-	return ControllerContainer->AddProperty(NewPropertyName, InPropertyClass, InValueType, InValueTypeObject);
+	URCVirtualPropertyInContainer* NewController = ControllerContainer->AddProperty(NewPropertyName, InPropertyClass, InValueType, InValueTypeObject);
+	if (NewController)
+	{
+		OnControllerAdded().Broadcast(this, NewPropertyName, NewController->Id);
+		InitializeEntityMetadata(NewController);
+	}
+	
+	return NewController;
 }
 
-bool URemoteControlPreset::RemoveVirtualProperty(const FName& InPropertyName)
+bool URemoteControlPreset::RemoveController(const FName& InPropertyName)
 {
 	if (ensure(ControllerContainer))
 	{
@@ -568,23 +580,29 @@ bool URemoteControlPreset::RemoveVirtualProperty(const FName& InPropertyName)
 		ControllerContainer->Modify();
 #endif
 
-		return ControllerContainer->RemoveProperty(InPropertyName);
+		if (const URCVirtualPropertyBase* ControllerToDelete = ControllerContainer->GetVirtualProperty(InPropertyName))
+		{
+			ControllerContainer->RemoveProperty(ControllerToDelete->GetPropertyName());
+			OnControllerRemoved().Broadcast(this, ControllerToDelete->Id);
+			return true;
+		}
 	}
 
 	return false;
 }
 
-URCVirtualPropertyInContainer* URemoteControlPreset::DuplicateVirtualProperty(URCVirtualPropertyInContainer* InVirtualProperty)
+URCVirtualPropertyInContainer* URemoteControlPreset::DuplicateController(URCVirtualPropertyInContainer* InVirtualProperty)
 {
 	if (ensure(ControllerContainer))
 	{
-		return ControllerContainer->DuplicateVirtualProperty(InVirtualProperty);
+		URCVirtualPropertyInContainer* DuplicatedController = ControllerContainer->DuplicateVirtualProperty(InVirtualProperty);
+		return DuplicatedController;
 	}
 
 	return nullptr;
 }
 
-void URemoteControlPreset::ResetVirtualProperties()
+void URemoteControlPreset::ResetControllers()
 {
 	if (ensure(ControllerContainer))
 	{
@@ -592,11 +610,16 @@ void URemoteControlPreset::ResetVirtualProperties()
 		ControllerContainer->Modify();
 #endif
 
+		// Broadcasting each removal.
+		for (TObjectPtr<URCVirtualPropertyBase> Controller : ControllerContainer->VirtualProperties)
+		{
+			OnControllerRemoved().Broadcast(this, Controller->Id);
+		}
 		ControllerContainer->Reset();
 	}
 }
 
-int32 URemoteControlPreset::GetNumVirtualProperties() const
+int32 URemoteControlPreset::GetNumControllers() const
 {
 	if (ensure(ControllerContainer))
 	{
@@ -631,11 +654,15 @@ void URemoteControlPreset::OnNotifyPreChangeVirtualProperty(const FPropertyChang
 	}
 }
 
-void URemoteControlPreset::OnModifyVirtualProperty(const FPropertyChangedEvent& PropertyChangedEvent)
+void URemoteControlPreset::OnModifyController(const FPropertyChangedEvent& PropertyChangedEvent)
 {
 	if (ensure(ControllerContainer))
 	{
-		ControllerContainer->OnModifyPropertyValue(PropertyChangedEvent);
+		if (URCVirtualPropertyBase* ModifiedController = ControllerContainer->GetVirtualProperty(PropertyChangedEvent.Property->GetFName()))
+		{
+			ControllerContainer->OnModifyPropertyValue(PropertyChangedEvent);
+			OnControllerModified().Broadcast(this, {ModifiedController->Id});
+		}
 	}	
 } 
 #endif
@@ -890,6 +917,11 @@ void URemoteControlPreset::InitializeEntitiesMetadata()
 	{
 		InitializeEntityMetadata(Entity);
 	}
+
+	for (URCVirtualPropertyBase* Controller : GetControllers())
+	{
+		InitializeEntityMetadata(Controller);
+	}
 }
 
 void URemoteControlPreset::InitializeEntityMetadata(const TSharedPtr<FRemoteControlEntity>& Entity)
@@ -913,6 +945,27 @@ void URemoteControlPreset::InitializeEntityMetadata(const TSharedPtr<FRemoteCont
     }
 }
 
+void URemoteControlPreset::InitializeEntityMetadata(URCVirtualPropertyBase* Controller)
+{
+	if (!Controller)
+	{
+		return;
+	}
+    	
+    const TMap<FName, FEntityMetadataInitializer>& Initializers = IRemoteControlModule::Get().GetDefaultMetadataInitializers();
+    for (const TPair<FName, FEntityMetadataInitializer>& Entry : Initializers)
+    {
+    	if (Entry.Value.IsBound())
+    	{
+    		// Don't reset the metadata entry if already present.
+    		if (!Controller->Metadata.Contains(Entry.Key))
+    		{
+    			Controller->Metadata.Add(Entry.Key, Entry.Value.Execute(this, Controller->Id));
+    		}
+    	}
+    }
+}
+
 void URemoteControlPreset::RegisterEntityDelegates()
 {
 	for (const TSharedPtr<FRemoteControlEntity>& Entity : Registry->GetExposedEntities())
@@ -1295,9 +1348,9 @@ void URemoteControlPreset::Unexpose(const FGuid& EntityId)
 {
 	if (EntityId.IsValid() && Registry->GetExposedEntity(EntityId).IsValid())
 	{
-		OnEntityUnexposedDelegate.Broadcast(this, EntityId);
 		Registry->Modify();
 		Registry->RemoveExposedEntity(EntityId);
+		OnEntityUnexposedDelegate.Broadcast(this, EntityId);
 		FRCCachedFieldData CachedData = FieldCache.FindChecked(EntityId);
 		Layout.RemoveField(CachedData.LayoutGroupId, EntityId);
 		FieldCache.Remove(EntityId);
diff --git a/Engine/Plugins/VirtualProduction/RemoteControl/Source/RemoteControl/Public/RCVirtualProperty.h b/Engine/Plugins/VirtualProduction/RemoteControl/Source/RemoteControl/Public/RCVirtualProperty.h
--- a/Engine/Plugins/VirtualProduction/RemoteControl/Source/RemoteControl/Public/RCVirtualProperty.h
+++ b/Engine/Plugins/VirtualProduction/RemoteControl/Source/RemoteControl/Public/RCVirtualProperty.h
@@ -63,6 +63,15 @@ public:
 	/** Return property bag property type */
 	EPropertyBagPropertyType GetValueType() const;
 
+	/** Get Metadata from the internal PropertyBag */
+	FString GetMetadataValue(FName Key) const;
+
+	/** Add/Set Metadata into the internal PropertyBag */
+	void SetMetadataValue(FName Key, FString Data);
+
+	/** Remove Metadata from the internal PropertyBag */
+	void RemoveMetadataValue(FName Key);
+
 	/** Return pointer to object that defines the Enum, Struct, or Class. */
 	const UObject* GetValueTypeObjectWeakPtr() const;
 
@@ -304,6 +313,9 @@ public:
 	/** User configurable Display Index for this Virtual Property (as Logic Controller) when represented as a row in the RC Logic Controllers list  */
 	UPROPERTY()
 	int32 DisplayIndex;
+
+	UPROPERTY()
+	TMap<FName, FString> Metadata;
 };
 
 /**
diff --git a/Engine/Plugins/VirtualProduction/RemoteControl/Source/RemoteControl/Public/RemoteControlPreset.h b/Engine/Plugins/VirtualProduction/RemoteControl/Source/RemoteControl/Public/RemoteControlPreset.h
--- a/Engine/Plugins/VirtualProduction/RemoteControl/Source/RemoteControl/Public/RemoteControlPreset.h
+++ b/Engine/Plugins/VirtualProduction/RemoteControl/Source/RemoteControl/Public/RemoteControlPreset.h
@@ -585,6 +585,18 @@ public:
 	DECLARE_MULTICAST_DELEGATE_OneParam(FOnPresetLayoutModified, URemoteControlPreset* /*Preset*/);
 	FOnPresetLayoutModified& OnPresetLayoutModified() { return OnPresetLayoutModifiedDelegate; }
 
+	DECLARE_MULTICAST_DELEGATE_ThreeParams(FOnControllerAdded, URemoteControlPreset* /*Preset*/, const FName /*NewControllerName*/, const FGuid& /*ControllerId*/);
+	FOnControllerAdded& OnControllerAdded() { return OnControllerAddedDelegate; }
+
+	DECLARE_MULTICAST_DELEGATE_TwoParams(FOnControllerRemoved, URemoteControlPreset* /*Preset*/, const FGuid& /*ControllerId*/);
+	FOnControllerRemoved& OnControllerRemoved() { return OnControllerRemovedDelegate; }
+
+	DECLARE_MULTICAST_DELEGATE_ThreeParams(FOnControllerRenamed, URemoteControlPreset* /*Preset*/, const FName /*OldLabel*/, const FName /*NewLabel*/);
+	FOnControllerRenamed& OnControllerRenamed() { return OnControllerRenamedDelegate; }
+	
+	DECLARE_MULTICAST_DELEGATE_TwoParams(FOnControllerModified, URemoteControlPreset* /*Preset*/, const TSet<FGuid>& /*ModifiedControllerIds*/);
+	FOnControllerModified& OnControllerModified() { return OnControllerModifiedDelegate; }
+	
 	UE_DEPRECATED(4.27, "This function is deprecated.")
 	void NotifyExposedPropertyChanged(FName PropertyLabel);
 
@@ -620,29 +632,32 @@ public:
 	* The goal is to hide the Controller Container and provide a simple interface for Controller access to UI and Web.
 	*/
 
-	/** Fetches a virtual property by internal property name. */
-	URCVirtualPropertyBase* GetVirtualProperty(const FName InPropertyName) const;
+	/** Fetches a controller by internal property name. */
+	URCVirtualPropertyBase* GetController(const FName InPropertyName) const;
+
+	/** Fetches a controller by unique Id. */
+	URCVirtualPropertyBase* GetController(const FGuid& InId) const;
 
-	/** Fetches a virtual property by unique Id. */
-	URCVirtualPropertyBase* GetVirtualProperty(const FGuid& InId) const;
+	/** Fetches all controller */
+	TArray<URCVirtualPropertyBase*> GetControllers() const;
 
 	/** Fetches a virtual property by specified name. */
-	URCVirtualPropertyBase* GetVirtualPropertyByDisplayName(const FName InDisplayName) const;
+	URCVirtualPropertyBase* GetControllerByDisplayName(const FName InDisplayName) const;
 
 	/** Adds a Virtual Property (Controller) to the Remote Control Preset */
-	URCVirtualPropertyInContainer* AddVirtualProperty(TSubclassOf<URCVirtualPropertyInContainer> InPropertyClass, const EPropertyBagPropertyType InValueType, UObject* InValueTypeObject = nullptr, const FName InPropertyName = NAME_None);
+	URCVirtualPropertyInContainer* AddController(TSubclassOf<URCVirtualPropertyInContainer> InPropertyClass, const EPropertyBagPropertyType InValueType, UObject* InValueTypeObject = nullptr, const FName InPropertyName = NAME_None);
 
 	/** Removes a given Virtual Property (by Name) from the Remote Control preset */
-	bool RemoveVirtualProperty(const FName& InPropertyName);
+	bool RemoveController(const FName& InPropertyName);
 
 	/** Duplicates a given Virtual Property from the Remote Control preset */
-	URCVirtualPropertyInContainer* DuplicateVirtualProperty(URCVirtualPropertyInContainer* InVirtualProperty);
+	URCVirtualPropertyInContainer* DuplicateController(URCVirtualPropertyInContainer* InVirtualProperty);
 
 	/** Removes all virtual properties held by this Remote Control preset*/
-	void ResetVirtualProperties();
+	void ResetControllers();
 
 	/** Returns the number of Virtual Properties contained in this Remote Control preset*/
-	int32 GetNumVirtualProperties() const;
+	int32 GetNumControllers() const;
 
 	/** Returns the Struct On Scope of the Controller Container (value ptr of the virtual properties)*
 	* Currently used by the UI class SRCControllerPanelList for user value entry via the RC Controllers panel*/
@@ -663,7 +678,7 @@ public:
 	void OnNotifyPreChangeVirtualProperty(const FPropertyChangedEvent& PropertyChangedEvent);
 
 	/** Called when a virtual property is modified. This call is routed to the Controller for evaluating associated Logic & Behaviours*/
-	void OnModifyVirtualProperty(const FPropertyChangedEvent& PropertyChangedEvent);
+	void OnModifyController(const FPropertyChangedEvent& PropertyChangedEvent);
 #endif
 
 	FOnVirtualPropertyContainerModified& OnVirtualPropertyContainerModified() const;
@@ -745,6 +760,9 @@ private:
 	
 	/** Initialize an entity's metadata based on the module's externally registered initializers. */
 	void InitializeEntityMetadata(const TSharedPtr<FRemoteControlEntity>& Entity);
+	
+	/** Initialize an controllers's metadata based on the module's externally registered initializers. */
+	void InitializeEntityMetadata(URCVirtualPropertyBase* Controller);
 
 	/** Register delegates for all exposed entities. */
 	void RegisterEntityDelegates();
@@ -807,6 +825,14 @@ private:
 	FOnActorPropertyModified OnActorPropertyModifiedDelegate;
 	/** Delegate triggered when the layout is modified. */
 	FOnPresetLayoutModified OnPresetLayoutModifiedDelegate;
+	/** Delegate triggered when a controller is added */
+	FOnControllerAdded OnControllerAddedDelegate;
+	/** Delegate triggered when a controller is removed */
+	FOnControllerRemoved OnControllerRemovedDelegate;
+	/** Delegate triggered when a Controller is renamed */
+	FOnControllerRenamed OnControllerRenamedDelegate;
+	/** Delegate triggered when a Controller has been changed */
+	FOnControllerModified OnControllerModifiedDelegate;
 
 	struct FPreObjectsModifiedCache
 	{
diff --git a/Engine/Plugins/VirtualProduction/RemoteControl/Source/RemoteControlLogic/Private/Behaviour/Builtin/Path/RCSetAssetByPathBehaviour.cpp b/Engine/Plugins/VirtualProduction/RemoteControl/Source/RemoteControlLogic/Private/Behaviour/Builtin/Path/RCSetAssetByPathBehaviour.cpp
--- a/Engine/Plugins/VirtualProduction/RemoteControl/Source/RemoteControlLogic/Private/Behaviour/Builtin/Path/RCSetAssetByPathBehaviour.cpp
+++ b/Engine/Plugins/VirtualProduction/RemoteControl/Source/RemoteControlLogic/Private/Behaviour/Builtin/Path/RCSetAssetByPathBehaviour.cpp
@@ -298,7 +298,6 @@ FString URCSetAssetByPathBehaviour::GetCurrentPath()
 			}
 			
 			PathPart.RemoveFromStart(SetAssetByPathBehaviourHelpers::InputToken);
-			// TODO: Test and rewrite to something better
 			int32 CharIndexStart;
 			int32 CharIndexEnd;
 			PathPart.FindChar('(', CharIndexStart);
@@ -312,7 +311,7 @@ FString URCSetAssetByPathBehaviour::GetCurrentPath()
 
 			PathPart.RemoveAt(CharIndexStart);
 			PathPart.RemoveAt(CharIndexEnd-1);
-			const URCVirtualPropertyBase* TokenController = RemoteControlPreset->GetVirtualPropertyByDisplayName(FName(PathPart));
+			const URCVirtualPropertyBase* TokenController = RemoteControlPreset->GetControllerByDisplayName(FName(PathPart));
 			if (!TokenController)
 			{
 				ensureMsgf(false, TEXT("No Controller with given name found."));
@@ -387,7 +386,6 @@ bool URCSetAssetByPathBehaviour::SetExternalAsset(FString InExternalPath)
 		return false;
 	}
 
-	// TODO: Find better solution.
 	UTexture2D* ImportedTexture = UKismetRenderingLibrary::ImportFileAsTexture2D(this, InExternalPath);
 	if (!ImportedTexture)
 	{
diff --git a/Engine/Plugins/VirtualProduction/RemoteControl/Source/RemoteControlLogic/Private/Tests/RemoteControlLogicTests.cpp b/Engine/Plugins/VirtualProduction/RemoteControl/Source/RemoteControlLogic/Private/Tests/RemoteControlLogicTests.cpp
--- a/Engine/Plugins/VirtualProduction/RemoteControl/Source/RemoteControlLogic/Private/Tests/RemoteControlLogicTests.cpp
+++ b/Engine/Plugins/VirtualProduction/RemoteControl/Source/RemoteControlLogic/Private/Tests/RemoteControlLogicTests.cpp
@@ -46,17 +46,17 @@ bool FRemoteControlLogicTest::RunTest(const FString& Parameters)
  
 	// 3. Add Controllers
 	// 3.1.1 Create Float Property
-	URCController* FloatController =  Cast<URCController>(Preset->AddVirtualProperty(URCController::StaticClass(), EPropertyBagPropertyType::Float, nullptr, TEXT("FloatController")) );
+	URCController* FloatController =  Cast<URCController>(Preset->AddController(URCController::StaticClass(), EPropertyBagPropertyType::Float, nullptr, TEXT("FloatController")) );
 	// 3.1.2 Create Second Float Property
-	URCController* FloatController1 = Cast<URCController>(Preset->AddVirtualProperty(URCController::StaticClass(), EPropertyBagPropertyType::Float, nullptr, TEXT("FloatController1")) );
+	URCController* FloatController1 = Cast<URCController>(Preset->AddController(URCController::StaticClass(), EPropertyBagPropertyType::Float, nullptr, TEXT("FloatController1")) );
 	// 3.1.3 Create Bool Property
-	URCController* BoolController = Cast<URCController>(Preset->AddVirtualProperty(URCController::StaticClass(), EPropertyBagPropertyType::Bool, nullptr, TEXT("BoolController")));
+	URCController* BoolController = Cast<URCController>(Preset->AddController(URCController::StaticClass(), EPropertyBagPropertyType::Bool, nullptr, TEXT("BoolController")));
 	// 3.1.4 Create Int Controller
-	URCController* IntController =  Cast<URCController>(Preset->AddVirtualProperty(URCController::StaticClass(), EPropertyBagPropertyType::Int32, nullptr, TEXT("IntController")));
+	URCController* IntController =  Cast<URCController>(Preset->AddController(URCController::StaticClass(), EPropertyBagPropertyType::Int32, nullptr, TEXT("IntController")));
 	// 3.1.5 Create String Controller
-	URCController* StrController =  Cast<URCController>(Preset->AddVirtualProperty(URCController::StaticClass(), EPropertyBagPropertyType::String, nullptr, TEXT("StrController")));
+	URCController* StrController =  Cast<URCController>(Preset->AddController(URCController::StaticClass(), EPropertyBagPropertyType::String, nullptr, TEXT("StrController")));
 	// 3.1.6 Create Vector Struct Controller
-	URCController* VectorStructController =  Cast<URCController>(Preset->AddVirtualProperty(URCController::StaticClass(), EPropertyBagPropertyType::Struct, TBaseStructure<FVector>::Get(), TEXT("VectorProperty")));
+	URCController* VectorStructController =  Cast<URCController>(Preset->AddController(URCController::StaticClass(), EPropertyBagPropertyType::Struct, TBaseStructure<FVector>::Get(), TEXT("VectorProperty")));
  
 	// 3.2 Set Float Property Value to controller
 	constexpr float FloatValue = 0.65f;
@@ -151,19 +151,19 @@ bool FRemoteControlLogicTest::RunTest(const FString& Parameters)
 	TestEqual(TEXT("After empty actions the count should be = 0"), FloatController->Behaviours.Num(), 0);
  
 	// 9. Remove Controllers
-	const int32 NumPropertiesBeforeRemove = Preset->GetNumVirtualProperties();
+	const int32 NumPropertiesBeforeRemove = Preset->GetNumControllers();
 	float OutFloatValue = 0.f;
 	float OutFloatValue1 = 0.f;
 	FloatController->GetValueFloat(OutFloatValue);
-	Preset->RemoveVirtualProperty(FloatController1->PropertyName);
+	Preset->RemoveController(FloatController1->PropertyName);
 	FloatController->GetValueFloat(OutFloatValue1);
 	TestEqual(TEXT("After remove properties the old value should be the same"), OutFloatValue1, OutFloatValue);
-	Preset->RemoveVirtualProperty(FloatController->PropertyName);
-	Preset->RemoveVirtualProperty(BoolController->PropertyName);
-	Preset->RemoveVirtualProperty(IntController->PropertyName);
-	TestEqual(TEXT("After remove 4 properties the count should be =-4"), Preset->GetNumVirtualProperties(), NumPropertiesBeforeRemove - 4);
-	Preset->ResetVirtualProperties();
-	TestEqual(TEXT("After empty controllers could should be equal 0"), Preset->GetNumVirtualProperties(), 0);
+	Preset->RemoveController(FloatController->PropertyName);
+	Preset->RemoveController(BoolController->PropertyName);
+	Preset->RemoveController(IntController->PropertyName);
+	TestEqual(TEXT("After remove 4 properties the count should be =-4"), Preset->GetNumControllers(), NumPropertiesBeforeRemove - 4);
+	Preset->ResetControllers();
+	TestEqual(TEXT("After empty controllers could should be equal 0"), Preset->GetNumControllers(), 0);
  
 	return true;
 }
diff --git a/Engine/Plugins/VirtualProduction/RemoteControl/Source/RemoteControlUI/Private/UI/Controller/RCControllerModel.cpp b/Engine/Plugins/VirtualProduction/RemoteControl/Source/RemoteControlUI/Private/UI/Controller/RCControllerModel.cpp
--- a/Engine/Plugins/VirtualProduction/RemoteControl/Source/RemoteControlUI/Private/UI/Controller/RCControllerModel.cpp
+++ b/Engine/Plugins/VirtualProduction/RemoteControl/Source/RemoteControlUI/Private/UI/Controller/RCControllerModel.cpp
@@ -95,8 +95,10 @@ void FRCControllerModel::OnControllerNameCommitted(const FText& InNewControllerN
 	{
 		if(URCVirtualPropertyBase* Controller = GetVirtualProperty())
 		{
+			FName OldName = Controller->DisplayName;
 			Controller->DisplayName = *InNewControllerName.ToString();
 			ControllerNameTextBox->SetText(InNewControllerName);
+			Preset->OnControllerRenamed().Broadcast(Preset, OldName, *InNewControllerName.ToString());
 		}
 	}
 }
diff --git a/Engine/Plugins/VirtualProduction/RemoteControl/Source/RemoteControlUI/Private/UI/Controller/SRCControllerPanel.cpp b/Engine/Plugins/VirtualProduction/RemoteControl/Source/RemoteControlUI/Private/UI/Controller/SRCControllerPanel.cpp
--- a/Engine/Plugins/VirtualProduction/RemoteControl/Source/RemoteControlUI/Private/UI/Controller/SRCControllerPanel.cpp
+++ b/Engine/Plugins/VirtualProduction/RemoteControl/Source/RemoteControlUI/Private/UI/Controller/SRCControllerPanel.cpp
@@ -216,7 +216,7 @@ void SRCControllerPanel::DuplicateController(URCController* InController)
 
 	if (URemoteControlPreset* Preset = GetPreset())
 	{
-		if (URCController* NewController = Cast<URCController>(Preset->DuplicateVirtualProperty(InController)))
+		if (URCController* NewController = Cast<URCController>(Preset->DuplicateController(InController)))
 		{
 			NewController->SetDisplayIndex(ControllerPanelList->NumControllerItems());
 
@@ -344,7 +344,7 @@ void SRCControllerPanel::OnAddControllerClicked(const EPropertyBagPropertyType I
 		FScopedTransaction Transaction(LOCTEXT("AddController", "Add Controller"));
 		Preset->Modify();
 
-		URCVirtualPropertyInContainer* NewVirtualProperty = Preset->AddVirtualProperty(URCController::StaticClass(), InValueType, InValueTypeObject);
+		URCVirtualPropertyInContainer* NewVirtualProperty = Preset->AddController(URCController::StaticClass(), InValueType, InValueTypeObject);
 
 		if (ControllerPanelList.IsValid())
 		{
@@ -373,7 +373,7 @@ FReply SRCControllerPanel::OnClickEmptyButton()
 		FScopedTransaction Transaction(LOCTEXT("EmptyControllers", "Empty Controllers"));
 		Preset->Modify();
 
-		Preset->ResetVirtualProperties();
+		Preset->ResetControllers();
 	}
 
 	if (const TSharedPtr<SRemoteControlPanel> RemoteControlPanel = GetRemoteControlPanel())
diff --git a/Engine/Plugins/VirtualProduction/RemoteControl/Source/RemoteControlUI/Private/UI/Controller/SRCControllerPanelList.cpp b/Engine/Plugins/VirtualProduction/RemoteControl/Source/RemoteControlUI/Private/UI/Controller/SRCControllerPanelList.cpp
--- a/Engine/Plugins/VirtualProduction/RemoteControl/Source/RemoteControlUI/Private/UI/Controller/SRCControllerPanelList.cpp
+++ b/Engine/Plugins/VirtualProduction/RemoteControl/Source/RemoteControlUI/Private/UI/Controller/SRCControllerPanelList.cpp
@@ -375,7 +375,7 @@ void SRCControllerPanelList::Reset()
 			FProperty* Property = Child->CreatePropertyHandle()->GetProperty();
  			check(Property);
 
-			if (URCVirtualPropertyBase* Controller = Preset->GetVirtualProperty(Property->GetFName()))
+			if (URCVirtualPropertyBase* Controller = Preset->GetController(Property->GetFName()))
 			{
 				if(ensureAlways(ControllerItems.IsValidIndex(Controller->DisplayIndex)))
 					ControllerItems[Controller->DisplayIndex] = MakeShared<FRCControllerModel>(Controller, Child, RemoteControlPanel);
@@ -444,7 +444,7 @@ void SRCControllerPanelList::OnFinishedChangingProperties(const FPropertyChanged
 {
 	if (URemoteControlPreset* Preset = GetPreset())
 	{
-		Preset->OnModifyVirtualProperty(PropertyChangedEvent);
+		Preset->OnModifyController(PropertyChangedEvent);
 	}
 }
 
@@ -502,7 +502,7 @@ int32 SRCControllerPanelList::RemoveModel(const TSharedPtr<FRCLogicModeBase> InM
 					Preset->Modify();
 
 					// Remove Model from Data Container
-					const bool bRemoved = Preset->RemoveVirtualProperty(SelectedController->GetPropertyName());
+					const bool bRemoved = Preset->RemoveController(SelectedController->GetPropertyName());
 					if (bRemoved)
 					{
 						// Shift all display indices higher than the deleted index down by 1
@@ -675,8 +675,8 @@ void SRCControllerPanelList::CreateAutoBindForProperty(TSharedPtr<const FRemoteC
 		if(bSuccess)
 		{
 			// Step 1. Create a Controller of matching type
-			URCController* NewController = Cast<URCController>(Preset->AddVirtualProperty(URCController::StaticClass(), PropertyBagType, StructObject));
-			NewController->DisplayIndex = Preset->GetNumVirtualProperties() - 1;
+			URCController* NewController = Cast<URCController>(Preset->AddController(URCController::StaticClass(), PropertyBagType, StructObject));
+			NewController->DisplayIndex = Preset->GetNumControllers() - 1;
 
 			// Transfer property value from Exposed Property to the New Controller.
 			// The goal is to keep values synced for a Controller newly created via "Auto Bind"
diff --git a/Engine/Plugins/VirtualProduction/RemoteControl/Source/WebRemoteControl/Private/RemoteControlModels.h b/Engine/Plugins/VirtualProduction/RemoteControl/Source/WebRemoteControl/Private/RemoteControlModels.h
--- a/Engine/Plugins/VirtualProduction/RemoteControl/Source/WebRemoteControl/Private/RemoteControlModels.h
+++ b/Engine/Plugins/VirtualProduction/RemoteControl/Source/WebRemoteControl/Private/RemoteControlModels.h
@@ -5,6 +5,7 @@
 #include "Algo/Transform.h"
 #include "AssetRegistry/AssetData.h"
 #include "AssetRegistry/ARFilter.h"
+#include "Controller/RCController.h"
 #include "GameFramework/Actor.h"
 #include "RemoteControlActor.h"
 #include "RemoteControlField.h"
@@ -356,6 +357,93 @@ struct FRCExposedActorDescription
 	FRCObjectDescription UnderlyingActor;
 };
 
+USTRUCT()
+struct FRCControllerDescription
+{
+	GENERATED_BODY()
+
+	FRCControllerDescription() = default;
+
+	FRCControllerDescription(const URCVirtualPropertyBase* InController)
+	{
+		if (!InController || !InController->GetProperty())
+		{
+			return;
+		}
+
+		DisplayName = InController->DisplayName;
+		ID = InController->Id.ToString();
+		Path = InController->GetPathName();
+		Type = InController->GetProperty()->GetCPPType();
+		Metadata = InController->Metadata;
+	}
+	
+	/** The label displayed in the remote control panel for this controller. */
+	UPROPERTY()
+	FName DisplayName;
+	
+	/** Unique identifier for the controller. */
+	UPROPERTY()
+	FString ID;
+
+	/** Type of the Controller */
+	UPROPERTY()
+	FString Type;
+
+	/** Path of the Controller */
+	UPROPERTY()
+	FString Path;
+
+	/** Metadata of the Controller. Initially Empty. */
+	UPROPERTY()
+	TMap<FName, FString> Metadata;
+};
+
+USTRUCT()
+struct FRCControllerModifiedDescription
+{
+	GENERATED_BODY()
+
+	FRCControllerModifiedDescription () = default;
+
+	FRCControllerModifiedDescription(const URemoteControlPreset* InPreset, const TArray<FGuid>& InModifiedControllers)
+	{
+		if (!InPreset)
+		{
+			return;
+		}
+
+		for (const FGuid& ControllerId : InModifiedControllers)
+		{
+			AddModifiedController(InPreset, ControllerId);
+		}
+		
+	}
+
+public:
+	/** The list of modified RC controllers. */
+	UPROPERTY()
+	TArray<FRCControllerDescription> Controllers;
+
+	UPROPERTY()
+	TArray<FString> ChangedValues;
+	
+private:
+	void AddModifiedController(const URemoteControlPreset* InPreset, const FGuid& InControllerId)
+	{
+		if (!InPreset)
+		{
+			return;
+		}
+
+		if (URCVirtualPropertyBase* Controller = InPreset->GetController(InControllerId))
+		{
+			Controllers.Emplace(Controller);
+			ChangedValues.Emplace(Controller->GetDisplayValueAsString());
+		}
+	}
+};
+
 USTRUCT()
 struct FRCPresetLayoutGroupDescription
 {
@@ -502,6 +590,7 @@ struct FRCPresetDescription
 		ID = Preset->GetPresetId().ToString();
 
 		Algo::Transform(Preset->Layout.GetGroups(), Groups, [Preset](const FRemoteControlPresetGroup& Group) { return FRCPresetLayoutGroupDescription{ Preset, Group }; });
+		Algo::Transform(Preset->GetControllers(), Controllers, [Preset](const URCVirtualPropertyBase* Controller){ return FRCControllerDescription{ Controller };});
 	}
 
 	/**
@@ -527,6 +616,12 @@ struct FRCPresetDescription
 	 */
 	UPROPERTY()
 	TArray<FRCPresetLayoutGroupDescription> Groups;
+
+	/**
+	 * The controllers held by the Preset (no groups)
+	 */
+	UPROPERTY()
+	TArray<FRCControllerDescription> Controllers;
 };
 
 USTRUCT()
diff --git a/Engine/Plugins/VirtualProduction/RemoteControl/Source/WebRemoteControl/Private/RemoteControlResponse.h b/Engine/Plugins/VirtualProduction/RemoteControl/Source/WebRemoteControl/Private/RemoteControlResponse.h
--- a/Engine/Plugins/VirtualProduction/RemoteControl/Source/WebRemoteControl/Private/RemoteControlResponse.h
+++ b/Engine/Plugins/VirtualProduction/RemoteControl/Source/WebRemoteControl/Private/RemoteControlResponse.h
@@ -399,6 +399,137 @@ struct FRCPresetEntitiesModifiedEvent
 	FRCPresetModifiedEntitiesDescription ModifiedEntities;
 };
 
+USTRUCT()
+struct FRCPresetControllersRenamedEvent
+{
+	GENERATED_BODY()
+
+	FRCPresetControllersRenamedEvent() = default;
+
+	FRCPresetControllersRenamedEvent(FName InPresetName, FGuid InPresetId, TArray<TTuple<FName, FName>> InRenamedControllers)
+		: Type(TEXT("PresetControllersRenamed"))
+		, PresetName(InPresetName)
+		, PresetId(InPresetId.ToString())
+		, RenamedControllers(InRenamedControllers)
+	{
+	}
+		
+	UPROPERTY()
+	FString Type;
+
+	UPROPERTY()
+	FName PresetName;
+
+	UPROPERTY()
+	FString PresetId;
+
+	UPROPERTY()
+	TArray<FRCPresetFieldRenamed> RenamedControllers;
+};
+
+USTRUCT()
+struct FRCPresetControllersRemovedEvent
+{
+	GENERATED_BODY()
+
+	FRCPresetControllersRemovedEvent() = default;
+
+	FRCPresetControllersRemovedEvent(FName InPresetName, FGuid InPresetId, TArray<FName> InRemoveControllers, const TArray<FGuid>& InRemovedControllerIds)
+		: Type(TEXT("PresetControllersRemoved"))
+		, PresetName(InPresetName)
+		, PresetId(InPresetId.ToString())
+		, RemovedControllers(MoveTemp(InRemoveControllers))
+	{
+		Algo::Transform(InRemovedControllerIds, RemovedControllerIds, [](const FGuid& Id){ return Id.ToString();});
+	}
+
+	UPROPERTY()
+	FString Type;
+
+	UPROPERTY()
+	FName PresetName;
+
+	UPROPERTY()
+	FString PresetId;
+
+	UPROPERTY()
+	TArray<FName> RemovedControllers;
+
+	UPROPERTY()
+	TArray<FString> RemovedControllerIds;
+};
+
+USTRUCT()
+struct FRCPresetControllersAddedEvent
+{
+	GENERATED_BODY()
+
+	FRCPresetControllersAddedEvent() = default;
+
+	FRCPresetControllersAddedEvent(FName InPresetName, FGuid InPresetId, FRCPresetDescription InPresetDescription)
+		: Type(TEXT("PresetControllersAdded"))
+		, PresetName(InPresetName)
+		, PresetId(InPresetId.ToString())
+		, Description(MoveTemp(InPresetDescription))
+	{
+	}
+
+	UPROPERTY()
+	FString Type;
+	
+	UPROPERTY()
+	FName PresetName;
+
+	UPROPERTY()
+	FString PresetId;
+
+	UPROPERTY()
+	FRCPresetDescription Description;
+};
+
+/** Event which is triggered whenever a Controller is modified. */
+USTRUCT()
+struct FRCPresetControllersModifiedEvent
+{
+	GENERATED_BODY()
+
+	FRCPresetControllersModifiedEvent() = default;
+
+	FRCPresetControllersModifiedEvent(const URemoteControlPreset* InPreset, const TArray<FGuid>& InModifiedControllers)
+		: Type(TEXT("PresetControllersModified"))
+	{
+		checkSlow(InPreset);
+		PresetName = InPreset->GetPresetName();
+		PresetId = InPreset->GetPresetId().ToString();
+		ModifiedControllers = FRCControllerModifiedDescription(InPreset, InModifiedControllers);
+		
+	}
+
+	/**
+	 * Type of the event.
+	 */
+	UPROPERTY()
+	FString Type;
+
+	/**
+	 * Name of the preset which contains the modified controller.
+	 */
+	UPROPERTY()
+	FName PresetName;
+	
+	/**
+	 * ID of the preset that contains the modified controller.
+	 */
+	UPROPERTY()
+	FString PresetId;
+
+	/**
+	 * The controllers that were modified in the last frame.
+	 */
+	UPROPERTY()
+	FRCControllerModifiedDescription ModifiedControllers;
+};
+
 /**
  * Data about actors that have changed in the scene.
  */
diff --git a/Engine/Plugins/VirtualProduction/RemoteControl/Source/WebRemoteControl/Private/WebRemoteControl.cpp b/Engine/Plugins/VirtualProduction/RemoteControl/Source/WebRemoteControl/Private/WebRemoteControl.cpp
--- a/Engine/Plugins/VirtualProduction/RemoteControl/Source/WebRemoteControl/Private/WebRemoteControl.cpp
+++ b/Engine/Plugins/VirtualProduction/RemoteControl/Source/WebRemoteControl/Private/WebRemoteControl.cpp
@@ -183,6 +183,26 @@ namespace WebRemoteControl
 		return Preset->GetExposedEntity<EntityType>(Preset->GetExposedEntityId(*PropertyLabelOrId)).Pin();
 	}
 
+	URCVirtualPropertyBase* GetController(URemoteControlPreset* Preset, FString PropertyLabelOrId)
+	{
+		if (!Preset)
+		{
+			return nullptr;
+		}
+		
+		FGuid Id;
+		if (FGuid::ParseExact(PropertyLabelOrId, EGuidFormats::Digits, Id))
+		{
+			if (URCVirtualPropertyBase* Controller = Preset->GetController(Id))
+			{
+        		return Controller;
+			}
+		}
+
+		// Not supporting label, in case Exposed Property and Controller have the same name.
+		return nullptr;
+	}
+
 	URemoteControlPreset* GetPreset(FString PresetNameOrId)
 	{
 		FGuid Id;
@@ -1186,10 +1206,40 @@ bool FWebRemoteControlModule::HandlePresetSetPropertyRoute(const FHttpServerRequ
 		return true;
 	}
 
-	TSharedPtr<FRemoteControlProperty> RemoteControlProperty = WebRemoteControl::GetRCEntity<FRemoteControlProperty>(Preset, Args.FieldLabel);
+	const TSharedPtr<FRemoteControlProperty> RemoteControlProperty = WebRemoteControl::GetRCEntity<FRemoteControlProperty>(Preset, Args.FieldLabel);
 
 	if (!RemoteControlProperty.IsValid())
 	{
+		// In case the Property is not found or not valid, see if we do have that Id for the Controllers.
+		if (URCVirtualPropertyBase* Controller = WebRemoteControl::GetController(Preset, *Args.FieldLabel))
+		{
+			TArray<uint8> NewPayload;
+			const FName PropertyValueKey = WebRemoteControlStructUtils::Prop_PropertyValue;
+			const FName PropertyNameInternal = Controller->GetPropertyName();
+
+			// Objrct Ref Method
+			FRCObjectReference ObjectRef;
+			ObjectRef.Property = Controller->GetProperty();
+			ObjectRef.Access = ERCAccess::WRITE_ACCESS;
+
+			RemotePayloadSerializer::ReplaceFirstOccurence(SetPropertyRequest.TCHARBody, PropertyValueKey.ToString(), PropertyNameInternal.ToString(), NewPayload);
+			FMemoryReader Reader(NewPayload);
+			FJsonStructDeserializerBackend Backend(Reader);
+			
+			if (IRemoteControlModule::Get().SetPresetController(*Args.PresetName, Controller, Backend, NewPayload, true))
+			{
+				Response->Code = EHttpServerResponseCodes::Ok;
+			}
+			else
+			{
+				WebRemoteControlInternalUtils::CreateUTF8ErrorMessage(FString::Printf(TEXT("Error while trying to set controller %s."), *Args.FieldLabel), Response->Body);
+			}
+
+			OnComplete(MoveTemp(Response));
+
+			return true;
+		}
+		
 		Response->Code = EHttpServerResponseCodes::NotFound;
 		WebRemoteControlInternalUtils::CreateUTF8ErrorMessage(TEXT("Unable to resolve the preset field."), Response->Body);
 		OnComplete(MoveTemp(Response));
@@ -1231,9 +1281,45 @@ bool FWebRemoteControlModule::HandlePresetGetPropertyRoute(const FHttpServerRequ
 	}
 
 	TSharedPtr<FRemoteControlProperty> RemoteControlProperty = WebRemoteControl::GetRCEntity<FRemoteControlProperty>(Preset, Args.FieldLabel);
-
 	if (!RemoteControlProperty)
 	{
+		// In case we do not find a Property, instead go look for a Controller
+		if (URCVirtualPropertyBase* Controller = WebRemoteControl::GetController(Preset, *Args.FieldLabel))
+		{
+			// Define a StructOnScope for representing our generic value output
+			FWebRCGenerateStructArgs StructArgs;
+			FName PropertyValueKey = WebRemoteControlStructUtils::Prop_PropertyValue;
+			StructArgs.GenericProperties.Emplace(PropertyValueKey, Controller->GetProperty());
+
+			FString& StructName = ControllersSerializerStructNameCache.FindOrAdd(Controller->PropertyName);
+			if(StructName.IsEmpty())
+			{
+				StructName = FString::Format(TEXT("{0}_{1}"), { *PropertyValueKey.ToString(), Controller->Id.ToString() });
+				ControllersSerializerStructNameCache.Add(Controller->PropertyName, StructName);
+			}
+			
+			FStructOnScope Result(UE::WebRCReflectionUtils::GenerateStruct(*StructName, StructArgs));
+
+			// Copy raw memory from the controller onto our dynamic struct
+			const FProperty* TargetValueProp = Result.GetStruct()->FindPropertyByName(PropertyValueKey);
+			uint8* DestPtr = TargetValueProp->ContainerPtrToValuePtr<uint8>((void*)Result.GetStructMemory());
+
+			Controller->CopyCompleteValue(TargetValueProp, DestPtr);
+
+			// Serialize to JSON
+			TArray<uint8> JsonBuffer;
+			FMemoryWriter Writer(JsonBuffer);
+			WebRemoteControlInternalUtils::SerializeStructOnScope(Result, Writer);
+
+			// Finalize HTTP Response
+			WebRemoteControlUtils::ConvertToUTF8(JsonBuffer, Response->Body);
+			Response->Code = EHttpServerResponseCodes::Ok;
+
+			OnComplete(MoveTemp(Response));
+
+			return true;
+		}
+		
 		Response->Code = EHttpServerResponseCodes::NotFound;
 		WebRemoteControlInternalUtils::CreateUTF8ErrorMessage(TEXT("Unable to resolve the exposed entity."), Response->Body);
 		OnComplete(MoveTemp(Response));
@@ -1822,6 +1908,56 @@ bool FWebRemoteControlModule::HandleEntityMetadataOperationsRoute(const FHttpSer
 
 	if (!Entity.IsValid())
 	{
+		// Test for Controllers
+		if (URCVirtualPropertyBase* Controller = WebRemoteControl::GetController(Preset, Label))
+		{
+			if (Request.Verb == EHttpServerRequestVerbs::VERB_GET)
+			{
+				FString MetadataValue = Controller->GetMetadataValue(*Key);
+				if (!MetadataValue.IsEmpty())
+				{
+					WebRemoteControlUtils::SerializeMessage(FGetMetadataFieldResponse{ *MetadataValue }, Response->Body);
+				}
+				else
+				{
+					WebRemoteControlInternalUtils::CreateUTF8ErrorMessage(FString::Printf(TEXT("Metadata entry %s could not be found."), *Key), Response->Body);
+					Response->Code = EHttpServerResponseCodes::NotFound;
+				}
+			}
+			else if (Request.Verb == EHttpServerRequestVerbs::VERB_PUT)
+			{
+				if (!WebRemoteControlInternalUtils::ValidateContentType(Request, TEXT("application/json"), OnComplete))
+				{
+					return true;
+				}
+				
+				FSetEntityMetadataRequest SetMetadataRequest;
+				if (!WebRemoteControlInternalUtils::DeserializeRequest(Request, &OnComplete, SetMetadataRequest))
+				{
+					return true;
+				}
+
+#if WITH_EDITOR
+				FScopedTransaction Transaction( LOCTEXT("ModifyEntityMetadata", "Modify exposed entity metadata entry"));
+				Preset->Modify();
+#endif
+				Controller->SetMetadataValue(*Key, SetMetadataRequest.Value);
+			}
+			else if (Request.Verb == EHttpServerRequestVerbs::VERB_DELETE)
+			{
+#if WITH_EDITOR
+				FScopedTransaction Transaction( LOCTEXT("DeleteEntityMetadata", "Delete exposed entity metadata entry"));
+				Preset->Modify();
+#endif
+				Controller->RemoveMetadataValue(*Key);
+			}
+
+			Response->Code = EHttpServerResponseCodes::Ok;
+
+			OnComplete(MoveTemp(Response));
+			return true;
+		}
+		
 		Response->Code = EHttpServerResponseCodes::NotFound;
 		WebRemoteControlInternalUtils::CreateUTF8ErrorMessage(TEXT("Unable to resolve the preset entity."), Response->Body);
 		OnComplete(MoveTemp(Response));
@@ -2036,7 +2172,7 @@ bool FWebRemoteControlModule::HandlePresetSetControllerRoute(const FHttpServerRe
 	}
 
 	// 6. Acquire Controller
-	URCVirtualPropertyBase* Controller = Preset->GetVirtualPropertyByDisplayName(*Args.ControllerName);
+	URCVirtualPropertyBase* Controller = Preset->GetControllerByDisplayName(*Args.ControllerName);
 	if (!Controller)
 	{
 		Response->Code = EHttpServerResponseCodes::NotFound;
@@ -2096,7 +2232,7 @@ bool FWebRemoteControlModule::HandlePresetGetControllerRoute(const FHttpServerRe
 	}
 
 	// 4. Acquire Controller
-	URCVirtualPropertyBase* Controller = Preset->GetVirtualPropertyByDisplayName(*Args.ControllerName);
+	URCVirtualPropertyBase* Controller = Preset->GetControllerByDisplayName(*Args.ControllerName);
 	if (!Controller)
 	{
 		Response->Code = EHttpServerResponseCodes::NotFound;
diff --git a/Engine/Plugins/VirtualProduction/RemoteControl/Source/WebRemoteControl/Private/WebSocketMessageHandler.cpp b/Engine/Plugins/VirtualProduction/RemoteControl/Source/WebRemoteControl/Private/WebSocketMessageHandler.cpp
--- a/Engine/Plugins/VirtualProduction/RemoteControl/Source/WebRemoteControl/Private/WebSocketMessageHandler.cpp
+++ b/Engine/Plugins/VirtualProduction/RemoteControl/Source/WebRemoteControl/Private/WebSocketMessageHandler.cpp
@@ -9,6 +9,7 @@
 #include "Engine/Engine.h"
 #include "EngineUtils.h"
 #include "IRemoteControlModule.h"
+#include "RCVirtualProperty.h"
 #include "GameFramework/Actor.h"
 #include "RemoteControlPreset.h"
 #include "RemoteControlRequest.h"
@@ -168,6 +169,22 @@ namespace WebSocketMessageHandlerStructUtils
 
 		return StructOnScope;
 	}
+	
+	FStructOnScope CreatePresetControllerChangedStructOnScope(const URemoteControlPreset* Preset, const TArray<FStructOnScope*>& PropertyValuesOnScope, int64 SequenceNumber)
+	{
+		UScriptStruct* PropertyValueStruct = (UScriptStruct*)PropertyValuesOnScope[0]->GetStruct();
+		check(PropertyValueStruct);
+
+		UScriptStruct* TopLevelStruct = CreatePresetFieldsChangedStruct(PropertyValueStruct);
+
+		FStructOnScope FieldsChangedOnScope{ TopLevelStruct };
+		SetStringPropertyValue(Prop_Type, FieldsChangedOnScope, TEXT("PresetFieldsChanged"));
+		SetStringPropertyValue(Prop_PresetName, FieldsChangedOnScope, *Preset->GetPresetName().ToString());
+		SetStringPropertyValue(Prop_PresetId, FieldsChangedOnScope, *Preset->GetPresetId().ToString());
+		SetStringPropertyValue(Prop_SequenceNumber, FieldsChangedOnScope, FString::Printf(TEXT("%lld"), SequenceNumber));
+
+		return FieldsChangedOnScope;
+	}
 
 	FStructOnScope CreatePresetFieldsChangedStructOnScope(const URemoteControlPreset* Preset, const TArray<FStructOnScope>& PropertyValuesOnScope, int64 SequenceNumber)
 	{
@@ -255,6 +272,24 @@ namespace WebSocketMessageHandlerMiscUtils
 
 		return nullptr;
 	}
+
+	void* GetPresetControllerClassPointer(URemoteControlPreset* Preset, const FGuid& ControllerId)
+	{
+		if (const URCVirtualPropertyBase* Controller = Preset->GetController(ControllerId))
+		{
+			if (const FProperty* Property = Controller->GetProperty())
+			{
+				if (const FStructProperty* StructProperty = CastField<FStructProperty>(Property))
+				{
+					return StructProperty->Struct;
+				}
+
+				return Property->GetClass();
+			}
+		}
+
+		return nullptr;
+	}
 }
 
 FWebSocketMessageHandler::FWebSocketMessageHandler(FRCWebSocketServer* InServer, const FGuid& InActingClientId)
@@ -480,6 +515,10 @@ void FWebSocketMessageHandler::HandleWebSocketPresetRegister(const FRemoteContro
 		Preset->OnActorPropertyModified().AddRaw(this, &FWebSocketMessageHandler::OnActorPropertyChanged);
 		Preset->OnEntitiesUpdated().AddRaw(this, &FWebSocketMessageHandler::OnEntitiesModified);
 		Preset->OnPresetLayoutModified().AddRaw(this, &FWebSocketMessageHandler::OnLayoutModified);
+		Preset->OnControllerAdded().AddRaw(this, &FWebSocketMessageHandler::OnControllerAdded);
+		Preset->OnControllerRemoved().AddRaw(this, &FWebSocketMessageHandler::OnControllerRemoved);
+		Preset->OnControllerRenamed().AddRaw(this, &FWebSocketMessageHandler::OnControllerRenamed);
+		Preset->OnControllerModified().AddRaw(this, &FWebSocketMessageHandler::OnControllerModified);
 	}
 
 	ClientIds->AddUnique(WebSocketMessage.ClientId);
@@ -798,6 +837,137 @@ void FWebSocketMessageHandler::HandleWebSocketEndEditorTransaction(const FRemote
 #endif
 }
 
+void FWebSocketMessageHandler::ProcessChangedControllers()
+{
+	// Go over each controller that was changed for each preset
+	for (TTuple<FGuid, TMap<FGuid, TSet<FGuid>>>& Entry : PerFrameModifiedControllers)
+	{
+		if (!ShouldProcessEventForPreset(Entry.Key) || !Entry.Value.Num())
+		{
+			continue;
+		}
+
+		URemoteControlPreset* Preset = IRemoteControlModule::Get().ResolvePreset(Entry.Key);
+		if (!Preset)
+		{
+			continue;
+		}
+		
+		UE_LOG(LogRemoteControl, VeryVerbose, TEXT("(%s) Broadcasting controllers changed event."), *Preset->GetName());
+
+		// Each client will have a custom payload that doesn't contain the events it triggered.
+		for (const TPair<FGuid, TSet<FGuid>>& ClientToModifications : Entry.Value)
+		{
+			TMap<void*, TSet<FGuid>> ControllersIdByType;
+			for (const FGuid& Id : ClientToModifications.Value)
+			{
+				void* ClassPtr = WebSocketMessageHandlerMiscUtils::GetPresetControllerClassPointer(Preset, Id);
+				ControllersIdByType.FindOrAdd(ClassPtr).Emplace(Id);
+			}
+
+			// Send a Controller Change Event for each Type
+			for (const TPair<void*, TSet<FGuid>>& ClassToEventsPair : ControllersIdByType)
+			{
+				const uint64 ClientSequenceNumber = GetSequenceNumber(ClientToModifications.Key);
+				
+				TArray<uint8> WorkingBuffer;
+				if (ClientToModifications.Value.Num() && WriteControllerChangeEventPayload(Preset, ClassToEventsPair.Value, ClientSequenceNumber, WorkingBuffer))
+				{
+					TArray<uint8> Payload;
+					WebRemoteControlUtils::ConvertToUTF8(WorkingBuffer, Payload);
+					Server->Send(ClientToModifications.Key, Payload);
+				}
+			}
+		}
+	}
+
+	PerFrameModifiedControllers.Empty();
+}
+
+void FWebSocketMessageHandler::ProcessAddedControllers()
+{
+	for (const TPair<FGuid, TArray<FGuid>>& Entry : PerFrameAddedControllers)
+	{
+		if (Entry.Value.Num() <= 0 || !ShouldProcessEventForPreset(Entry.Key))
+		{
+			continue;
+		}
+
+		URemoteControlPreset* Preset = IRemoteControlModule::Get().ResolvePreset(Entry.Key);
+		if (!Preset)
+		{
+			continue;
+		}
+		
+		UE_LOG(LogRemoteControl, VeryVerbose, TEXT("(%s) Broadcasting controllers added event."), *Preset->GetName());
+
+		FRCPresetDescription AddedControllerDescription;
+		AddedControllerDescription.Name = Preset->GetName();
+		AddedControllerDescription.Path = Preset->GetPathName();
+		AddedControllerDescription.ID = Preset->GetPresetId().ToString();
+		Algo::Transform(Entry.Value, AddedControllerDescription.Controllers, [Preset](const FGuid& ControllerId){ return Preset->GetController(ControllerId); });
+		
+		TArray<uint8> Payload;
+		WebRemoteControlUtils::SerializeMessage(FRCPresetControllersAddedEvent { Preset->GetPresetName(), Preset->GetPresetId(), AddedControllerDescription }, Payload);
+		BroadcastToPresetListeners(Entry.Key, Payload);
+	}
+
+	PerFrameAddedControllers.Empty();
+}
+
+void FWebSocketMessageHandler::ProcessRemovedControllers()
+{
+	for (const TPair<FGuid, TTuple<TArray<FGuid>, TArray<FName>>>& Entry : PerFrameRemovedControllers)
+	{
+		if (Entry.Value.Key.Num() <= 0 || !ShouldProcessEventForPreset(Entry.Key))
+		{
+			continue;
+		}
+
+		URemoteControlPreset* Preset = IRemoteControlModule::Get().ResolvePreset(Entry.Key);
+		if (Preset == nullptr)
+		{
+			continue;
+		}
+
+		ensure(Entry.Value.Key.Num() == Entry.Value.Value.Num());
+		
+		UE_LOG(LogRemoteControl, VeryVerbose, TEXT("(%s) Broadcasting controllers removed event."), *Preset->GetName());
+		
+		TArray<uint8> Payload;
+		WebRemoteControlUtils::SerializeMessage(FRCPresetControllersRemovedEvent{ Preset->GetPresetName(), Preset->GetPresetId(), Entry.Value.Value, Entry.Value.Key }, Payload);
+		BroadcastToPresetListeners(Entry.Key, Payload);
+	}
+
+	PerFrameRemovedControllers.Empty();
+}
+
+// Actually not called currently
+void FWebSocketMessageHandler::ProcessRenamedControllers()
+{
+	for (const TPair<FGuid, TArray<TTuple<FName, FName>>>& Entry : PerFrameRenamedControllers)
+	{
+		if (Entry.Value.Num() <= 0 || !ShouldProcessEventForPreset(Entry.Key))
+		{
+			continue;
+		}
+
+		URemoteControlPreset* Preset = IRemoteControlModule::Get().ResolvePreset(Entry.Key);
+		if (Preset == nullptr)
+		{
+			continue;
+		}
+		
+		UE_LOG(LogRemoteControl, VeryVerbose, TEXT("(%s) Broadcasting controllers renamed event."), *Preset->GetName());
+
+		TArray<uint8> Payload;
+		WebRemoteControlUtils::SerializeMessage(FRCPresetControllersRenamedEvent{Preset->GetPresetName(), Preset->GetPresetId(), Entry.Value}, Payload);
+		BroadcastToPresetListeners(Entry.Key, Payload);
+	}
+
+	PerFrameRenamedControllers.Empty();
+}
+
 void FWebSocketMessageHandler::ProcessChangedProperties()
 {
 	//Go over each property that were changed for each preset
@@ -1150,7 +1320,75 @@ void FWebSocketMessageHandler::OnEndFrame()
 		ProcessModifiedMetadata();
 		ProcessModifiedPresetLayouts();
 		ProcessActorChanges();
+		ProcessAddedControllers();
+		ProcessChangedControllers();
+		ProcessRenamedControllers();
+		ProcessRemovedControllers();
+	}
+}
+
+void FWebSocketMessageHandler::OnControllerAdded(URemoteControlPreset* Owner, FName NewControllerName, const FGuid& EntityId)
+{
+	if (Owner == nullptr)
+	{
+		return;
+	}
+
+	if (PresetNotificationMap.Num() <= 0)
+	{
+		return;
+	}
+
+	PerFrameAddedControllers.FindOrAdd(Owner->GetPresetId()).AddUnique(EntityId);
+}
+
+void FWebSocketMessageHandler::OnControllerRemoved(URemoteControlPreset* Owner, const FGuid& EntityId)
+{
+	if (Owner == nullptr)
+	{
+		return;
+	}
+
+	if (PresetNotificationMap.Num() <= 0)
+	{
+		return;
+	}
+
+	const URCVirtualPropertyBase* Controller = Owner->GetController(EntityId);
+	check(Controller);
+
+	// Cache the Controller which was removed for end of frame notification
+	TTuple<TArray<FGuid>, TArray<FName>>& Entries = PerFrameRemovedControllers.FindOrAdd(Owner->GetPresetId());
+	Entries.Key.AddUnique(EntityId);
+	Entries.Value.AddUnique(Controller->DisplayName);
+}
+
+void FWebSocketMessageHandler::OnControllerRenamed(URemoteControlPreset* Owner, FName OldControllerLabel, FName NewControllerLabel)
+{
+	if (Owner == nullptr)
+	{
+		return;
+	}
+
+	if (PresetNotificationMap.Num() <= 0)
+	{
+		return;
 	}
+
+	PerFrameRenamedControllers.FindOrAdd(Owner->GetPresetId()).AddUnique(TTuple<FName, FName>(OldControllerLabel, NewControllerLabel));
+}
+
+void FWebSocketMessageHandler::OnControllerModified(URemoteControlPreset* Owner, const TSet<FGuid>& ModifiedControllerIds)
+{
+	// We do not need to store these event for the current frame since this was already handled by the preset in this case.
+	if (!Owner || ModifiedControllerIds.Num() == 0)
+	{
+		return;
+	}
+	
+	TArray<uint8> Payload;
+	WebRemoteControlUtils::SerializeMessage(FRCPresetControllersModifiedEvent{Owner, ModifiedControllerIds.Array()}, Payload);
+	BroadcastToPresetListeners(Owner->GetPresetId(), Payload);
 }
 
 #if WITH_EDITOR
@@ -1487,6 +1725,35 @@ bool FWebSocketMessageHandler::ShouldProcessEventForPreset(const FGuid& PresetId
 	return PresetNotificationMap.Contains(PresetId) && PresetNotificationMap[PresetId].Num() > 0;
 }
 
+bool FWebSocketMessageHandler::WriteControllerChangeEventPayload(URemoteControlPreset* InPreset, const TSet<FGuid>& InModifiedControllerIds, int64 InSequenceNumber, TArray<uint8>& OutBuffer)
+{
+	bool bHasController = false;
+
+	TArray<FStructOnScope*> ControllerValueOnScope;
+	for (const FGuid& ControllerId : InModifiedControllerIds)
+	{
+		if (const URCVirtualPropertySelfContainer* Controller = Cast<URCVirtualPropertySelfContainer>(InPreset->GetController(ControllerId)))
+		{
+			bHasController = true;
+
+			if (FStructOnScope* ControllerStructOnScope = Controller->CreateStructOnScope().Get())
+			{
+				ControllerValueOnScope.Add(ControllerStructOnScope);
+			}
+		}
+	}
+
+	if (ControllerValueOnScope.Num())
+	{
+		FStructOnScope ValueStructOnScope =	WebSocketMessageHandlerStructUtils::CreatePresetControllerChangedStructOnScope(InPreset, ControllerValueOnScope, InSequenceNumber);
+		
+		FMemoryWriter Writer(OutBuffer);
+		WebRemoteControlInternalUtils::SerializeStructOnScope(ValueStructOnScope, Writer);
+	}
+	
+	return bHasController;
+}
+
 bool FWebSocketMessageHandler::WritePropertyChangeEventPayload(URemoteControlPreset* InPreset, const TSet<FGuid>& InModifiedPropertyIds, int64 InSequenceNumber, TArray<uint8>& OutBuffer)
 {
 	bool bHasProperty = false;
diff --git a/Engine/Plugins/VirtualProduction/RemoteControl/Source/WebRemoteControl/Private/WebSocketMessageHandler.h b/Engine/Plugins/VirtualProduction/RemoteControl/Source/WebRemoteControl/Private/WebSocketMessageHandler.h
--- a/Engine/Plugins/VirtualProduction/RemoteControl/Source/WebRemoteControl/Private/WebSocketMessageHandler.h
+++ b/Engine/Plugins/VirtualProduction/RemoteControl/Source/WebRemoteControl/Private/WebSocketMessageHandler.h
@@ -122,6 +122,12 @@ private:
 	void OnEntitiesModified(URemoteControlPreset* Owner, const TSet<FGuid>& ModifiedEntities);
 	void OnLayoutModified(URemoteControlPreset* Owner);
 
+	//Logic callbacks
+	void OnControllerAdded(URemoteControlPreset* Owner, FName NewControllerName, const FGuid& EntityId);
+	void OnControllerRemoved(URemoteControlPreset* Owner, const FGuid& EntityId);
+	void OnControllerRenamed(URemoteControlPreset* Owner, FName OldControllerLabel, FName NewControllerLabel);
+	void OnControllerModified(URemoteControlPreset* Owner, const TSet<FGuid>& ModifiedControllerIds);
+
 	/** Callback when a websocket connection was closed. Let us clean out registrations */
 	void OnConnectionClosedCallback(FGuid ClientId);
 
@@ -131,6 +137,18 @@ private:
 	/** Check if any transactions have timed out and end them. */
 	void TimeOutTransactions();
 
+	/** If controllers have changed during the frame, send out notifications to listeners */
+	void ProcessChangedControllers();
+
+	/** If a new controller has been added during the frame, send out notifications to listeners */
+	void ProcessAddedControllers();
+
+	/** If controllers has been removed during the frame, send out notifications to listeners */
+	void ProcessRemovedControllers();
+
+	/** If a controller has been renamed during the frame, snd out notifications to listeners */
+	void ProcessRenamedControllers();
+	
 	/** If properties have changed during the frame, send out notifications to listeners */
 	void ProcessChangedProperties();
 
@@ -189,6 +207,11 @@ private:
 	 */
 	bool WritePropertyChangeEventPayload(URemoteControlPreset* InPreset, const TSet<FGuid>& InModifiedPropertyIds, int64 InSequenceNumber, TArray<uint8>& OutBuffer);
 
+	/**
+	 * Write the provided list of controller events to a buffer.
+	 */
+	bool WriteControllerChangeEventPayload(URemoteControlPreset* InPreset, const TSet<FGuid>& InModifiedControllerIds, int64 InSequenceNumber, TArray<uint8>& OutBuffer);
+
 	/**
 	 * Write the provided list of actor modifications to a buffer.
 	 */
@@ -353,6 +376,18 @@ private:
 	/** Fields that were renamed for a frame, per preset */
 	TMap<FGuid, TArray<TTuple<FName, FName>>> PerFrameRenamedFields;
 
+	/** Controllers that were Added for a frame, per preset */
+	TMap<FGuid, TArray<FGuid>> PerFrameAddedControllers;
+
+	/** Controllers that were Removed for a frame, per preset */
+	TMap<FGuid, TTuple<TArray<FGuid>, TArray<FName>>> PerFrameRemovedControllers;
+
+	/** Controllers that were Renamed for a frame, per preset */
+	TMap<FGuid, TArray<TTuple<FName, FName>>> PerFrameRenamedControllers;
+
+	/** Controllers that were Modified for a frame, per preset */
+	TMap<FGuid, TMap<FGuid, TSet<FGuid>>> PerFrameModifiedControllers;
+
 	/** Actors that were added for a frame, per watched class. */
 	FChangedActorMap PerFrameActorsAdded;
 
diff --git a/Engine/Plugins/VirtualProduction/RemoteControl/Source/WebRemoteControl/WebRemoteControl.Build.cs b/Engine/Plugins/VirtualProduction/RemoteControl/Source/WebRemoteControl/WebRemoteControl.Build.cs
--- a/Engine/Plugins/VirtualProduction/RemoteControl/Source/WebRemoteControl/WebRemoteControl.Build.cs
+++ b/Engine/Plugins/VirtualProduction/RemoteControl/Source/WebRemoteControl/WebRemoteControl.Build.cs
@@ -21,6 +21,7 @@ public class WebRemoteControl : ModuleRules
 				"Networking",
 				"RemoteControl",
 				"RemoteControlCommon",
+				"RemoteControlLogic",
 				"Sockets",
 				"WebSocketNetworking"
 			}
diff --git a/Engine/Plugins/VirtualProduction/RemoteControlWebInterface/Source/RemoteControlWebInterface/Private/RCWebInterfaceCustomizations.cpp b/Engine/Plugins/VirtualProduction/RemoteControlWebInterface/Source/RemoteControlWebInterface/Private/RCWebInterfaceCustomizations.cpp
--- a/Engine/Plugins/VirtualProduction/RemoteControlWebInterface/Source/RemoteControlWebInterface/Private/RCWebInterfaceCustomizations.cpp
+++ b/Engine/Plugins/VirtualProduction/RemoteControlWebInterface/Source/RemoteControlWebInterface/Private/RCWebInterfaceCustomizations.cpp
@@ -2,6 +2,8 @@
 
 #include "RCWebInterfaceCustomizations.h"
 
+#include "RCVirtualProperty.h"
+
 #if WITH_EDITOR
 #include "DetailCategoryBuilder.h"
 #include "DetailLayoutBuilder.h"
@@ -500,6 +502,11 @@ FString FRCWebInterfaceCustomizations::OnInitializeWidgetMetadata(URemoteControl
 			Metadata = TEXT("Button");
 		}
 	}
+	else if (const URCVirtualPropertyBase* Controller = Preset->GetController(EntityId))
+	{
+		Metadata = RCWebInterface::GetDefaultWidget(const_cast<FProperty*>(Controller->GetProperty()));
+	}
+	
 	return Metadata;
 }
 
