Static Analysis Fix for RemoteControlPropertyHandleTests

Fix an issue in which LinearProperty could be dereferenced.

diff --git a/Engine/Plugins/VirtualProduction/RemoteControl/Source/RemoteControl/Private/Tests/RemoteControlPropertyHandleTests.cpp b/Engine/Plugins/VirtualProduction/RemoteControl/Source/RemoteControl/Private/Tests/RemoteControlPropertyHandleTests.cpp
--- a/Engine/Plugins/VirtualProduction/RemoteControl/Source/RemoteControl/Private/Tests/RemoteControlPropertyHandleTests.cpp
+++ b/Engine/Plugins/VirtualProduction/RemoteControl/Source/RemoteControl/Private/Tests/RemoteControlPropertyHandleTests.cpp
@@ -1003,10 +1003,11 @@ namespace RemoteControlAPIIntegrationTest
 
 		// Set a test value
 		TestObject->StaticMeshComponent->BodyInstance.LinearDamping = 5.0;
-
-		FProperty* LinearProperty = FBodyInstance::StaticStruct()->FindPropertyByName(PROP_NAME(FBodyInstance, LinearDamping));
-		if(!Test.TestFalse(TEXT("LineaProperty is valid"), LinearProperty == nullptr))
+		
+		const FProperty* LinearProperty = FBodyInstance::StaticStruct()->FindPropertyByName(PROP_NAME(FBodyInstance, LinearDamping));
+		if(!LinearProperty)
 		{
+			Test.AddError(TEXT("LinearProperty was not valid"));
 			return;
 		}
 
