Remote control fix Compile Error with non-unity

Fix bug which wasn't caught in prefight.

diff --git a/Engine/Plugins/VirtualProduction/RemoteControl/Source/RemoteControl/Private/Tests/RemoteControlPropertyHandleTestData.h b/Engine/Plugins/VirtualProduction/RemoteControl/Source/RemoteControl/Private/Tests/RemoteControlPropertyHandleTestData.h
--- a/Engine/Plugins/VirtualProduction/RemoteControl/Source/RemoteControl/Private/Tests/RemoteControlPropertyHandleTestData.h
+++ b/Engine/Plugins/VirtualProduction/RemoteControl/Source/RemoteControl/Private/Tests/RemoteControlPropertyHandleTestData.h
@@ -2,10 +2,9 @@
 #pragma once
 
 #include "CoreMinimal.h"
+#include "Components/StaticMeshComponent.h"
 #include "RemoteControlPropertyHandleTestData.generated.h"
 
-class UStaticMeshComponent;
-
 UENUM()
 enum class ERemoteControlEnumClass : uint8
 {
