Fix remote control panel transactions not working

diff --git a/Engine/Plugins/VirtualProduction/RemoteControl/Source/RemoteControlUI/Private/UI/SRemoteControlPanel.cpp b/Engine/Plugins/VirtualProduction/RemoteControl/Source/RemoteControlUI/Private/UI/SRemoteControlPanel.cpp
--- a/Engine/Plugins/VirtualProduction/RemoteControl/Source/RemoteControlUI/Private/UI/SRemoteControlPanel.cpp
+++ b/Engine/Plugins/VirtualProduction/RemoteControl/Source/RemoteControlUI/Private/UI/SRemoteControlPanel.cpp
@@ -906,12 +906,6 @@ void SRemoteControlPanel::ToggleProperty(const FRCExposesPropertyArgs& InPropert
 		return;
 	}
 
-	auto PreExpose = [this]()
-	{
-		FScopedTransaction Transaction(LOCTEXT("ExposeProperty", "Expose Property"));
-		Preset->Modify();
-	};
-
 	auto PostExpose = [this, InPropertyArgs]()
 	{
 		CachedExposedPropertyArgs.Emplace(InPropertyArgs);
@@ -930,7 +924,8 @@ void SRemoteControlPanel::ToggleProperty(const FRCExposesPropertyArgs& InPropert
 
 		if (UniqueOuterObjects.Num())
 		{
-			PreExpose();
+			FScopedTransaction Transaction(LOCTEXT("ExposeProperty", "Expose Property"));
+			Preset->Modify();
 
 			for (UObject* Object : UniqueOuterObjects)
 			{
@@ -946,7 +941,8 @@ void SRemoteControlPanel::ToggleProperty(const FRCExposesPropertyArgs& InPropert
 	}
 	else if (ExtensionArgsType == FRCExposesPropertyArgs::EType::E_OwnerObject)
 	{
-		PreExpose();
+		FScopedTransaction Transaction(LOCTEXT("ExposeProperty", "Expose Property"));
+		Preset->Modify();
 
 		constexpr bool bCleanDuplicates = true; // GeneratePathToProperty duplicates container name (Array.Array[1], Set.Set[1], etc...)
 		ExposeProperty(InPropertyArgs.OwnerObject, FRCFieldPathInfo{ InPropertyArgs.PropertyPath, bCleanDuplicates });
