Remote Control: Fix intermittent crash in remote control preset

diff --git a/Engine/Plugins/VirtualProduction/RemoteControl/Source/RemoteControlUI/Private/UI/RCPanelWidgetRegistry.cpp b/Engine/Plugins/VirtualProduction/RemoteControl/Source/RemoteControlUI/Private/UI/RCPanelWidgetRegistry.cpp
--- a/Engine/Plugins/VirtualProduction/RemoteControl/Source/RemoteControlUI/Private/UI/RCPanelWidgetRegistry.cpp
+++ b/Engine/Plugins/VirtualProduction/RemoteControl/Source/RemoteControlUI/Private/UI/RCPanelWidgetRegistry.cpp
@@ -98,7 +98,7 @@ FRCPanelWidgetRegistry::~FRCPanelWidgetRegistry()
 	{
 		if (Pair.Value)
 		{
-			Pair.Value->OnRowsRefreshed().RemoveAll(this);
+			Pair.Value->OnRowsRefreshed().Clear();
 		}
 	}
 
@@ -140,14 +140,7 @@ TSharedPtr<IDetailTreeNode> FRCPanelWidgetRegistry::GetObjectTreeNode(UObject* I
 	}
 	else
 	{
-		// Since we must keep many PRG objects alive in order to access the handle data, validating the nodes each tick is very taxing.
-		// We can override the validation with a lambda since the validation function in PRG is not necessary for our implementation
-		auto ValidationLambda = ([](const FRootPropertyNodeList& PropertyNodeList) { return true; });
-		FPropertyRowGeneratorArgs Args;
-		Generator = FModuleManager::GetModuleChecked<FPropertyEditorModule>("PropertyEditor").CreatePropertyRowGenerator(Args);
-		Generator->SetCustomValidatePropertyNodesFunction(FOnValidatePropertyRowGeneratorNodes::CreateLambda(MoveTemp(ValidationLambda)));
-		Generator->SetObjects({InObject});
-		Generator->OnRowsRefreshed().AddRaw(this, &FRCPanelWidgetRegistry::OnRowsRefreshed, Generator);
+		Generator = CreateGenerator(InObject);
 		ObjectToRowGenerator.Add(WeakObject, Generator);
 	}
 	
@@ -198,6 +191,14 @@ void FRCPanelWidgetRegistry::Refresh(const TSharedPtr<FStructOnScope>& InStruct)
 
 void FRCPanelWidgetRegistry::Clear()
 {
+	for (const TPair <TWeakObjectPtr<UObject>, TSharedPtr<IPropertyRowGenerator>>& Pair : ObjectToRowGenerator)
+	{
+		if (Pair.Value)
+		{
+			Pair.Value->OnRowsRefreshed().Clear();
+		}
+	}
+
 	ObjectToRowGenerator.Empty();
 	StructToRowGenerator.Empty();
 	TreeNodeCache.Empty();
@@ -250,15 +251,7 @@ TSharedPtr<IDetailTreeNode> FRCPanelWidgetRegistry::FindNDisplayTreeNode(UObject
 	}
 	else
 	{
-		// Since we must keep many PRG objects alive in order to access the handle data, validating the nodes each tick is very taxing.
-		// We can override the validation with a lambda since the validation function in PRG is not necessary for our implementation
-		auto ValidationLambda = ([](const FRootPropertyNodeList& PropertyNodeList) { return true; });
-		FPropertyRowGeneratorArgs Args;
-		Args.bShouldShowHiddenProperties = false;
-		Generator = FModuleManager::GetModuleChecked<FPropertyEditorModule>("PropertyEditor").CreatePropertyRowGenerator(Args);
-		Generator->SetObjects({ InObject });
-		Generator->SetCustomValidatePropertyNodesFunction(FOnValidatePropertyRowGeneratorNodes::CreateLambda(MoveTemp(ValidationLambda)));
-		Generator->OnRowsRefreshed().AddRaw(this, &FRCPanelWidgetRegistry::OnRowsRefreshed, Generator);
+		Generator = CreateGenerator(InObject);
 		ObjectToRowGenerator.Add(WeakObject, Generator);
 	}
 
@@ -279,3 +272,31 @@ void FRCPanelWidgetRegistry::OnRowsRefreshed(TSharedPtr<IPropertyRowGenerator> G
 		OnObjectRefreshedDelegate.Broadcast(SelectedObjects);
 	}
 }
+
+TSharedPtr<IPropertyRowGenerator> FRCPanelWidgetRegistry::CreateGenerator(UObject* InObject)
+{
+	// Since we must keep many PRG objects alive in order to access the handle data, validating the nodes each tick is very taxing.
+	// We can override the validation with a lambda since the validation function in PRG is not necessary for our implementation
+	auto ValidationLambda = ([](const FRootPropertyNodeList& PropertyNodeList) { return true; });
+	FPropertyRowGeneratorArgs Args;
+	TSharedPtr<IPropertyRowGenerator> Generator = FModuleManager::GetModuleChecked<FPropertyEditorModule>("PropertyEditor").CreatePropertyRowGenerator(Args);
+	Generator->SetObjects({ InObject });
+	Generator->SetCustomValidatePropertyNodesFunction(FOnValidatePropertyRowGeneratorNodes::CreateLambda(MoveTemp(ValidationLambda)));
+
+	Generator->OnRowsRefreshed().AddLambda([WeakGenerator = TWeakPtr<IPropertyRowGenerator>(Generator), WeakThis = TWeakPtr<FRCPanelWidgetRegistry>(AsShared())]
+		{
+			// This was crashing on an invalid widget registry ptr, maybe because Clear() getting called and the generator being kept alive by the shared ptr 
+			// passed to the AddRaw delegate was keeping the Genrator and its invocation list alive.
+			// This should be revisited and retested thoroughly after returning to AddRaw or AddSP and passing in a weak generator ptr.
+			if (TSharedPtr<FRCPanelWidgetRegistry> WidgetRegistry = WeakThis.Pin())
+			{
+				if (TSharedPtr<IPropertyRowGenerator> Generator = WeakGenerator.Pin())
+				{
+					WidgetRegistry->OnRowsRefreshed(Generator);
+				}
+			}
+		});
+
+	return Generator;
+}
+
diff --git a/Engine/Plugins/VirtualProduction/RemoteControl/Source/RemoteControlUI/Private/UI/RCPanelWidgetRegistry.h b/Engine/Plugins/VirtualProduction/RemoteControl/Source/RemoteControlUI/Private/UI/RCPanelWidgetRegistry.h
--- a/Engine/Plugins/VirtualProduction/RemoteControl/Source/RemoteControlUI/Private/UI/RCPanelWidgetRegistry.h
+++ b/Engine/Plugins/VirtualProduction/RemoteControl/Source/RemoteControlUI/Private/UI/RCPanelWidgetRegistry.h
@@ -3,6 +3,7 @@
 #pragma once
 
 #include "CoreMinimal.h"
+#include "Templates/SharedPointer.h"
 #include "UObject/WeakObjectPtr.h"
 #include "UObject/WeakObjectPtrTemplates.h"
 
@@ -27,7 +28,7 @@ struct FRCTreeNodeFinderHandler
 	FRCTreeNodeFinderFunction FinderFunction;
 };
 
-class FRCPanelWidgetRegistry
+class FRCPanelWidgetRegistry : public TSharedFromThis<FRCPanelWidgetRegistry>
 {
 public:
 	
@@ -79,6 +80,8 @@ private:
 	bool IsNDisplayObject(UObject* InObject, const FString& InField, ERCFindNodeMethod InFindMethod);
 	TSharedPtr<IDetailTreeNode> FindNDisplayTreeNode(UObject* InObject, const FString& InField, ERCFindNodeMethod InFindMethod);
 	void OnRowsRefreshed(TSharedPtr<IPropertyRowGenerator> Generator);
+	/** Create a property row generator */
+	TSharedPtr<IPropertyRowGenerator> CreateGenerator(UObject* InObject);
 
 private:
 	/** Map of objects to row generator, used to have one row generator per object. */
