Fix remote control preset transactions being ignored by Multi-User

diff --git a/Engine/Plugins/VirtualProduction/RemoteControl/RemoteControl.uplugin b/Engine/Plugins/VirtualProduction/RemoteControl/RemoteControl.uplugin
--- a/Engine/Plugins/VirtualProduction/RemoteControl/RemoteControl.uplugin
+++ b/Engine/Plugins/VirtualProduction/RemoteControl/RemoteControl.uplugin
@@ -25,6 +25,10 @@
 		{
 			"Name": "StructUtils",
 			"Enabled": true
+		},
+		{
+			"Name": "ConcertSyncClient",
+			"Enabled": true
 		}
 	],
 	"Modules": [
@@ -67,6 +71,11 @@
 			"Name": "RemoteControlProtocolWidgets",
 			"Type": "Editor",
 			"LoadingPhase": "Default"
+		},
+		{
+			"Name": "RemoteControlMultiUser",
+			"Type": "UncookedOnly",
+			"LoadingPhase": "PostEngineInit"
 		}
 	]
 }
\ No newline at end of file
diff --git a/Engine/Plugins/VirtualProduction/RemoteControl/Source/RemoteControl/Private/RemoteControlPreset.cpp b/Engine/Plugins/VirtualProduction/RemoteControl/Source/RemoteControl/Private/RemoteControlPreset.cpp
--- a/Engine/Plugins/VirtualProduction/RemoteControl/Source/RemoteControl/Private/RemoteControlPreset.cpp
+++ b/Engine/Plugins/VirtualProduction/RemoteControl/Source/RemoteControl/Private/RemoteControlPreset.cpp
@@ -800,7 +800,7 @@ URemoteControlBinding* URemoteControlPreset::FindOrAddBinding(const TSoftObjectP
 	
 	for (URemoteControlBinding* Binding : Bindings)
 	{
-		if (Binding->IsBound(Object))
+		if (Binding && Binding->IsBound(Object))
 		{
 			return Binding;
 		}
@@ -857,7 +857,8 @@ URemoteControlBinding* URemoteControlPreset::FindMatchingBinding(const URemoteCo
 		URemoteControlLevelDependantBinding* LevelDependantBindingIt = Cast<URemoteControlLevelDependantBinding>(Binding);
 		const URemoteControlLevelDependantBinding* InLevelDependingBinding = Cast<URemoteControlLevelDependantBinding>(InBinding);
 
-		if (InBinding == Binding
+		if (!Binding 
+			|| InBinding == Binding
 			|| Binding->Resolve() != InObject)
 		{
 			continue;
@@ -1352,10 +1353,12 @@ void URemoteControlPreset::Unexpose(const FGuid& EntityId)
 
 		Registry->Modify();
 		Registry->RemoveExposedEntity(EntityId);
-		FRCCachedFieldData CachedData = FieldCache.FindChecked(EntityId);
-		Layout.RemoveField(CachedData.LayoutGroupId, EntityId);
-		FieldCache.Remove(EntityId);
-		PropertyWatchers.Remove(EntityId);
+		if (FRCCachedFieldData* CachedData = FieldCache.Find(EntityId))
+		{
+			Layout.RemoveField(CachedData->LayoutGroupId, EntityId);
+			FieldCache.Remove(EntityId);
+			PropertyWatchers.Remove(EntityId);
+		}
 	}
 }
 
@@ -1721,7 +1724,7 @@ void URemoteControlPreset::OnObjectTransacted(UObject* InObject, const FTransact
 	bool bIsObjectBound = false;
 	for (const URemoteControlBinding* Binding : Bindings)
 	{
-		if (Binding->IsBound(InObject))
+		if (Binding && Binding->IsBound(InObject))
 		{
 			bIsObjectBound = true;
 			break;
@@ -1910,6 +1913,11 @@ void URemoteControlPreset::OnReplaceObjects(const TMap<UObject*, UObject*>& Repl
 
 	for (URemoteControlBinding* Binding : Bindings)
 	{
+		if (!Binding)
+		{
+			continue;
+		}
+
 		UObject* NewObject = nullptr;
 		UObject* ResolvedBinding = Binding->Resolve();
 
diff --git a/Engine/Plugins/VirtualProduction/RemoteControl/Source/RemoteControlMultiUser/Private/RemoteControlMultiUserModule.cpp b/Engine/Plugins/VirtualProduction/RemoteControl/Source/RemoteControlMultiUser/Private/RemoteControlMultiUserModule.cpp
new file mode 100644
--- /dev/null
+++ b/Engine/Plugins/VirtualProduction/RemoteControl/Source/RemoteControlMultiUser/Private/RemoteControlMultiUserModule.cpp
@@ -0,0 +1,43 @@
+// Copyright Epic Games, Inc. All Rights Reserved.
+
+#include "RemoteControlMultiUserModule.h"
+
+#include "IConcertClientTransactionBridge.h"
+#include "IConcertSyncClientModule.h"
+#include "RemoteControlPreset.h"
+
+
+namespace RemoteControlMultiUserUtils
+{
+	ETransactionFilterResult HandleTransactionFiltering(UObject* ObjectToFilter, UPackage* ObjectsPackage)
+	{
+		//Always allow RemoteControlPresets and RemoteControlExposeRegistry
+		static const FName ExposeRegistryName = "RemoteControlExposeRegistry";
+		if(Cast<URemoteControlPreset>(ObjectToFilter) != nullptr || (ObjectToFilter && ObjectToFilter->GetClass()->GetFName() == ExposeRegistryName))
+		{
+			return ETransactionFilterResult::IncludeObject;
+		}
+
+		return ETransactionFilterResult::UseDefault;
+	}
+}
+
+void FRemoteControlMultiUserModule::StartupModule()
+{
+	if (IConcertSyncClientModule::IsAvailable())
+	{
+		IConcertClientTransactionBridge& TransactionBridge = IConcertSyncClientModule::Get().GetTransactionBridge();
+		TransactionBridge.RegisterTransactionFilter( TEXT("RemoteControlTransactionFilter"), FTransactionFilterDelegate::CreateStatic(&RemoteControlMultiUserUtils::HandleTransactionFiltering));
+	}
+}
+
+void FRemoteControlMultiUserModule::ShutdownModule()
+{
+	if (IConcertSyncClientModule::IsAvailable())
+	{
+		IConcertClientTransactionBridge& TransactionBridge = IConcertSyncClientModule::Get().GetTransactionBridge();
+		TransactionBridge.UnregisterTransactionFilter( TEXT("RemoteControlTransactionFilter"));
+	}	
+}
+
+IMPLEMENT_MODULE(FRemoteControlMultiUserModule, RemoteControlMultiUser);
diff --git a/Engine/Plugins/VirtualProduction/RemoteControl/Source/RemoteControlMultiUser/Private/RemoteControlMultiUserModule.h b/Engine/Plugins/VirtualProduction/RemoteControl/Source/RemoteControlMultiUser/Private/RemoteControlMultiUserModule.h
new file mode 100644
--- /dev/null
+++ b/Engine/Plugins/VirtualProduction/RemoteControl/Source/RemoteControlMultiUser/Private/RemoteControlMultiUserModule.h
@@ -0,0 +1,18 @@
+// Copyright Epic Games, Inc. All Rights Reserved.
+
+#pragma once
+
+#include "Modules/ModuleManager.h"
+
+/**
+ * Implements a simple Remote Control Multiuser module to register our own transaction filter.
+ */
+class FRemoteControlMultiUserModule : public IModuleInterface
+{
+public:
+
+	//~ Begin IModuleInterface interface
+	virtual void StartupModule() override;
+	virtual void ShutdownModule() override;
+	//~ End IModuleInterface interface
+};
diff --git a/Engine/Plugins/VirtualProduction/RemoteControl/Source/RemoteControlMultiUser/RemoteControlMultiUser.Build.cs b/Engine/Plugins/VirtualProduction/RemoteControl/Source/RemoteControlMultiUser/RemoteControlMultiUser.Build.cs
new file mode 100644
--- /dev/null
+++ b/Engine/Plugins/VirtualProduction/RemoteControl/Source/RemoteControlMultiUser/RemoteControlMultiUser.Build.cs
@@ -0,0 +1,24 @@
+// Copyright Epic Games, Inc. All Rights Reserved.
+
+namespace UnrealBuildTool.Rules
+{
+	public class RemoteControlMultiUser : ModuleRules
+	{
+		public RemoteControlMultiUser(ReadOnlyTargetRules Target) : base(Target)
+		{
+			PublicDependencyModuleNames.AddRange(
+			new string[]
+			{
+			});
+
+			PrivateDependencyModuleNames.AddRange(
+			new string[]
+			{
+				"ConcertSyncClient", // Dependancy to register our transaction filters. Needs to enable that plugin. Would be better to have modular feature interface functionality
+				"Core",
+				"CoreUObject",
+				"RemoteControl",
+			});
+		}
+	}
+}
