Change away the Range Map Normalized Input to ones fitting the Min and Max Value

A small change which changes the current handling of Range Map, going from 0 - 1 as normalized values between an arbritrary Min and Max value, to an actual value set between it.

diff --git a/Engine/Plugins/VirtualProduction/RemoteControl/Source/RemoteControlLogic/Private/Behaviour/Builtin/RangeMap/RCRangeMapBehaviour.cpp b/Engine/Plugins/VirtualProduction/RemoteControl/Source/RemoteControlLogic/Private/Behaviour/Builtin/RangeMap/RCRangeMapBehaviour.cpp
--- a/Engine/Plugins/VirtualProduction/RemoteControl/Source/RemoteControlLogic/Private/Behaviour/Builtin/RangeMap/RCRangeMapBehaviour.cpp
+++ b/Engine/Plugins/VirtualProduction/RemoteControl/Source/RemoteControlLogic/Private/Behaviour/Builtin/RangeMap/RCRangeMapBehaviour.cpp
@@ -2,6 +2,7 @@
 
 #include "Behaviour/Builtin/RangeMap/RCRangeMapBehaviour.h"
 
+#include "IRemoteControlModule.h"
 #include "Action/RCAction.h"
 #include "Action/RCActionContainer.h"
 #include "Action/RCFunctionAction.h"
@@ -112,7 +113,8 @@ bool URCRangeMapBehaviour::GetNearestActionByThreshold(TTuple<URCAction*, bool>&
 	for (int InputIndex = 0; InputIndex < InputActionArray.Num(); InputIndex++)
 	{
 		const double InputValue = InputActionArray[InputIndex];
-		const double ValueDifference = FMath::Abs(InputValue - NormalizedControllerValue);
+		const double NormalizedInputValue = UKismetMathLibrary::NormalizeToRange(InputValue, InputMin, InputMax);
+		const double ValueDifference = FMath::Abs(NormalizedInputValue - NormalizedControllerValue);
 		
 		if (InputActionArray.Num() > 1 && InputIndex > 0)
 		{
@@ -124,7 +126,7 @@ bool URCRangeMapBehaviour::GetNearestActionByThreshold(TTuple<URCAction*, bool>&
 				break;
 			}
 		}
-		OutTuple = TTuple<URCAction*, bool>(NonLerpActions[InputValue], ValueDifference <= UE::RCRangeMapBehaviour::Threshold);
+		OutTuple = TTuple<URCAction*, bool>(NonLerpActions[NormalizedInputValue], ValueDifference <= UE::RCRangeMapBehaviour::Threshold);
 	}
 	
 	return true;
@@ -197,6 +199,12 @@ void URCRangeMapBehaviour::Execute()
 		return; // Allow custom Blueprints to drive their own behaviour entirely
 	}
 
+	if (InputMax <= InputMin)
+	{
+		UE_LOG(LogRemoteControl, Warning, TEXT("Mapping Behaviour Input Max cannot be equal or smaller than the Minimum."))
+		return;
+	}
+
 	// Do this beforehand.
 	BehaviourNode->PreExecute(this);
 	
@@ -261,8 +269,8 @@ void URCRangeMapBehaviour::Execute()
 		}
 
 		// Denormalize and Map them based on our Min and Max Value
-		const double MappedMinInput = UKismetMathLibrary::Lerp(InputMin, InputMax, MinRangeValue);
-		const double MappedMaxInput = UKismetMathLibrary::Lerp(InputMin, InputMax, MaxRangeValue);
+		const double MappedMinInput = UKismetMathLibrary::Lerp(InputMin, InputMax, UKismetMathLibrary::NormalizeToRange(MinRangeValue, InputMin, InputMax));
+		const double MappedMaxInput = UKismetMathLibrary::Lerp(InputMin, InputMax, UKismetMathLibrary::NormalizeToRange(MaxRangeValue, InputMin, InputMax));
 
 		// Normalize our Controller based on the new MappedMinInput and MappedMaxInput Range.
 		double CustomNormalizedInputValue = UKismetMathLibrary::NormalizeToRange(ControllerFloatValue, MappedMinInput, MappedMaxInput);
@@ -403,8 +411,10 @@ bool URCRangeMapBehaviour::GetRangeValuePairsForLerp(TMap<FGuid, TTuple<URCActio
 				continue;
 			}
 
+			const double NormalizedValue = UKismetMathLibrary::NormalizeToRange(Value, InputMin, InputMax);
+
 			// Deal with the nullptr
-			if (Value <= NormalizedControllerValue && NumericActionTuple.Value.Last() != CurrentAction)
+			if (NormalizedValue <= NormalizedControllerValue && NumericActionTuple.Value.Last() != CurrentAction)
 			{
 				MinAction = CurrentAction;
 				MinRangeMap = RangeMap;
@@ -412,11 +422,16 @@ bool URCRangeMapBehaviour::GetRangeValuePairsForLerp(TMap<FGuid, TTuple<URCActio
 				continue;
 			}
 
-			if (NormalizedControllerValue <= Value)
+			if (NormalizedControllerValue <= NormalizedValue)
 			{
 				MaxAction = CurrentAction;
 				MaxRangeMap = RangeMap;
 
+				if (MinRangeMap)
+				{
+					break;
+				}
+				
 				continue;
 			}
 
@@ -443,10 +458,10 @@ TMap<double, URCAction*> URCRangeMapBehaviour::GetNonLerpActions()
 		{
 			if (const FRCRangeMapInput* RangeMapInput = RangeMapActionContainer.Find(Action))
 			{
-				double StepValue;
-				if (RangeMapInput->GetInputValue(StepValue))
+				double InputValue;
+				if (RangeMapInput->GetInputValue(InputValue))
 				{
-					ActionMap.Add(StepValue, Action);
+					ActionMap.Add(InputValue, Action);
 				}
 			}
 		}
