Remote Control Panel - Button fix for lost bindings for Logic/Protocl Button

diff --git a/Engine/Plugins/VirtualProduction/RemoteControl/Source/RemoteControlUI/Private/RemoteControlUIModule.cpp b/Engine/Plugins/VirtualProduction/RemoteControl/Source/RemoteControlUI/Private/RemoteControlUIModule.cpp
--- a/Engine/Plugins/VirtualProduction/RemoteControl/Source/RemoteControlUI/Private/RemoteControlUIModule.cpp
+++ b/Engine/Plugins/VirtualProduction/RemoteControl/Source/RemoteControlUI/Private/RemoteControlUIModule.cpp
@@ -2,16 +2,17 @@
 
 #include "RemoteControlUIModule.h"
 
-#include "AssetToolsModule.h"
+#include "AssetEditor/RemoteControlPresetEditorToolkit.h"
 #include "AssetTools/RemoteControlPresetActions.h"
-#include "Styling/AppStyle.h"
+#include "AssetToolsModule.h"
 #include "Commands/RemoteControlCommands.h"
 #include "Components/SkeletalMeshComponent.h"
-#include "Engine/StaticMesh.h"
 #include "Engine/SkeletalMesh.h"
+#include "Engine/StaticMesh.h"
 #include "Framework/MultiBox/MultiBoxBuilder.h"
 #include "HAL/PlatformApplicationMisc.h"
 #include "IRemoteControlModule.h"
+#include "Interfaces/IMainFrameModule.h"
 #include "Kismet2/ComponentEditorUtils.h"
 #include "MaterialEditor/DEditorParameterValue.h"
 #include "Materials/Material.h"
@@ -21,19 +22,19 @@
 #include "RemoteControlInstanceMaterial.h"
 #include "RemoteControlPreset.h"
 #include "RemoteControlSettings.h"
-#include "AssetEditor/RemoteControlPresetEditorToolkit.h"
 #include "ScopedTransaction.h"
 #include "Styling/AppStyle.h"
+#include "Styling/AppStyle.h"
 #include "Textures/SlateIcon.h"
 #include "UI/Action/SRCActionPanel.h"
-#include "UI/Behaviour/SRCBehaviourPanel.h"
 #include "UI/Behaviour/Builtin/Conditional/SRCBehaviourConditional.h"
+#include "UI/Behaviour/SRCBehaviourPanel.h"
 #include "UI/Customizations/FPassphraseCustomization.h"
 #include "UI/Customizations/RemoteControlEntityCustomization.h"
-#include "UI/SRCPanelExposedField.h"
-#include "UI/SRCPanelExposedActor.h"
 #include "UI/RemoteControlPanelStyle.h"
+#include "UI/SRCPanelExposedActor.h"
 #include "UI/SRCPanelExposedEntitiesList.h"
+#include "UI/SRCPanelExposedField.h"
 #include "UI/SRemoteControlPanel.h"
 #include "Widgets/DeclarativeSyntaxSupport.h"
 
@@ -443,6 +444,25 @@ void FRemoteControlUIModule::BindRemoteControlCommands()
 
 void FRemoteControlUIModule::UnbindRemoteControlCommands()
 {
+	if (FModuleManager::Get().IsModuleLoaded("MainFrame"))
+	{
+		const FRemoteControlCommands& Commands = FRemoteControlCommands::Get();
+
+		IMainFrameModule& MainFrame = FModuleManager::Get().LoadModuleChecked<IMainFrameModule>("MainFrame");
+
+		FUICommandList& ActionList = *MainFrame.GetMainFrameCommandBindings();
+
+		ActionList.UnmapAction(Commands.SavePreset);
+		ActionList.UnmapAction(Commands.FindPresetInContentBrowser);
+		ActionList.UnmapAction(Commands.ToggleProtocolMappings);
+		ActionList.UnmapAction(Commands.ToggleLogicEditor);
+		ActionList.UnmapAction(Commands.DeleteEntity);
+		ActionList.UnmapAction(Commands.RenameEntity);
+		ActionList.UnmapAction(Commands.CopyItem);
+		ActionList.UnmapAction(Commands.PasteItem);
+		ActionList.UnmapAction(Commands.DuplicateItem);
+	}
+	
 	FRemoteControlCommands::Unregister();
 }
 
diff --git a/Engine/Plugins/VirtualProduction/RemoteControl/Source/RemoteControlUI/Private/UI/SRemoteControlPanel.cpp b/Engine/Plugins/VirtualProduction/RemoteControl/Source/RemoteControlUI/Private/UI/SRemoteControlPanel.cpp
--- a/Engine/Plugins/VirtualProduction/RemoteControl/Source/RemoteControlUI/Private/UI/SRemoteControlPanel.cpp
+++ b/Engine/Plugins/VirtualProduction/RemoteControl/Source/RemoteControlUI/Private/UI/SRemoteControlPanel.cpp
@@ -770,7 +770,6 @@ SRemoteControlPanel::~SRemoteControlPanel()
 {
 	SaveSettings();
 
-	UnbindRemoteControlCommands();
 	UnregisterPanels();
 	UnregisterEvents();
 
@@ -1385,25 +1384,6 @@ void SRemoteControlPanel::BindRemoteControlCommands()
 		FCanExecuteAction::CreateSP(this, &SRemoteControlPanel::CanDuplicateItem));
 }
 
-void SRemoteControlPanel::UnbindRemoteControlCommands()
-{
-	if (FModuleManager::Get().IsModuleLoaded("MainFrame"))
-	{
-		const FRemoteControlCommands& Commands = FRemoteControlCommands::Get();
-
-		IMainFrameModule& MainFrame = FModuleManager::Get().LoadModuleChecked<IMainFrameModule>("MainFrame");
-
-		FUICommandList& ActionList = *MainFrame.GetMainFrameCommandBindings();
-
-		ActionList.UnmapAction(Commands.SavePreset);
-		ActionList.UnmapAction(Commands.FindPresetInContentBrowser);
-		ActionList.UnmapAction(Commands.ToggleProtocolMappings);
-		ActionList.UnmapAction(Commands.ToggleLogicEditor);
-		ActionList.UnmapAction(Commands.DeleteEntity);
-		ActionList.UnmapAction(Commands.RenameEntity);
-	}
-}
-
 void SRemoteControlPanel::RegisterEvents()
 {
 	FEditorDelegates::MapChange.AddSP(this, &SRemoteControlPanel::OnMapChange);
diff --git a/Engine/Plugins/VirtualProduction/RemoteControl/Source/RemoteControlUI/Private/UI/SRemoteControlPanel.h b/Engine/Plugins/VirtualProduction/RemoteControl/Source/RemoteControlUI/Private/UI/SRemoteControlPanel.h
--- a/Engine/Plugins/VirtualProduction/RemoteControl/Source/RemoteControlUI/Private/UI/SRemoteControlPanel.h
+++ b/Engine/Plugins/VirtualProduction/RemoteControl/Source/RemoteControlUI/Private/UI/SRemoteControlPanel.h
@@ -166,7 +166,6 @@ private:
 
 	//~ Remote Control Commands
 	void BindRemoteControlCommands();
-	void UnbindRemoteControlCommands();
 
 	/** Register editor events needed to handle reloading objects and blueprint libraries. */
 	void RegisterEvents();
