Deleting a Property Forces the User Back to All Groups View

diff --git a/Engine/Plugins/VirtualProduction/RemoteControl/Source/RemoteControlUI/Private/UI/SRCPanelExposedEntitiesList.cpp b/Engine/Plugins/VirtualProduction/RemoteControl/Source/RemoteControlUI/Private/UI/SRCPanelExposedEntitiesList.cpp
--- a/Engine/Plugins/VirtualProduction/RemoteControl/Source/RemoteControlUI/Private/UI/SRCPanelExposedEntitiesList.cpp
+++ b/Engine/Plugins/VirtualProduction/RemoteControl/Source/RemoteControlUI/Private/UI/SRCPanelExposedEntitiesList.cpp
@@ -187,6 +187,14 @@ void SRCPanelExposedEntitiesList::Construct(const FArguments& InArgs, URemoteCon
 	RCPanelStyle = &FRemoteControlPanelStyle::Get()->GetWidgetStyle<FRCPanelStyle>("RemoteControlPanel.MinorPanel");
 	ActiveListMode = EEntitiesListMode::Default;
 
+	// The Default group will be selected by default.
+	if (Preset.IsValid())
+	{
+		const FRemoteControlPresetGroup& DefaultGroup = Preset->Layout.GetDefaultGroup();
+
+		CurrentlySelectedGroup = DefaultGroup.Id;
+	}
+
 	// Major Panel
 	TSharedPtr<SRCMajorPanel> ExposePanel = SNew(SRCMajorPanel)
 		.HeaderLabel(LOCTEXT("ExposePanelHeader", "EXPOSE"))
@@ -655,9 +663,17 @@ void SRCPanelExposedEntitiesList::Refresh()
 	{
 		constexpr bool bForceMouseClick = true;
 
-		FRemoteControlPresetGroup& DefaultGroup = Preset->Layout.GetDefaultGroup();
-		GenerateListWidgets(DefaultGroup);
-		SetSelection(FindGroupById(DefaultGroup.Id), bForceMouseClick);
+		if (const FRemoteControlPresetGroup* SelectedGroup = Preset->Layout.GetGroup(CurrentlySelectedGroup))
+		{
+			GenerateListWidgets(*SelectedGroup);
+			SetSelection(FindGroupById(SelectedGroup->Id), bForceMouseClick);
+		}
+		else
+		{
+			const FRemoteControlPresetGroup& DefaultGroup = Preset->Layout.GetDefaultGroup();
+			GenerateListWidgets(DefaultGroup);
+			SetSelection(FindGroupById(DefaultGroup.Id), bForceMouseClick);
+		}
 	}
 }
 
@@ -855,6 +871,8 @@ void SRCPanelExposedEntitiesList::OnSelectionChanged(TSharedPtr<SRCPanelTreeNode
 	{
 		if (FRemoteControlPresetGroup* RCGroup = Preset->Layout.GetGroup(Node->GetRCId()))
 		{
+			CurrentlySelectedGroup = RCGroup->Id;
+
 			GenerateListWidgets(*RCGroup);
 			
 			if (BackendFilter.HasAnyActiveFilters())
diff --git a/Engine/Plugins/VirtualProduction/RemoteControl/Source/RemoteControlUI/Private/UI/SRCPanelExposedEntitiesList.h b/Engine/Plugins/VirtualProduction/RemoteControl/Source/RemoteControlUI/Private/UI/SRCPanelExposedEntitiesList.h
--- a/Engine/Plugins/VirtualProduction/RemoteControl/Source/RemoteControlUI/Private/UI/SRCPanelExposedEntitiesList.h
+++ b/Engine/Plugins/VirtualProduction/RemoteControl/Source/RemoteControlUI/Private/UI/SRCPanelExposedEntitiesList.h
@@ -234,4 +234,6 @@ private:
 	FName ActiveProtocol;
 	/** Columns to be present by default in protocols mode. */
 	static TSet<FName> DefaultProtocolColumns;
+	/** Holds identifier of the selected group. */
+	FGuid CurrentlySelectedGroup;
 };
\ No newline at end of file
