Smart scaling of value widgets for better user experience

diff --git a/Engine/Plugins/VirtualProduction/RemoteControl/Source/RemoteControlUI/Private/UI/Controller/RCControllerModel.cpp b/Engine/Plugins/VirtualProduction/RemoteControl/Source/RemoteControlUI/Private/UI/Controller/RCControllerModel.cpp
--- a/Engine/Plugins/VirtualProduction/RemoteControl/Source/RemoteControlUI/Private/UI/Controller/RCControllerModel.cpp
+++ b/Engine/Plugins/VirtualProduction/RemoteControl/Source/RemoteControlUI/Private/UI/Controller/RCControllerModel.cpp
@@ -37,8 +37,7 @@ TSharedRef<SWidget> FRCControllerModel::GetWidget() const
 	{
 		FieldWidget->AddSlot()
 			.Padding(FMargin(10.0f, 2.0f))
-			.HAlign(HAlign_Right)
-			.AutoWidth()
+			.HAlign(HAlign_Left)
 			[
 				NodeWidgets.ValueWidget.ToSharedRef()
 			];
@@ -47,7 +46,6 @@ TSharedRef<SWidget> FRCControllerModel::GetWidget() const
 	{
 		FieldWidget->AddSlot()
 			.Padding(FMargin(10.0f, 2.0f))
-			.AutoWidth()
 			[
 				NodeWidgets.WholeRowWidget.ToSharedRef()
 			];
diff --git a/Engine/Plugins/VirtualProduction/RemoteControl/Source/RemoteControlUI/Private/UI/RCUIHelpers.cpp b/Engine/Plugins/VirtualProduction/RemoteControl/Source/RemoteControlUI/Private/UI/RCUIHelpers.cpp
--- a/Engine/Plugins/VirtualProduction/RemoteControl/Source/RemoteControlUI/Private/UI/RCUIHelpers.cpp
+++ b/Engine/Plugins/VirtualProduction/RemoteControl/Source/RemoteControlUI/Private/UI/RCUIHelpers.cpp
@@ -139,7 +139,6 @@ TSharedRef<SWidget> UE::RCUIHelpers::GetGenericFieldWidget(const TSharedPtr<IDet
 		FieldWidget->AddSlot()
 			.Padding(FMargin(3.0f, 2.0f))
 			.VAlign(VAlign_Center)
-			.AutoWidth()
 			[
 				NodeWidgets.ValueWidget.ToSharedRef()
 			];
@@ -153,7 +152,6 @@ TSharedRef<SWidget> UE::RCUIHelpers::GetGenericFieldWidget(const TSharedPtr<IDet
 		FieldWidget->AddSlot()
 			.Padding(FMargin(3.0f, 2.0f))
 			.VAlign(VAlign_Center)
-			.AutoWidth()
 			[
 				NodeWidgets.WholeRowWidget.ToSharedRef()
 			];
diff --git a/Engine/Plugins/VirtualProduction/RemoteControl/Source/RemoteControlUI/Private/UI/SRCPanelExposedField.cpp b/Engine/Plugins/VirtualProduction/RemoteControl/Source/RemoteControlUI/Private/UI/SRCPanelExposedField.cpp
--- a/Engine/Plugins/VirtualProduction/RemoteControl/Source/RemoteControlUI/Private/UI/SRCPanelExposedField.cpp
+++ b/Engine/Plugins/VirtualProduction/RemoteControl/Source/RemoteControlUI/Private/UI/SRCPanelExposedField.cpp
@@ -44,8 +44,7 @@ namespace ExposedFieldUtils
 		{
 			FieldWidget->AddSlot()
 				.Padding(FMargin(3.0f, 2.0f))
-				.HAlign(HAlign_Right)
-				.AutoWidth()
+				.HAlign(HAlign_Left)
 				[
 					NodeWidgets.ValueWidget.ToSharedRef()
 				];
@@ -54,7 +53,6 @@ namespace ExposedFieldUtils
 		{
 			FieldWidget->AddSlot()
 				.Padding(FMargin(3.0f, 2.0f))
-				.AutoWidth()
 				[
 					NodeWidgets.WholeRowWidget.ToSharedRef()
 				];
