Fix remote control unexpose not working on ndisplay viewport

diff --git a/Engine/Plugins/VirtualProduction/RemoteControl/Source/RemoteControlUI/Private/UI/SRemoteControlPanel.cpp b/Engine/Plugins/VirtualProduction/RemoteControl/Source/RemoteControlUI/Private/UI/SRemoteControlPanel.cpp
--- a/Engine/Plugins/VirtualProduction/RemoteControl/Source/RemoteControlUI/Private/UI/SRemoteControlPanel.cpp
+++ b/Engine/Plugins/VirtualProduction/RemoteControl/Source/RemoteControlUI/Private/UI/SRemoteControlPanel.cpp
@@ -1557,7 +1557,7 @@ void SRemoteControlPanel::Unexpose(const FRCExposesPropertyArgs& InPropertyArgs)
 		return;
 	}
 
-	auto CheckAndExpose = [&](TArray<UObject*> InOuterObjects, const FString& InPath)
+	auto CheckAndUnexpose = [&](TArray<UObject*> InOuterObjects, const FString& InPath, bool bInUsingDuplicatesInPath)
 	{
 		// Find an exposed property with the same path.
 		TArray<TSharedPtr<FRemoteControlProperty>, TInlineAllocator<1>> PotentialMatches;
@@ -1565,7 +1565,10 @@ void SRemoteControlPanel::Unexpose(const FRCExposesPropertyArgs& InPropertyArgs)
 		{
 			if (TSharedPtr<FRemoteControlProperty> Property = WeakProperty.Pin())
 			{
-				if (Property->CheckIsBoundToString(InPath))
+				// If that was exposed by property path it should be checked by the full path with duplicated like propertypath.propertypath[0]
+				// If that was exposed by the owner object it should be without duplicated in the path, just propertypath[0]
+				const bool bIsbound = bInUsingDuplicatesInPath ? Property->CheckIsBoundToPropertyPath(InPath) : Property->CheckIsBoundToString(InPath);
+				if (bIsbound)
 				{
 					PotentialMatches.Add(Property);
 				}
@@ -1589,11 +1592,13 @@ void SRemoteControlPanel::Unexpose(const FRCExposesPropertyArgs& InPropertyArgs)
 		TArray<UObject*> OuterObjects;
 		InPropertyArgs.PropertyHandle->GetOuterObjects(OuterObjects);
 
-		CheckAndExpose(OuterObjects, InPropertyArgs.PropertyHandle->GeneratePathToProperty());
+		constexpr bool bUsingDuplicatesInPath = true;
+		CheckAndUnexpose(OuterObjects, InPropertyArgs.PropertyHandle->GeneratePathToProperty(), bUsingDuplicatesInPath);
 	}
 	else if (ExtensionArgsType == FRCExposesPropertyArgs::EType::E_OwnerObject)
 	{
-		CheckAndExpose({ InPropertyArgs.OwnerObject }, InPropertyArgs.PropertyPath);
+		constexpr bool bUsingDuplicatesInPath = false;
+		CheckAndUnexpose({ InPropertyArgs.OwnerObject }, InPropertyArgs.PropertyPath, bUsingDuplicatesInPath);
 	}
 }
 
