[EpicStageApp] Fix cases where color grading delta functions would fail or produce incorrect results

diff --git a/Engine/Plugins/VirtualProduction/RemoteControl/Source/RemoteControl/Private/RemoteControlFunctionLibrary.cpp b/Engine/Plugins/VirtualProduction/RemoteControl/Source/RemoteControl/Private/RemoteControlFunctionLibrary.cpp
--- a/Engine/Plugins/VirtualProduction/RemoteControl/Source/RemoteControl/Private/RemoteControlFunctionLibrary.cpp
+++ b/Engine/Plugins/VirtualProduction/RemoteControl/Source/RemoteControl/Private/RemoteControlFunctionLibrary.cpp
@@ -1,6 +1,7 @@
 ﻿// Copyright Epic Games, Inc. All Rights Reserved.
 
 #include "RemoteControlFunctionLibrary.h"
+#include "IRemoteControlModule.h"
 
 namespace RemoteControlFunctionLibrary
 {
@@ -36,64 +37,41 @@ namespace RemoteControlFunctionLibrary
 			return false;
 		}
 
-		// Split the property name up on dots so we can find nested struct properties
-		TArray<FString> PropertyPathNames;
-		PropertyName.ParseIntoArray(PropertyPathNames, TEXT("."));
+		FRCObjectReference ObjectRef;
+		IRemoteControlModule::Get().ResolveObject(ERCAccess::WRITE_ACCESS, Object->GetPathName(), PropertyName, ObjectRef);
 
-		if (PropertyPathNames.IsEmpty())
+		if (!ObjectRef.IsValid())
 		{
 			return false;
 		}
 
-		void* Container = Object;
-		UStruct* ContainerClass = Object->GetClass();
-		if (!ContainerClass)
-		{
-			return false;
-		}
+		OutProperty = ObjectRef.Property.Get();
+		OutContainer = ObjectRef.ContainerAdress;
+		return true;
+	}
 
-		FProperty* Property = ContainerClass->FindPropertyByName(FName(*PropertyPathNames[0]));
-		if (!Property)
-		{
-			return false;
-		}
+	/** Convert a value from the range for a normal linear color (0-1) to the range for a color grading wheel (provided by MinValue and MaxValue). */
+	void TransformLinearColorRangeToColorGradingRange(FVector4& Value, float MinValue, float MaxValue)
+	{
+		check(MaxValue > MinValue);
 
-		// Walk through the chain of structs until we reach the end (or run out of matching struct properties)
-		for (int32 PathIndex = 1; PathIndex < PropertyPathNames.Num(); ++PathIndex)
-		{
-			if (const FStructProperty* StructProperty = CastField<FStructProperty>(Property))
-			{
-				ContainerClass = StructProperty->Struct;
-				if (ContainerClass == nullptr)
-				{
-					return false;
-				}
-
-				Property = ContainerClass->FindPropertyByName(FName(*PropertyPathNames[PathIndex]));
-				if (Property == nullptr)
-				{
-					return false;
-				}
-
-				Container = StructProperty->ContainerPtrToValuePtr<void>(Container);
-			}
-			else
-			{
-				return false;
-			}
-		}
+		Value *= (MaxValue - MinValue);
+		Value += FVector4(MinValue, MinValue, MinValue, MinValue);
+	}
 
-		ensure(Container && Property);
+	/** Convert a value from the range for a color grading wheel (provided by MinValue and MaxValue) to the range for a normal linear color (0-1). */
+	void TransformColorGradingRangeToLinearColorRange(FVector4& Value, float MinValue, float MaxValue)
+	{
+		check(MaxValue > MinValue);
 
-		OutProperty = Property;
-		OutContainer = Container;
-		return true;
+		Value -= FVector4(MinValue, MinValue, MinValue, MinValue);
+		Value /= (MaxValue - MinValue);
 	}
 
 	/**
 	 * Apply a wheel-based color delta to an RGB linear color.
 	 */
-	void ApplyWheelColorBaseDelta(FLinearColor& InOutColor, const FColorWheelColorBase& DeltaValue, const FColorWheelColorBase& ReferenceColor, float MaxValue)
+	void ApplyWheelColorBaseDelta(FLinearColor& InOutColor, const FColorWheelColorBase& DeltaValue, const FColorWheelColorBase& ReferenceColor, float MinValue, float MaxValue)
 	{
 		// Convert to HSV
 		InOutColor = InOutColor.LinearRGBToHSV();
@@ -118,10 +96,12 @@ namespace RemoteControlFunctionLibrary
 		// Apply the delta to the position, then convert back to hue and saturation
 		Position += FVector2D(DeltaValue.Position.X, DeltaValue.Position.Y);
 		InOutColor.R = FMath::Fmod(FMath::RadiansToDegrees(FMath::Atan2(Position.Y, Position.X)) + 360.0, 360.0);
-		InOutColor.G = FMath::Clamp(Position.Length(), 0.0, MaxValue);
+		InOutColor.G = FMath::Clamp(Position.Length(), MinValue, MaxValue);
 
-		// Apply the value change directly
-		InOutColor.B = FMath::Clamp(InOutColor.B + DeltaValue.Value, 0.0, MaxValue);
+		// We're operating on a color transformed to a (0, 1) range from its full range, so scale the value delta by the full range
+		check(MaxValue > MinValue);
+		float ScaledDeltaValue = DeltaValue.Value / (MaxValue - MinValue);
+		InOutColor.B = FMath::Clamp(InOutColor.B + ScaledDeltaValue, MinValue, MaxValue);
 
 		InOutColor = InOutColor.HSVToLinearRGB();
 	}
@@ -193,10 +173,10 @@ bool URemoteControlFunctionLibrary::ApplyColorWheelDelta(UObject* TargetObject,
 		FLinearColor Color;
 		ColorProperty->GetValue_InContainer(Container, &Color);
 
-		RemoteControlFunctionLibrary::ApplyWheelColorBaseDelta(Color, DeltaValue, ReferenceColor, 1.0);
+		RemoteControlFunctionLibrary::ApplyWheelColorBaseDelta(Color, DeltaValue, ReferenceColor, 0.f, 1.f);
 
 		// Apply the alpha change directly
-		Color.A = FMath::Clamp(Color.A + DeltaValue.Alpha, 0.0, 1.0);
+		Color.A = FMath::Clamp(Color.A + DeltaValue.Alpha, 0.f, 1.f);
 
 		RemoteControlFunctionLibrary::SetPropertyAndFireEvents(TargetObject, Property, Container, &Color, bIsInteractive);
 
@@ -206,8 +186,13 @@ bool URemoteControlFunctionLibrary::ApplyColorWheelDelta(UObject* TargetObject,
 	return false;
 }
 
-bool URemoteControlFunctionLibrary::ApplyColorGradingWheelDelta(UObject* TargetObject, const FString& PropertyName, const FColorGradingWheelColor& DeltaValue, const FColorGradingWheelColor& ReferenceColor, bool bIsInteractive)
+bool URemoteControlFunctionLibrary::ApplyColorGradingWheelDelta(UObject* TargetObject, const FString& PropertyName, const FColorGradingWheelColor& DeltaValue, const FColorGradingWheelColor& ReferenceColor, bool bIsInteractive, float MinValue, float MaxValue)
 {
+	if (MinValue >= MaxValue)
+	{
+		return false;
+	}
+
 	FProperty* Property = nullptr;
 	void* Container = nullptr;
 	if (!RemoteControlFunctionLibrary::GetNestedPropertyAndContainerFromPath(TargetObject, PropertyName, Property, Container))
@@ -227,12 +212,16 @@ bool URemoteControlFunctionLibrary::ApplyColorGradingWheelDelta(UObject* TargetO
 		FVector4 VectorValue;
 		ColorProperty->GetValue_InContainer(Container, &VectorValue);
 
+		const float OldLuminance = VectorValue.W;
+
+		RemoteControlFunctionLibrary::TransformColorGradingRangeToLinearColorRange(VectorValue, MinValue, MaxValue);
+
 		FLinearColor Color(VectorValue.X, VectorValue.Y, VectorValue.Z);
-		RemoteControlFunctionLibrary::ApplyWheelColorBaseDelta(Color, DeltaValue, ReferenceColor, 2.0);
+		RemoteControlFunctionLibrary::ApplyWheelColorBaseDelta(Color, DeltaValue, ReferenceColor, MinValue, MaxValue);
+		VectorValue = FVector4(Color);
 
-		// Apply the luminance change directly
-		const double NewLuminance = VectorValue.W + DeltaValue.Luminance;
-		VectorValue = FVector4(Color, NewLuminance);
+		RemoteControlFunctionLibrary::TransformLinearColorRangeToColorGradingRange(VectorValue, MinValue, MaxValue);
+		VectorValue.W = OldLuminance + DeltaValue.Luminance;
 
 		RemoteControlFunctionLibrary::SetPropertyAndFireEvents(TargetObject, Property, Container, &VectorValue, bIsInteractive);
 
diff --git a/Engine/Plugins/VirtualProduction/RemoteControl/Source/RemoteControl/Public/RemoteControlFunctionLibrary.h b/Engine/Plugins/VirtualProduction/RemoteControl/Source/RemoteControl/Public/RemoteControlFunctionLibrary.h
--- a/Engine/Plugins/VirtualProduction/RemoteControl/Source/RemoteControl/Public/RemoteControlFunctionLibrary.h
+++ b/Engine/Plugins/VirtualProduction/RemoteControl/Source/RemoteControl/Public/RemoteControlFunctionLibrary.h
@@ -36,8 +36,8 @@ struct FColorWheelColorBase
 	GENERATED_BODY()
 
 	/**
-		* The position on the unit circle. Magnitude should be in range [0, 1].
-		*/
+	* The position on the unit circle. Magnitude should be in range [0, 1].
+	*/
 	UPROPERTY()
 	FVector2D Position = FVector2D::ZeroVector;
 
@@ -71,7 +71,7 @@ struct FColorGradingWheelColor : public FColorWheelColorBase
 	 * The color's luminance component, which is separate from the HSV color itself.
 	 */
 	UPROPERTY()
-	double Luminance = 1.0;
+	double Luminance = 0.0;
 };
 
 UCLASS()
@@ -134,5 +134,5 @@ public:
 	 * @return Whether the operation was successful.
 	 */
 	UFUNCTION(BlueprintCallable, Category = "Remote Control Color")
-	static bool ApplyColorGradingWheelDelta(UObject* TargetObject, const FString& PropertyName, const FColorGradingWheelColor& DeltaValue, const FColorGradingWheelColor& ReferenceColor, bool bIsInteractive);
+	static bool ApplyColorGradingWheelDelta(UObject* TargetObject, const FString& PropertyName, const FColorGradingWheelColor& DeltaValue, const FColorGradingWheelColor& ReferenceColor, bool bIsInteractive, float MinValue = 0.f, float MaxValue = 2.f);
 };
\ No newline at end of file
