Fix desync between preset properties when entering/exiting multi user

diff --git a/Engine/Plugins/VirtualProduction/RemoteControl/Source/RemoteControl/Private/RemoteControlPreset.cpp b/Engine/Plugins/VirtualProduction/RemoteControl/Source/RemoteControl/Private/RemoteControlPreset.cpp
--- a/Engine/Plugins/VirtualProduction/RemoteControl/Source/RemoteControl/Private/RemoteControlPreset.cpp
+++ b/Engine/Plugins/VirtualProduction/RemoteControl/Source/RemoteControl/Private/RemoteControlPreset.cpp
@@ -812,13 +812,14 @@ URemoteControlBinding* URemoteControlPreset::FindOrAddBinding(const TSoftObjectP
 	{
 		if (ResolvedObject->GetTypedOuter<ULevel>())
         {
-        	NewBinding = NewObject<URemoteControlLevelDependantBinding>(this);	
+        	NewBinding = NewObject<URemoteControlLevelDependantBinding>(this, NAME_None, RF_Transactional);	
         }
         else
         {
-        	NewBinding = NewObject<URemoteControlLevelIndependantBinding>(this);
+        	NewBinding = NewObject<URemoteControlLevelIndependantBinding>(this, NAME_None, RF_Transactional);
         }
-		
+
+		NewBinding->Modify();
 		NewBinding->SetBoundObject(Object);
 	}
 	else
@@ -830,13 +831,14 @@ URemoteControlBinding* URemoteControlPreset::FindOrAddBinding(const TSoftObjectP
 		if (PersistentLevelIndex != INDEX_NONE)
 		{
 			TSoftObjectPtr<ULevel> Level = TSoftObjectPtr<ULevel>{ FSoftObjectPath{ Path.Left(PersistentLevelIndex + PersistentLevelText.Len() - 1) } };
-			URemoteControlLevelDependantBinding* LevelDependantBinding = NewObject<URemoteControlLevelDependantBinding>(this);
+			URemoteControlLevelDependantBinding* LevelDependantBinding = NewObject<URemoteControlLevelDependantBinding>(this, NAME_None, RF_Transactional);
 			LevelDependantBinding->SetBoundObject(Level, Object);
 			NewBinding = LevelDependantBinding;
 		}
 		else
 		{
-			NewBinding = NewObject<URemoteControlLevelIndependantBinding>(this);
+			NewBinding = NewObject<URemoteControlLevelIndependantBinding>(this, NAME_None, RF_Transactional);
+			NewBinding->Modify();
 			NewBinding->SetBoundObject(Object);
 		}
 		
@@ -1555,6 +1557,23 @@ void URemoteControlPreset::OnObjectPropertyChanged(UObject* Object, struct FProp
 	// Objects modified should have run through the preobjectmodified. If interesting, they will be cached
 	TRACE_CPUPROFILER_EVENT_SCOPE(URemoteControlPreset::OnObjectPropertyChanged);
 
+	// Handle enter / exit Multi-User session
+	if (Object == this->GetPackage() && Event.ChangeType == EPropertyChangeType::Redirected)
+	{
+		for (const TSharedPtr<FRemoteControlEntity>& Entity : Registry->GetExposedEntities<FRemoteControlEntity>())
+		{
+			PerFrameUpdatedEntities.Add(Entity->GetId());
+		}
+		return;
+	}
+
+	if (Object == this && Event.Property && Event.Property->GetFName() == GET_MEMBER_NAME_CHECKED(URemoteControlPreset, Metadata))
+	{
+		// Needed because re-joining a multi user session will modify this without going through the SetMetadata function and necessitates a website refresh.
+		OnMetadataModified().Broadcast(this);
+		return;
+	}
+
 	if (Object && Object->GetClass()->GetName() == TEXT("DisplayClusterConfigurationData"))
 	{
 		TRACE_CPUPROFILER_EVENT_SCOPE(URemoteControlPreset::OnObjectPropertyChanged::HandleDisplayClusterConfigChange);
@@ -2013,6 +2032,15 @@ void URemoteControlPreset::OnPackageReloaded(EPackageReloadPhase Phase, FPackage
 			RepointedPreset->OnMetadataModifiedDelegate = OnMetadataModifiedDelegate;
 			RepointedPreset->OnActorPropertyModifiedDelegate = OnActorPropertyModifiedDelegate;
 			RepointedPreset->OnPresetLayoutModifiedDelegate = OnPresetLayoutModifiedDelegate;
+
+			TSet<FGuid> PropertiesToUpdate;
+			for (const TSharedPtr<FRemoteControlEntity>& Entity : RepointedPreset->Registry->GetExposedEntities<FRemoteControlEntity>())
+			{
+				PropertiesToUpdate.Add(Entity->GetId());
+			}
+
+			RepointedPreset->OnMetadataModifiedDelegate.Broadcast(RepointedPreset);
+			RepointedPreset->OnEntitiesUpdatedDelegate.Broadcast(RepointedPreset, PropertiesToUpdate);
 		}
 	}
 }
diff --git a/Engine/Plugins/VirtualProduction/RemoteControl/Source/RemoteControlMultiUser/Private/RemoteControlMultiUserModule.cpp b/Engine/Plugins/VirtualProduction/RemoteControl/Source/RemoteControlMultiUser/Private/RemoteControlMultiUserModule.cpp
--- a/Engine/Plugins/VirtualProduction/RemoteControl/Source/RemoteControlMultiUser/Private/RemoteControlMultiUserModule.cpp
+++ b/Engine/Plugins/VirtualProduction/RemoteControl/Source/RemoteControlMultiUser/Private/RemoteControlMultiUserModule.cpp
@@ -4,6 +4,7 @@
 
 #include "IConcertClientTransactionBridge.h"
 #include "IConcertSyncClientModule.h"
+#include "RemoteControlBinding.h"
 #include "RemoteControlPreset.h"
 
 
@@ -13,7 +14,9 @@ namespace RemoteControlMultiUserUtils
 	{
 		//Always allow RemoteControlPresets and RemoteControlExposeRegistry
 		static const FName ExposeRegistryName = "RemoteControlExposeRegistry";
-		if(Cast<URemoteControlPreset>(ObjectToFilter) != nullptr || (ObjectToFilter && ObjectToFilter->GetClass()->GetFName() == ExposeRegistryName))
+		if (Cast<URemoteControlPreset>(ObjectToFilter) != nullptr
+			|| (ObjectToFilter && ObjectToFilter->GetClass()->GetFName() == ExposeRegistryName)
+			|| (ObjectToFilter && ObjectToFilter->IsA<URemoteControlBinding>()))
 		{
 			return ETransactionFilterResult::IncludeObject;
 		}
diff --git a/Engine/Plugins/VirtualProduction/RemoteControl/Source/WebRemoteControl/Private/WebRemoteControl.cpp b/Engine/Plugins/VirtualProduction/RemoteControl/Source/WebRemoteControl/Private/WebRemoteControl.cpp
--- a/Engine/Plugins/VirtualProduction/RemoteControl/Source/WebRemoteControl/Private/WebRemoteControl.cpp
+++ b/Engine/Plugins/VirtualProduction/RemoteControl/Source/WebRemoteControl/Private/WebRemoteControl.cpp
@@ -1856,13 +1856,18 @@ bool FWebRemoteControlModule::HandleMetadataFieldOperationsRoute(const FHttpServ
 			{
 				return true;
 			}
-			
+#if WITH_EDITOR
+			FScopedTransaction Transaction(LOCTEXT("ModifyPresetMetadata", "Modify preset metadata"));
+#endif
 			Preset->Modify();
 			FString& MetadataValue = Preset->Metadata.FindOrAdd(MoveTemp(MetadataField));
 			MetadataValue = MoveTemp(SetMetadataRequest.Value);
 		}
 		else if (Request.Verb == EHttpServerRequestVerbs::VERB_DELETE)
 		{
+#if WITH_EDITOR
+			FScopedTransaction Transaction(LOCTEXT("DeletePresetMetadata", "Delete preset metadata entry"));
+#endif
 			Preset->Modify();
 			Preset->Metadata.Remove(MoveTemp(MetadataField));
 		}
