Fix remote control crash when leaving multi user session

diff --git a/Engine/Plugins/VirtualProduction/RemoteControl/Source/RemoteControlUI/Private/RemoteControlUIModule.cpp b/Engine/Plugins/VirtualProduction/RemoteControl/Source/RemoteControlUI/Private/RemoteControlUIModule.cpp
--- a/Engine/Plugins/VirtualProduction/RemoteControl/Source/RemoteControlUI/Private/RemoteControlUIModule.cpp
+++ b/Engine/Plugins/VirtualProduction/RemoteControl/Source/RemoteControlUI/Private/RemoteControlUIModule.cpp
@@ -53,11 +53,7 @@ static const FName DetailsTabIdentifiers_LevelEditor[] = {
 static const int32 DetailsTabIdentifiers_LevelEditor_Count = 4;
 
 FRCExposesPropertyArgs::FRCExposesPropertyArgs()
-	: PropertyHandle(nullptr)
-	, OwnerObject(nullptr)
-	, PropertyPath(TEXT(""))
-	, Property(nullptr)
-	, Id(FGuid::NewGuid())
+	: Id(FGuid::NewGuid())
 {
 }
 
@@ -99,7 +95,7 @@ FRCExposesPropertyArgs::EType FRCExposesPropertyArgs::GetType() const
 	{
 		return EType::E_Handle;
 	}
-	if (OwnerObject && !PropertyPath.IsEmpty() && Property)
+	if (OwnerObject.IsValid() && !PropertyPath.IsEmpty() && Property.IsValid())
 	{
 		return EType::E_OwnerObject;
 	}
@@ -123,7 +119,7 @@ FProperty* FRCExposesPropertyArgs::GetProperty() const
 	}
 	else
 	{
-		return Property;
+		return Property.Get();
 	}
 
 	return nullptr;
@@ -178,7 +174,7 @@ namespace RemoteControlUIModule
 		}
 
 		FRCFieldPathInfo RCFieldPathInfo(InPropertyArgs.PropertyPath);
-		RCFieldPathInfo.Resolve(InPropertyArgs.OwnerObject);
+		RCFieldPathInfo.Resolve(InPropertyArgs.OwnerObject.Get());
 		FRCFieldResolvedData ResolvedData = RCFieldPathInfo.GetResolvedData();
 		return CastField<FObjectProperty>(InPropertyArgs.GetProperty());
 	}
@@ -553,7 +549,7 @@ TSharedPtr<SRemoteControlPanel> FRemoteControlUIModule::GetPanelForProperty(cons
 	}
 	else if (ExtensionArgsType == FRCExposesPropertyArgs::EType::E_OwnerObject)
 	{
-		return GetPanelForObject(InPropertyArgs.OwnerObject);
+		return GetPanelForObject(InPropertyArgs.OwnerObject.Get());
 	}
 
 	return GetPanelForObject(nullptr);
@@ -736,7 +732,7 @@ bool FRemoteControlUIModule::ShouldDisplayExposeIcon(const FRCExposesPropertyArg
 	}
 	else if (ExtensionArgsType == FRCExposesPropertyArgs::EType::E_OwnerObject)
 	{
-		if (!IsAllowedOwnerObjects({ InPropertyArgs.OwnerObject }))
+		if (!IsAllowedOwnerObjects({ InPropertyArgs.OwnerObject.Get()}))
 		{
 			return false;
 		}
@@ -792,7 +788,7 @@ bool FRemoteControlUIModule::ShouldSkipOwnPanelProperty(const FRCExposesProperty
 	}
 	else if (ArgsType == FRCExposesPropertyArgs::EType::E_OwnerObject)
 	{
-		return CheckOwnProperty(InPropertyArgs.Property);
+		return CheckOwnProperty(InPropertyArgs.Property.Get());
 	}
 
 	return false;
@@ -940,9 +936,9 @@ void FRemoteControlUIModule::TryOverridingMaterials(const FRCExposesPropertyArgs
 	else if (ExtensionArgsType == FRCExposesPropertyArgs::EType::E_OwnerObject)
 	{
 		FRCFieldPathInfo RCFieldPathInfo(InPropertyArgs.PropertyPath);
-		RCFieldPathInfo.Resolve(InPropertyArgs.OwnerObject);
+		RCFieldPathInfo.Resolve(InPropertyArgs.OwnerObject.Get());
 		const FRCFieldResolvedData ResolvedData = RCFieldPathInfo.GetResolvedData();
-		FObjectProperty* ObjectProperty = CastField<FObjectProperty>(InPropertyArgs.Property);
+		FObjectProperty* ObjectProperty = CastField<FObjectProperty>(InPropertyArgs.Property.Get());
 
 		if (!ObjectProperty)
 		{
@@ -954,9 +950,9 @@ void FRemoteControlUIModule::TryOverridingMaterials(const FRCExposesPropertyArgs
 		OriginalMaterial = OriginalMaterial ? OriginalMaterial : UMaterial::GetDefaultMaterial(MD_Surface);
 		if (OriginalMaterial)
 		{
-			if (UMeshComponent* MeshComponentToBeModified = GetSelectedMeshComponentToBeModified(InPropertyArgs.OwnerObject, OriginalMaterial))
+			if (UMeshComponent* MeshComponentToBeModified = GetSelectedMeshComponentToBeModified(InPropertyArgs.OwnerObject.Get(), OriginalMaterial))
 			{
-				const int32 NumStaticMaterials = RemoteControlUIModule::GetNumStaticMaterials(InPropertyArgs.OwnerObject);
+				const int32 NumStaticMaterials = RemoteControlUIModule::GetNumStaticMaterials(InPropertyArgs.OwnerObject.Get());
 				if (NumStaticMaterials > 0)
 				{
 					FScopedTransaction Transaction(LOCTEXT("OverrideMaterial", "Override Material"));
diff --git a/Engine/Plugins/VirtualProduction/RemoteControl/Source/RemoteControlUI/Private/UI/SRemoteControlPanel.cpp b/Engine/Plugins/VirtualProduction/RemoteControl/Source/RemoteControlUI/Private/UI/SRemoteControlPanel.cpp
--- a/Engine/Plugins/VirtualProduction/RemoteControl/Source/RemoteControlUI/Private/UI/SRemoteControlPanel.cpp
+++ b/Engine/Plugins/VirtualProduction/RemoteControl/Source/RemoteControlUI/Private/UI/SRemoteControlPanel.cpp
@@ -842,7 +842,7 @@ bool SRemoteControlPanel::IsExposed(const FRCExposesPropertyArgs& InPropertyArgs
 	else if (ExtensionArgsType == FRCExposesPropertyArgs::EType::E_OwnerObject)
 	{
 		constexpr bool bIsCheckIsBoundByFullPath = false;
-		return CheckCachedExposedArgs({ InPropertyArgs.OwnerObject }, InPropertyArgs.PropertyPath, bIsCheckIsBoundByFullPath);
+		return CheckCachedExposedArgs({ InPropertyArgs.OwnerObject.Get()}, InPropertyArgs.PropertyPath, bIsCheckIsBoundByFullPath);
 	}
 
 	// It never should hit this point
@@ -944,7 +944,7 @@ void SRemoteControlPanel::ToggleProperty(const FRCExposesPropertyArgs& InPropert
 		Preset->Modify();
 
 		constexpr bool bCleanDuplicates = true; // GeneratePathToProperty duplicates container name (Array.Array[1], Set.Set[1], etc...)
-		ExposeProperty(InPropertyArgs.OwnerObject, FRCFieldPathInfo{ InPropertyArgs.PropertyPath, bCleanDuplicates });
+		ExposeProperty(InPropertyArgs.OwnerObject.Get(), FRCFieldPathInfo{InPropertyArgs.PropertyPath, bCleanDuplicates});
 
 		PostExpose();
 	}
@@ -1598,7 +1598,7 @@ void SRemoteControlPanel::Unexpose(const FRCExposesPropertyArgs& InPropertyArgs)
 	else if (ExtensionArgsType == FRCExposesPropertyArgs::EType::E_OwnerObject)
 	{
 		constexpr bool bUsingDuplicatesInPath = false;
-		CheckAndUnexpose({ InPropertyArgs.OwnerObject }, InPropertyArgs.PropertyPath, bUsingDuplicatesInPath);
+		CheckAndUnexpose({ InPropertyArgs.OwnerObject.Get()}, InPropertyArgs.PropertyPath, bUsingDuplicatesInPath);
 	}
 }
 
diff --git a/Engine/Plugins/VirtualProduction/RemoteControl/Source/RemoteControlUI/Public/IRemoteControlUIModule.h b/Engine/Plugins/VirtualProduction/RemoteControl/Source/RemoteControlUI/Public/IRemoteControlUIModule.h
--- a/Engine/Plugins/VirtualProduction/RemoteControl/Source/RemoteControlUI/Public/IRemoteControlUIModule.h
+++ b/Engine/Plugins/VirtualProduction/RemoteControl/Source/RemoteControlUI/Public/IRemoteControlUIModule.h
@@ -6,6 +6,8 @@
 #include "Modules/ModuleManager.h"
 #include "Modules/ModuleInterface.h"
 #include "Templates/SharedPointer.h"
+#include "UObject/WeakFieldPtr.h"
+#include "UObject/WeakObjectPtrTemplates.h"
 #include "Widgets/Layout/SSplitter.h"
 
 class FDelegateHandle;
@@ -92,13 +94,13 @@ public:
 	TSharedPtr<IPropertyHandle> PropertyHandle;
 
 	/** Owner object for the property extension */
-	UObject* OwnerObject;
+	TWeakObjectPtr<UObject> OwnerObject;
 
 	/** Path of the exposed property */
 	FString PropertyPath;
 
 	/** Exposed property */
-	FProperty* Property;
+	TWeakFieldPtr<FProperty> Property;
 
 private:
 	/** Unique generated ID of the struct */
