Fix regression in remote control where class filters don't work anymore

diff --git a/Engine/Plugins/VirtualProduction/RemoteControl/Source/WebRemoteControl/Private/RemoteControlModels.h b/Engine/Plugins/VirtualProduction/RemoteControl/Source/WebRemoteControl/Private/RemoteControlModels.h
--- a/Engine/Plugins/VirtualProduction/RemoteControl/Source/WebRemoteControl/Private/RemoteControlModels.h
+++ b/Engine/Plugins/VirtualProduction/RemoteControl/Source/WebRemoteControl/Private/RemoteControlModels.h
@@ -124,6 +124,7 @@ struct FRCPropertyDescription
 #endif
 
 		const FProperty* ValueProperty = Property;
+		
 		if (const FArrayProperty* ArrayProperty = CastField<FArrayProperty>(Property))
 		{
 			ContainerType = Property->GetCPPType();
@@ -148,6 +149,15 @@ struct FRCPropertyDescription
 		//Write the type name
 		Type = ValueProperty->GetCPPType();
 
+		if (const FObjectProperty* ObjectProperty = CastField<FObjectPtrProperty>(ValueProperty))
+		{
+			if (UClass* Class = ObjectProperty->PropertyClass)
+			{
+				Type = FString::Printf(TEXT("%s%s"), Class->GetPrefixCPP(), *Class->GetName());
+				TypePath = *Class->GetPathName();
+			}
+		}
+
 #if WITH_EDITOR
 		Metadata = Property->GetMetaDataMap() ? RemoteControlModels::SanitizeMetadata(*Property->GetMetaDataMap()) : TMap<FName, FString>();
 		Description = Property->GetMetaData("ToolTip");
@@ -208,6 +218,10 @@ struct FRCPropertyDescription
 	UPROPERTY()
 	FString Type;
 
+	/** Type of the property value (If an array, this will be the content of the array) */
+	UPROPERTY()
+	FName TypePath;
+
 	/** Type of the container (TMap, TArray, CArray, TSet) or empty if none */
 	UPROPERTY()
 	FString ContainerType;
@@ -696,7 +710,7 @@ struct FRCAssetDescription
 	FRCAssetDescription() = default;
 	FRCAssetDescription(const FAssetData& InAsset)
 		: Name(InAsset.AssetName)
-		, Class(InAsset.AssetClassPath)
+		, Class(*InAsset.AssetClassPath.ToString())
 		, Path(*InAsset.GetObjectPathString())
 	{
 		Metadata = RemoteControlModels::SanitizeAssetMetadata(InAsset.TagsAndValues.CopyMap());
@@ -706,7 +720,7 @@ struct FRCAssetDescription
 	FName Name;
 
 	UPROPERTY()
-	FTopLevelAssetPath Class;
+	FName Class;
 
 	UPROPERTY()
 	FName Path;
@@ -744,10 +758,15 @@ struct FRCAssetFilter
 
 	FARFilter ToARFilter() const
 	{
+		auto ToAssetPath = [](FName InClassName) { return FTopLevelAssetPath(InClassName.ToString()); };
+
 		FARFilter Filter;
 		Filter.PackageNames = PackageNames;
-		Filter.ClassPaths = ClassNames;
-		Filter.RecursiveClassPathsExclusionSet = RecursiveClassesExclusionSet;
+
+		Filter.ClassPaths.Reserve(ClassNames.Num());
+		Algo::Transform(ClassNames, Filter.ClassPaths, ToAssetPath);
+
+		Algo::Transform(RecursiveClassesExclusionSet, Filter.RecursiveClassPathsExclusionSet, ToAssetPath);
 		Filter.bRecursiveClasses = RecursiveClasses;
 		Filter.PackagePaths = PackagePaths;
 
@@ -778,15 +797,15 @@ struct FRCAssetFilter
 
 	/** The filter component for class names. Instances of the specified classes, but not subclasses (by default), will be included. Derived classes will be included only if bRecursiveClasses is true. */
 	UPROPERTY()
-	TArray<FTopLevelAssetPath> ClassNames;
+	TArray<FName> ClassNames;
 
 	/** Only if bRecursiveClasses is true, the results will exclude classes (and subclasses) in this list */
 	UPROPERTY()
-	TSet<FTopLevelAssetPath> RecursiveClassesExclusionSet;
+	TSet<FName> RecursiveClassesExclusionSet;
 
 	/** Only if EnableBlueprintNativeClassFiltering is true, resulting asset will be filtered for dependants of classes in this list. */
 	UPROPERTY()
-	TArray<FTopLevelAssetPath> NativeParentClasses;
+	TArray<FName> NativeParentClasses;
 
 	/** If true, subclasses of ClassNames will also be included and RecursiveClassesExclusionSet will be excluded. */
 	UPROPERTY()
diff --git a/Engine/Plugins/VirtualProduction/RemoteControl/Source/WebRemoteControl/Private/WebRemoteControl.cpp b/Engine/Plugins/VirtualProduction/RemoteControl/Source/WebRemoteControl/Private/WebRemoteControl.cpp
--- a/Engine/Plugins/VirtualProduction/RemoteControl/Source/WebRemoteControl/Private/WebRemoteControl.cpp
+++ b/Engine/Plugins/VirtualProduction/RemoteControl/Source/WebRemoteControl/Private/WebRemoteControl.cpp
@@ -235,9 +235,9 @@ namespace WebRemoteControl
 		{
 			TSet<UClass*> NativeClasses;
 			NativeClasses.Reserve(SearchAssetRequest.Filter.NativeParentClasses.Num());
-			for (FTopLevelAssetPath NativeClass : SearchAssetRequest.Filter.NativeParentClasses)
+			for (FName NativeClassName : SearchAssetRequest.Filter.NativeParentClasses)
 			{
-				if (UClass* Class = FindObject<UClass>(NativeClass))
+				if (UClass* Class = FindObject<UClass>(FTopLevelAssetPath{ NativeClassName.ToString()}))
 				{
 					NativeClasses.Add(Class);
 				}
@@ -259,6 +259,10 @@ namespace WebRemoteControl
 						}
 					}
 				}
+				else
+				{
+					It.RemoveCurrent();
+				}
 			}
 		}
 	}
diff --git a/Engine/Plugins/VirtualProduction/RemoteControlWebInterface/WebApp/Client/src/components/Screens/ColorCorrection/ColorCorrection.tsx b/Engine/Plugins/VirtualProduction/RemoteControlWebInterface/WebApp/Client/src/components/Screens/ColorCorrection/ColorCorrection.tsx
--- a/Engine/Plugins/VirtualProduction/RemoteControlWebInterface/WebApp/Client/src/components/Screens/ColorCorrection/ColorCorrection.tsx
+++ b/Engine/Plugins/VirtualProduction/RemoteControlWebInterface/WebApp/Client/src/components/Screens/ColorCorrection/ColorCorrection.tsx
@@ -792,6 +792,7 @@ export class ColorCorrection extends React.Component<Props, State> {
                     <AssetWidget browse
                                  reset={false}
                                  type="Texture"
+                                 typePath="/Script/Engine.Texture"
                                  label="Texture"
                                  onChange={v => this.onPropertyChange('Texture', v)}
                                  value={_.get(value, 'Texture') ?? ''} />
diff --git a/Engine/Plugins/VirtualProduction/RemoteControlWebInterface/WebApp/Client/src/components/Screens/Playlist.tsx b/Engine/Plugins/VirtualProduction/RemoteControlWebInterface/WebApp/Client/src/components/Screens/Playlist.tsx
--- a/Engine/Plugins/VirtualProduction/RemoteControlWebInterface/WebApp/Client/src/components/Screens/Playlist.tsx
+++ b/Engine/Plugins/VirtualProduction/RemoteControlWebInterface/WebApp/Client/src/components/Screens/Playlist.tsx
@@ -68,7 +68,7 @@ export class Playlist extends React.Component<Props, State> {
 
     this.player = await _api.proxy.function("Object'/Script/SequencerPlaylists.Default__SequencerPlaylistPlayer'", 'GetDefaultPlayer', { WorldContextObject });
 
-    const sequences = await _api.assets.search('', ['LevelSequence'], '/Game', 200);
+    const sequences = await _api.assets.search('', ['/Script/LevelSequence.LevelSequence'], '/Game', 200);
     this.setState({ sequences });
 
     const currentPlaylist = await _api.proxy.function(this.player, 'GetPlaylist');
@@ -76,7 +76,7 @@ export class Playlist extends React.Component<Props, State> {
       const folder = path.dirname(currentPlaylist);
       const [filename] = path.basename(currentPlaylist).split('.');
 
-      const playlistAsset = await _api.assets.search(filename, ['SequencerPlaylist'], folder, 200);
+      const playlistAsset = await _api.assets.search(filename, ['/Script/SequencerPlaylists.SequencerPlaylist'], folder, 200);
       await this.onLoadPlaylist(playlistAsset[0]);
     }
   }
@@ -119,7 +119,7 @@ export class Playlist extends React.Component<Props, State> {
   }
 
   onPlaylistBrowse = async () => {
-    const playlist = await SearchModal.open({ placeholder: 'Search Playlist', types: ['SequencerPlaylist'], prefix: '/Game', count: 200 });
+    const playlist = await SearchModal.open({ placeholder: 'Search Playlist', types: ['/Script/SequencerPlaylists.SequencerPlaylist'], prefix: '/Game', count: 200 });
     if (!playlist)
       return;
 
diff --git a/Engine/Plugins/VirtualProduction/RemoteControlWebInterface/WebApp/Client/src/components/Screens/Sequencer.tsx b/Engine/Plugins/VirtualProduction/RemoteControlWebInterface/WebApp/Client/src/components/Screens/Sequencer.tsx
--- a/Engine/Plugins/VirtualProduction/RemoteControlWebInterface/WebApp/Client/src/components/Screens/Sequencer.tsx
+++ b/Engine/Plugins/VirtualProduction/RemoteControlWebInterface/WebApp/Client/src/components/Screens/Sequencer.tsx
@@ -32,7 +32,7 @@ export class Sequencer extends React.Component<Props, State> {
 
   fetchAssets = async () => {
     const { search } = this.state;
-    const assets = await _api.assets.search(search, ['LevelSequence'], '/Game', 200);
+    const assets = await _api.assets.search(search, ['/Script/LevelSequence.LevelSequence'], '/Game', 200);
     this.setState({ assets });
   }
 
diff --git a/Engine/Plugins/VirtualProduction/RemoteControlWebInterface/WebApp/Client/src/components/Screens/Snapshot.tsx b/Engine/Plugins/VirtualProduction/RemoteControlWebInterface/WebApp/Client/src/components/Screens/Snapshot.tsx
--- a/Engine/Plugins/VirtualProduction/RemoteControlWebInterface/WebApp/Client/src/components/Screens/Snapshot.tsx
+++ b/Engine/Plugins/VirtualProduction/RemoteControlWebInterface/WebApp/Client/src/components/Screens/Snapshot.tsx
@@ -48,18 +48,18 @@ export class Snapshot extends React.Component<Props, State> {
       map = await _api.proxy.function(editorSubsystem, 'GetEditorWorld', {});
     }
 
-    const nativeFilters = await _api.assets.search('', ['LevelSnapshotFilter'], '/Game', {}, 50);
+    const nativeFilters = await _api.assets.search('', ['/Script/LevelSnapshotFilters.LevelSnapshotFilter'], '/Game', {}, 50);
 
     const filterArgs = {
       EnableBlueprintNativeClassFiltering: true,
-      NativeParentClasses: ['LevelSnapshotFilter'],
+      NativeParentClasses: ['/Script/LevelSnapshotFilters.LevelSnapshotBlueprintFilter'],
     };
-    const bpFilters = _.filter(await _api.assets.search('', ['Blueprint'], '/Game', filterArgs, 50), f => !!f.Metadata['GeneratedClass']);
+    const bpFilters = _.filter(await _api.assets.search('', ['/Script/Engine.Blueprint'], '/Game', filterArgs, 50), f => !!f.Metadata['GeneratedClass']);
 
     const filters = [...bpFilters, ...nativeFilters];
 
     const { search, sortBy, sortAsc } = this.state;
-    const assets = await _api.assets.search(search, ['LevelSnapshot'], '/Game', {}, 200);
+    const assets = await _api.assets.search(search, ['/Script/LevelSnapshots.LevelSnapshot'], '/Game', {}, 200);
     let snapshots = _.filter(assets, asset => !map || asset.Metadata?.MapPath === map);
     
     snapshots = this.sortAssets(snapshots, sortBy, sortAsc);
diff --git a/Engine/Plugins/VirtualProduction/RemoteControlWebInterface/WebApp/Client/src/components/Stack.tsx b/Engine/Plugins/VirtualProduction/RemoteControlWebInterface/WebApp/Client/src/components/Stack.tsx
--- a/Engine/Plugins/VirtualProduction/RemoteControlWebInterface/WebApp/Client/src/components/Stack.tsx
+++ b/Engine/Plugins/VirtualProduction/RemoteControlWebInterface/WebApp/Client/src/components/Stack.tsx
@@ -764,6 +764,7 @@ export class Stack extends React.Component<Props & PropsFromState, State> {
         <Widgets.AssetWidget label={this.getLabel(widget)}
                              value={this.getValue(widget)}
                              type={exposed?.Type ?? 'Object'}
+                             typePath={exposed.TypePath}
                              onChange={value => this.onPropertyValueChange(widget, value)}
                              browse />
       </DraggableWidget>
diff --git a/Engine/Plugins/VirtualProduction/RemoteControlWebInterface/WebApp/Client/src/shared/index.ts b/Engine/Plugins/VirtualProduction/RemoteControlWebInterface/WebApp/Client/src/shared/index.ts
--- a/Engine/Plugins/VirtualProduction/RemoteControlWebInterface/WebApp/Client/src/shared/index.ts
+++ b/Engine/Plugins/VirtualProduction/RemoteControlWebInterface/WebApp/Client/src/shared/index.ts
@@ -87,6 +87,7 @@ export interface IProperty {
   Name: string;
   Description: string;
   Type: PropertyType;
+  TypePath: string;
   Metadata: { [key: string]: string };
 }
 
@@ -100,12 +101,14 @@ export interface IExposedProperty {
 
   //Added
   Type: PropertyType;
+  TypePath: string;
 }
 
 export interface IController {
   ID: string;
   Name: string;
   Type: PropertyType;
+  TypePath: string;
   Path: string;
 
   DisplayName: string;
diff --git a/Engine/Plugins/VirtualProduction/RemoteControlWebInterface/WebApp/Server/src/UnrealEngine.ts b/Engine/Plugins/VirtualProduction/RemoteControlWebInterface/WebApp/Server/src/UnrealEngine.ts
--- a/Engine/Plugins/VirtualProduction/RemoteControlWebInterface/WebApp/Server/src/UnrealEngine.ts
+++ b/Engine/Plugins/VirtualProduction/RemoteControlWebInterface/WebApp/Server/src/UnrealEngine.ts
@@ -503,6 +503,7 @@ export namespace UnrealEngine {
     for (const Group of preset.Groups) {
       for (const Property of Group.ExposedProperties) {
         Property.Type = Property.UnderlyingProperty.Type;
+        Property.TypePath = Property.UnderlyingProperty.TypePath;
         preset.Exposed[Property.ID] = Property;
       }
 
@@ -531,12 +532,13 @@ export namespace UnrealEngine {
           Metadata: Controller.Metadata,
           Description: '',
           Type: Controller.Type,
+          TypePath: Controller.TypePath
         };
 
         preset.Exposed[Controller.ID] = Controller;
         Controllers.ExposedProperties.push(Controller);
       }
-    }
+    } 
 
     preset.IsFavorite = favorites?.[preset.ID];
     return preset;
diff --git a/Engine/Plugins/VirtualProduction/RemoteControlWebInterface/WebApp/_package.json b/Engine/Plugins/VirtualProduction/RemoteControlWebInterface/WebApp/_package.json
--- a/Engine/Plugins/VirtualProduction/RemoteControlWebInterface/WebApp/_package.json
+++ b/Engine/Plugins/VirtualProduction/RemoteControlWebInterface/WebApp/_package.json
@@ -1,6 +1,6 @@
 {
   "name": "conductor",
-  "version": "0.3.83",
+  "version": "0.3.84 ",
   "scripts": {
     "postinstall": "concurrently \"cd ./Client && npm install && cd ..\" \"cd ./Server && npm install && cd ..\"",
     "build": "node ./Scripts/build.js",
