Refresh pivot gizmo post RC location updates

diff --git a/Engine/Plugins/VirtualProduction/RemoteControl/Source/RemoteControl/Private/RemoteControlModule.cpp b/Engine/Plugins/VirtualProduction/RemoteControl/Source/RemoteControl/Private/RemoteControlModule.cpp
--- a/Engine/Plugins/VirtualProduction/RemoteControl/Source/RemoteControl/Private/RemoteControlModule.cpp
+++ b/Engine/Plugins/VirtualProduction/RemoteControl/Source/RemoteControl/Private/RemoteControlModule.cpp
@@ -32,8 +32,10 @@
 #include "Editor.h"
 #include "Editor/EditorEngine.h"
 #include "Editor/Transactor.h"
+#include "Editor/UnrealEdEngine.h"
 #include "HAL/IConsoleManager.h"
 #include "ScopedTransaction.h"
+#include "UnrealEdGlobals.h"
 #endif
 
 DEFINE_LOG_CATEGORY(LogRemoteControl);
@@ -1119,6 +1121,12 @@ bool FRemoteControlModule::SetObjectProperties(const FRCObjectReference& ObjectA
 			if (RemoteControlSetterUtils::ConvertModificationToFunctionCall(Args, InterceptionPayload))
 			{
 				const bool bResult = InvokeCall(Call, InterceptionPayload.Type, InterceptionPayload.Payload);
+
+				if (bResult)
+				{
+					RefreshEditorPostSetObjectProperties(ObjectAccess);
+				}
+
 				return bResult;
 			}
 		}
@@ -1180,6 +1188,12 @@ bool FRemoteControlModule::SetObjectProperties(const FRCObjectReference& ObjectA
 		if (RemoteControlSetterUtils::ConvertModificationToFunctionCall(Args))
 		{
 			const bool bResult = InvokeCall(Call);
+
+			if (bResult)
+			{
+				RefreshEditorPostSetObjectProperties(ObjectAccess);
+			}
+
 			return bResult;
 		}
 	}
@@ -1351,11 +1365,40 @@ bool FRemoteControlModule::SetObjectProperties(const FRCObjectReference& ObjectA
 			EntityFactoryPair.Value->PostSetObjectProperties(ObjectAccess.Object.Get(), bSuccess);
 		}
 
+		if (bSuccess)
+		{
+			RefreshEditorPostSetObjectProperties(ObjectAccess);
+		}
+
 		return bSuccess;
 	}
 	return false;
 }
 
+void FRemoteControlModule::RefreshEditorPostSetObjectProperties(const FRCObjectReference& ObjectAccess)
+{
+#if WITH_EDITOR
+	UObject* Object = ObjectAccess.Object.Get();
+
+	if (USceneComponent* SceneComponent = Cast<USceneComponent>(Object))
+	{
+		if (FProperty* Property = ObjectAccess.Property.Get())
+		{
+			if(Property->GetName() == "RelativeLocation")
+			{
+				if (AActor* Actor = SceneComponent->GetOwner())
+				{
+					if (Actor->IsSelectedInEditor())
+					{
+						GUnrealEd->UpdatePivotLocationForSelection();
+					}
+				}
+			}
+		}
+	}
+#endif
+}
+
 bool FRemoteControlModule::SetPresetController(const FName PresetName, const FName ControllerName, IStructDeserializerBackend& Backend, const TArray<uint8>& InPayload, const bool bAllowIntercept)
 {
 	URemoteControlPreset* Preset = IRemoteControlModule::Get().ResolvePreset(PresetName);
diff --git a/Engine/Plugins/VirtualProduction/RemoteControl/Source/RemoteControl/Private/RemoteControlModule.h b/Engine/Plugins/VirtualProduction/RemoteControl/Source/RemoteControl/Private/RemoteControlModule.h
--- a/Engine/Plugins/VirtualProduction/RemoteControl/Source/RemoteControl/Private/RemoteControlModule.h
+++ b/Engine/Plugins/VirtualProduction/RemoteControl/Source/RemoteControl/Private/RemoteControlModule.h
@@ -67,6 +67,8 @@ public:
 	//~ End IRemoteControlModule
 
 private:
+	/** Refreshes Editor related visuals like location Gizmo for relevant properties (like Location of a SceneComponent) */
+	void RefreshEditorPostSetObjectProperties(const FRCObjectReference& ObjectAccess);
 
 	/** Cache all presets in the project for the ResolvePreset function. */
 	void CachePresets() const;
diff --git a/Engine/Plugins/VirtualProduction/RemoteControl/Source/RemoteControlLogic/Private/Action/Bind/RCPropertyBindAction.cpp b/Engine/Plugins/VirtualProduction/RemoteControl/Source/RemoteControlLogic/Private/Action/Bind/RCPropertyBindAction.cpp
--- a/Engine/Plugins/VirtualProduction/RemoteControl/Source/RemoteControlLogic/Private/Action/Bind/RCPropertyBindAction.cpp
+++ b/Engine/Plugins/VirtualProduction/RemoteControl/Source/RemoteControlLogic/Private/Action/Bind/RCPropertyBindAction.cpp
@@ -211,6 +211,9 @@ void URCPropertyBindAction::Execute() const
 				{
 					if (Actor->IsSelectedInEditor())
 					{
+						// Note: For regular Exposed Properties this is handled via RefreshEditorPostSetObjectProperties in FRemoteControlModule::SetObjectProperties
+						// Bind needs this explicit path because the Controller's property name is not guaranteed to be "RelativeLocation" (which is the pattern match used in that function).
+						// This can be revisited in the future depending on the number of related usecases that arise in this space.
 						GUnrealEd->UpdatePivotLocationForSelection();
 					}
 				}
