[EpicStageApp] Fix interactive changes to nested colors failing to propagate

diff --git a/Engine/Plugins/VirtualProduction/RemoteControl/Source/RemoteControl/Private/RemoteControlFunctionLibrary.cpp b/Engine/Plugins/VirtualProduction/RemoteControl/Source/RemoteControl/Private/RemoteControlFunctionLibrary.cpp
--- a/Engine/Plugins/VirtualProduction/RemoteControl/Source/RemoteControl/Private/RemoteControlFunctionLibrary.cpp
+++ b/Engine/Plugins/VirtualProduction/RemoteControl/Source/RemoteControl/Private/RemoteControlFunctionLibrary.cpp
@@ -26,30 +26,6 @@ namespace RemoteControlFunctionLibrary
 		return GroupId;
 	}
 
-	/**
-	 * Traverse the nested struct properties on an object to find the property with the given period-separated path.
-	 * If found, return true and set OutProperty to the property matching the name and OutContainer to the object containing its value (within TargetObject).
-	 */
-	bool GetNestedPropertyAndContainerFromPath(UObject* Object, const FString& PropertyName, FProperty*& OutProperty, void*& OutContainer)
-	{
-		if (!Object)
-		{
-			return false;
-		}
-
-		FRCObjectReference ObjectRef;
-		IRemoteControlModule::Get().ResolveObject(ERCAccess::WRITE_ACCESS, Object->GetPathName(), PropertyName, ObjectRef);
-
-		if (!ObjectRef.IsValid())
-		{
-			return false;
-		}
-
-		OutProperty = ObjectRef.Property.Get();
-		OutContainer = ObjectRef.ContainerAdress;
-		return true;
-	}
-
 	/** Convert a value from the range for a normal linear color (0-1) to the range for a color grading wheel (provided by MinValue and MaxValue). */
 	void TransformLinearColorRangeToColorGradingRange(FVector4& Value, float MinValue, float MaxValue)
 	{
@@ -107,18 +83,25 @@ namespace RemoteControlFunctionLibrary
 	}
 
 	/** Set a property on an object and fire relevant events. */
-	void SetPropertyAndFireEvents(UObject* Object, FProperty* Property, void* Container, void* Value, bool bIsInteractive)
+	void SetPropertyAndFireEvents(const FRCObjectReference& InObjectRef, const void* InValue, bool bIsInteractive)
 	{
 #if WITH_EDITOR
-		Object->PreEditChange(Property);
-		Object->Modify();
+		FEditPropertyChain PropertyChain;
+		InObjectRef.PropertyPathInfo.ToEditPropertyChain(PropertyChain);
+		InObjectRef.Object->PreEditChange(PropertyChain);
+		InObjectRef.Object->Modify();
 #endif
 
-		Property->SetValue_InContainer(Container, Value);
+		InObjectRef.Property->SetValue_InContainer(InObjectRef.ContainerAdress, InValue);
 
 #if WITH_EDITOR
-		FPropertyChangedEvent ChangeEvent(Property, bIsInteractive ? EPropertyChangeType::Interactive : EPropertyChangeType::ValueSet);
-		Object->PostEditChangeProperty(ChangeEvent);
+		SnapshotTransactionBuffer(InObjectRef.Object.Get());
+
+		FPropertyChangedEvent PropertyEvent = InObjectRef.PropertyPathInfo.ToPropertyChangedEvent();
+		PropertyEvent.ChangeType = bIsInteractive ? EPropertyChangeType::Interactive : EPropertyChangeType::ValueSet;
+
+		FPropertyChangedChainEvent PropertyChainEvent(PropertyChain, PropertyEvent);
+		InObjectRef.Object->PostEditChangeChainProperty(PropertyChainEvent);
 #endif
 	}
 }
@@ -155,14 +138,15 @@ bool URemoteControlFunctionLibrary::ExposeActor(URemoteControlPreset* Preset, AA
 
 bool URemoteControlFunctionLibrary::ApplyColorWheelDelta(UObject* TargetObject, const FString& PropertyName, const FColorWheelColor& DeltaValue, const FColorWheelColor& ReferenceColor, bool bIsInteractive)
 {
-	FProperty* Property;
-	void* Container;
-	if (!RemoteControlFunctionLibrary::GetNestedPropertyAndContainerFromPath(TargetObject, PropertyName, Property, Container))
+	FRCObjectReference ObjectRef;
+	IRemoteControlModule::Get().ResolveObject(ERCAccess::WRITE_ACCESS, TargetObject->GetPathName(), PropertyName, ObjectRef);
+
+	if (!ObjectRef.IsValid())
 	{
 		return false;
 	}
 
-	if (const FStructProperty* ColorProperty = CastField<FStructProperty>(Property))
+	if (const FStructProperty* ColorProperty = CastField<FStructProperty>(ObjectRef.Property.Get()))
 	{
 		if (ColorProperty->Struct != TBaseStructure<FLinearColor>::Get())
 		{
@@ -171,14 +155,14 @@ bool URemoteControlFunctionLibrary::ApplyColorWheelDelta(UObject* TargetObject,
 		}
 
 		FLinearColor Color;
-		ColorProperty->GetValue_InContainer(Container, &Color);
+		ColorProperty->GetValue_InContainer(ObjectRef.ContainerAdress, &Color);
 
 		RemoteControlFunctionLibrary::ApplyWheelColorBaseDelta(Color, DeltaValue, ReferenceColor, 0.f, 1.f);
 
 		// Apply the alpha change directly
 		Color.A = FMath::Clamp(Color.A + DeltaValue.Alpha, 0.f, 1.f);
 
-		RemoteControlFunctionLibrary::SetPropertyAndFireEvents(TargetObject, Property, Container, &Color, bIsInteractive);
+		RemoteControlFunctionLibrary::SetPropertyAndFireEvents(ObjectRef, &Color, bIsInteractive);
 
 		return true;
 	}
@@ -188,19 +172,20 @@ bool URemoteControlFunctionLibrary::ApplyColorWheelDelta(UObject* TargetObject,
 
 bool URemoteControlFunctionLibrary::ApplyColorGradingWheelDelta(UObject* TargetObject, const FString& PropertyName, const FColorGradingWheelColor& DeltaValue, const FColorGradingWheelColor& ReferenceColor, bool bIsInteractive, float MinValue, float MaxValue)
 {
-	if (MinValue >= MaxValue)
+	if (MinValue >= MaxValue || !TargetObject)
 	{
 		return false;
 	}
 
-	FProperty* Property = nullptr;
-	void* Container = nullptr;
-	if (!RemoteControlFunctionLibrary::GetNestedPropertyAndContainerFromPath(TargetObject, PropertyName, Property, Container))
+	FRCObjectReference ObjectRef;
+	IRemoteControlModule::Get().ResolveObject(ERCAccess::WRITE_ACCESS, TargetObject->GetPathName(), PropertyName, ObjectRef);
+	
+	if (!ObjectRef.IsValid())
 	{
 		return false;
 	}
 
-	if (const FStructProperty* ColorProperty = CastField<FStructProperty>(Property))
+	if (const FStructProperty* ColorProperty = CastField<FStructProperty>(ObjectRef.Property.Get()))
 	{
 		if (ColorProperty->Struct != TBaseStructure<FVector4>::Get())
 		{
@@ -210,7 +195,7 @@ bool URemoteControlFunctionLibrary::ApplyColorGradingWheelDelta(UObject* TargetO
 
 		// Get the vector value and convert the color portion of it to HSV
 		FVector4 VectorValue;
-		ColorProperty->GetValue_InContainer(Container, &VectorValue);
+		ColorProperty->GetValue_InContainer(ObjectRef.ContainerAdress, &VectorValue);
 
 		const float OldLuminance = VectorValue.W;
 
@@ -223,7 +208,7 @@ bool URemoteControlFunctionLibrary::ApplyColorGradingWheelDelta(UObject* TargetO
 		RemoteControlFunctionLibrary::TransformLinearColorRangeToColorGradingRange(VectorValue, MinValue, MaxValue);
 		VectorValue.W = OldLuminance + DeltaValue.Luminance;
 
-		RemoteControlFunctionLibrary::SetPropertyAndFireEvents(TargetObject, Property, Container, &VectorValue, bIsInteractive);
+		RemoteControlFunctionLibrary::SetPropertyAndFireEvents(ObjectRef, &VectorValue, bIsInteractive);
 
 		return true;
 	}
diff --git a/Engine/Plugins/VirtualProduction/RemoteControl/Source/WebRemoteControl/Private/WebRemoteControlInternalUtils.h b/Engine/Plugins/VirtualProduction/RemoteControl/Source/WebRemoteControl/Private/WebRemoteControlInternalUtils.h
--- a/Engine/Plugins/VirtualProduction/RemoteControl/Source/WebRemoteControl/Private/WebRemoteControlInternalUtils.h
+++ b/Engine/Plugins/VirtualProduction/RemoteControl/Source/WebRemoteControl/Private/WebRemoteControlInternalUtils.h
@@ -399,12 +399,13 @@ namespace WebRemoteControlInternalUtils
 			}
 
 #if WITH_EDITOR
+			FEditPropertyChain PropertyChain;
+
 			if (Access == ERCAccess::WRITE_MANUAL_TRANSACTION_ACCESS && ObjectRef.Object.IsValid())
 			{
 				// This transaction is being manually controlled, so RemoteControlModule's automatic transaction handling won't call this for us
-				FEditPropertyChain PreEditChain;
-				ObjectRef.PropertyPathInfo.ToEditPropertyChain(PreEditChain);
-				ObjectRef.Object->PreEditChange(PreEditChain);
+				ObjectRef.PropertyPathInfo.ToEditPropertyChain(PropertyChain);
+				ObjectRef.Object->PreEditChange(PropertyChain);
 			}
 #endif
 
@@ -424,10 +425,14 @@ namespace WebRemoteControlInternalUtils
 #if WITH_EDITOR
 			if (Access == ERCAccess::WRITE_MANUAL_TRANSACTION_ACCESS && ObjectRef.Object.IsValid())
 			{
-				// This transaction is being manually controlled, so RemoteControlModule's automatic transaction handling won't call this for us
+				// This transaction is being manually controlled, so RemoteControlModule's automatic transaction handling won't call these functions for us
+				SnapshotTransactionBuffer(ObjectRef.Object.Get());
+
 				FPropertyChangedEvent PropertyEvent = ObjectRef.PropertyPathInfo.ToPropertyChangedEvent();
 				PropertyEvent.ChangeType = EPropertyChangeType::Interactive;
-				ObjectRef.Object->PostEditChangeProperty(PropertyEvent);
+
+				FPropertyChangedChainEvent PropertyChainEvent(PropertyChain, PropertyEvent);
+				ObjectRef.Object->PostEditChangeChainProperty(PropertyChainEvent);
 			}
 #endif
 		}
