Restore missing UI options for Path behaviour (Texture type)

diff --git a/Engine/Plugins/VirtualProduction/RemoteControl/Source/RemoteControlUI/Private/UI/Behaviour/Builtin/Path/SRCBehaviourSetAssetByPath.cpp b/Engine/Plugins/VirtualProduction/RemoteControl/Source/RemoteControlUI/Private/UI/Behaviour/Builtin/Path/SRCBehaviourSetAssetByPath.cpp
--- a/Engine/Plugins/VirtualProduction/RemoteControl/Source/RemoteControlUI/Private/UI/Behaviour/Builtin/Path/SRCBehaviourSetAssetByPath.cpp
+++ b/Engine/Plugins/VirtualProduction/RemoteControl/Source/RemoteControlUI/Private/UI/Behaviour/Builtin/Path/SRCBehaviourSetAssetByPath.cpp
@@ -24,7 +24,7 @@ void SRCBehaviourSetAssetByPath::Construct(const FArguments& InArgs, TSharedRef<
 	const TSharedRef<SHorizontalBox> InternalExternalSwitchWidget = SNew(SHorizontalBox)
 		.Visibility_Lambda([PathBehaviour]()
 		{
-			return PathBehaviour->AssetClass->GetFName() == NAME_Vector ? EVisibility::Visible : EVisibility::Collapsed;
+				return PathBehaviour->AssetClass->GetFName() == "Texture" ? EVisibility::Visible : EVisibility::Collapsed;
 		});
 
 	if (!PathBehaviour)
