[RemoteControl] Let user choose whether to force a property's edit conditions on when exposed remotely

diff --git a/Engine/Plugins/VirtualProduction/RemoteControl/Source/RemoteControl/Private/RemoteControlPreset.cpp b/Engine/Plugins/VirtualProduction/RemoteControl/Source/RemoteControl/Private/RemoteControlPreset.cpp
--- a/Engine/Plugins/VirtualProduction/RemoteControl/Source/RemoteControl/Private/RemoteControlPreset.cpp
+++ b/Engine/Plugins/VirtualProduction/RemoteControl/Source/RemoteControl/Private/RemoteControlPreset.cpp
@@ -116,9 +116,10 @@ FRemoteControlPresetExposeArgs::FRemoteControlPresetExposeArgs()
 {
 }
 
-FRemoteControlPresetExposeArgs::FRemoteControlPresetExposeArgs(FString InLabel, FGuid InGroupId)
+FRemoteControlPresetExposeArgs::FRemoteControlPresetExposeArgs(FString InLabel, FGuid InGroupId, bool bEnableEditCondition)
 	: Label(MoveTemp(InLabel))
 	, GroupId(MoveTemp(InGroupId))
+	, bEnableEditCondition(bEnableEditCondition)
 {
 }
 
@@ -712,7 +713,10 @@ TWeakPtr<FRemoteControlProperty> URemoteControlPreset::ExposeProperty(UObject* O
 		return nullptr;
 	}
 
-	RCPropertyPtr->EnableEditCondition();
+	if (Args.bEnableEditCondition)
+	{
+		RCPropertyPtr->EnableEditCondition();
+	}
 
 	if (PropertyShouldBeWatched(RCPropertyPtr))
 	{
diff --git a/Engine/Plugins/VirtualProduction/RemoteControl/Source/RemoteControl/Public/RemoteControlPreset.h b/Engine/Plugins/VirtualProduction/RemoteControl/Source/RemoteControl/Public/RemoteControlPreset.h
--- a/Engine/Plugins/VirtualProduction/RemoteControl/Source/RemoteControl/Public/RemoteControlPreset.h
+++ b/Engine/Plugins/VirtualProduction/RemoteControl/Source/RemoteControl/Public/RemoteControlPreset.h
@@ -37,12 +37,14 @@ DECLARE_MULTICAST_DELEGATE(FOnVirtualPropertyContainerModified);
 struct REMOTECONTROL_API FRemoteControlPresetExposeArgs
 {
 	FRemoteControlPresetExposeArgs();
-	FRemoteControlPresetExposeArgs(FString Label, FGuid GroupId);
+	FRemoteControlPresetExposeArgs(FString Label, FGuid GroupId, bool bEnableEditCondition = true);
 
 	/** (Optional) The label to use for the new exposed entity. */
 	FString Label;
 	/** (Optional) The group in which to put the field. */
 	FGuid GroupId;
+	/** Whether to automatically enable the edit condition for the exposed property. */
+	bool bEnableEditCondition;
 };
 
 /**
diff --git a/Engine/Plugins/VirtualProduction/RemoteControl/Source/WebRemoteControl/Private/RemoteControlRequest.h b/Engine/Plugins/VirtualProduction/RemoteControl/Source/WebRemoteControl/Private/RemoteControlRequest.h
--- a/Engine/Plugins/VirtualProduction/RemoteControl/Source/WebRemoteControl/Private/RemoteControlRequest.h
+++ b/Engine/Plugins/VirtualProduction/RemoteControl/Source/WebRemoteControl/Private/RemoteControlRequest.h
@@ -300,6 +300,12 @@ struct FRCPresetExposePropertyRequest : public FRCRequest
 	 */
 	UPROPERTY()
 	FString GroupName;
+
+	/**
+	 * Whether to automatically enable the edit condition for the exposed property.
+	 */
+	UPROPERTY()
+	bool EnableEditCondition;
 };
 
 /**
diff --git a/Engine/Plugins/VirtualProduction/RemoteControl/Source/WebRemoteControl/Private/WebRemoteControl.cpp b/Engine/Plugins/VirtualProduction/RemoteControl/Source/WebRemoteControl/Private/WebRemoteControl.cpp
--- a/Engine/Plugins/VirtualProduction/RemoteControl/Source/WebRemoteControl/Private/WebRemoteControl.cpp
+++ b/Engine/Plugins/VirtualProduction/RemoteControl/Source/WebRemoteControl/Private/WebRemoteControl.cpp
@@ -1398,6 +1398,7 @@ bool FWebRemoteControlModule::HandlePresetExposePropertyRoute(const FHttpServerR
 
 	FRemoteControlPresetExposeArgs ExposeArgs;
 	ExposeArgs.Label = ExposePropertyRequest.Label;
+	ExposeArgs.bEnableEditCondition = ExposePropertyRequest.EnableEditCondition;
 
 	if (ExposePropertyRequest.GroupName.Len() > 0)
 	{
