Non-unity/pch compile fixes when building UnrealEditor linux

diff --git a/Engine/Plugins/Runtime/nDisplay/Source/DisplayClusterLightCardEditorShaders/Public/DisplayClusterMeshProjectionRenderer.h b/Engine/Plugins/Runtime/nDisplay/Source/DisplayClusterLightCardEditorShaders/Public/DisplayClusterMeshProjectionRenderer.h
--- a/Engine/Plugins/Runtime/nDisplay/Source/DisplayClusterLightCardEditorShaders/Public/DisplayClusterMeshProjectionRenderer.h
+++ b/Engine/Plugins/Runtime/nDisplay/Source/DisplayClusterLightCardEditorShaders/Public/DisplayClusterMeshProjectionRenderer.h
@@ -8,6 +8,7 @@
 
 class AActor;
 class UPrimitiveComponent;
+class FHitProxyConsumer;
 class FSceneView;
 class FViewInfo;
 class FPrimitiveDrawInterface;
diff --git a/Engine/Plugins/Runtime/nDisplay/Source/DisplayClusterMoviePipelineEditor/Private/Details/DisplayClusterMoviePipelineEditorNodeSelection.h b/Engine/Plugins/Runtime/nDisplay/Source/DisplayClusterMoviePipelineEditor/Private/Details/DisplayClusterMoviePipelineEditorNodeSelection.h
--- a/Engine/Plugins/Runtime/nDisplay/Source/DisplayClusterMoviePipelineEditor/Private/Details/DisplayClusterMoviePipelineEditorNodeSelection.h
+++ b/Engine/Plugins/Runtime/nDisplay/Source/DisplayClusterMoviePipelineEditor/Private/Details/DisplayClusterMoviePipelineEditorNodeSelection.h
@@ -4,9 +4,11 @@
 
 #include "CoreMinimal.h"
 #include "DisplayClusterMoviePipelineEditorBaseTypeCustomization.h"
+#include "Misc/Attribute.h"
 
 class UDisplayClusterConfigurationData;
 class SDisplayClusterMoviePipelineEditorSearchableComboBox;
+class SWidget;
 
 /**
  * Helper class used to construct and manage a custom array property widget that displays for each element a dropdown containing
diff --git a/Engine/Plugins/Runtime/nDisplay/Source/DisplayClusterScenePreview/Private/DisplayClusterScenePreviewModule.cpp b/Engine/Plugins/Runtime/nDisplay/Source/DisplayClusterScenePreview/Private/DisplayClusterScenePreviewModule.cpp
--- a/Engine/Plugins/Runtime/nDisplay/Source/DisplayClusterScenePreview/Private/DisplayClusterScenePreviewModule.cpp
+++ b/Engine/Plugins/Runtime/nDisplay/Source/DisplayClusterScenePreview/Private/DisplayClusterScenePreviewModule.cpp
@@ -8,8 +8,10 @@
 #include "Components/DisplayClusterPreviewComponent.h"
 #include "DisplayClusterLightCardActor.h"
 #include "DisplayClusterRootActor.h"
+#include "Engine/Canvas.h"
 #include "Engine/TextureRenderTarget2D.h"
 #include "Misc/TransactionObjectEvent.h"
+#include "TextureResource.h"
 
 #if WITH_EDITOR
 #include "LevelEditorViewport.h"
diff --git a/Engine/Plugins/Runtime/nDisplay/Source/DisplayClusterShaders/Private/Shaders/DisplayClusterShadersPreprocess_UVLightCards.cpp b/Engine/Plugins/Runtime/nDisplay/Source/DisplayClusterShaders/Private/Shaders/DisplayClusterShadersPreprocess_UVLightCards.cpp
--- a/Engine/Plugins/Runtime/nDisplay/Source/DisplayClusterShaders/Private/Shaders/DisplayClusterShadersPreprocess_UVLightCards.cpp
+++ b/Engine/Plugins/Runtime/nDisplay/Source/DisplayClusterShaders/Private/Shaders/DisplayClusterShadersPreprocess_UVLightCards.cpp
@@ -15,6 +15,7 @@
 #include "Renderer/Private/ScenePrivate.h"
 #include "SceneRenderTargetParameters.h"
 #include "MeshPassProcessor.inl"
+#include "UnrealClient.h"
 
 class FUVLightCardVS : public FMeshMaterialShader
 {
diff --git a/Engine/Plugins/Runtime/nDisplay/Source/DisplayClusterTests/Private/DisplayClusterTestUtils.cpp b/Engine/Plugins/Runtime/nDisplay/Source/DisplayClusterTests/Private/DisplayClusterTestUtils.cpp
--- a/Engine/Plugins/Runtime/nDisplay/Source/DisplayClusterTests/Private/DisplayClusterTestUtils.cpp
+++ b/Engine/Plugins/Runtime/nDisplay/Source/DisplayClusterTests/Private/DisplayClusterTestUtils.cpp
@@ -3,6 +3,8 @@
 #include "DisplayClusterTestUtils.h"
 
 #include "DisplayClusterTestsModule.h"
+#include "Engine/Engine.h"
+#include "Engine/World.h"
 #include "ObjectTools.h"
 #include "PackageTools.h"
 #include "AssetRegistry/AssetRegistryModule.h"
diff --git a/Engine/Plugins/Tests/EditorTests/Source/EditorTests/Private/UnrealEd/GeometryTests.cpp b/Engine/Plugins/Tests/EditorTests/Source/EditorTests/Private/UnrealEd/GeometryTests.cpp
--- a/Engine/Plugins/Tests/EditorTests/Source/EditorTests/Private/UnrealEd/GeometryTests.cpp
+++ b/Engine/Plugins/Tests/EditorTests/Source/EditorTests/Private/UnrealEd/GeometryTests.cpp
@@ -17,6 +17,7 @@
 #include "EditorModeManager.h"
 
 //Automation
+#include "Model.h"
 #include "Tests/AutomationCommon.h"
 #include "Tests/AutomationEditorCommon.h"
 
diff --git a/Engine/Plugins/Tests/EditorTests/Source/EditorTests/Private/UnrealEd/PlayMapInPIETests.cpp b/Engine/Plugins/Tests/EditorTests/Source/EditorTests/Private/UnrealEd/PlayMapInPIETests.cpp
--- a/Engine/Plugins/Tests/EditorTests/Source/EditorTests/Private/UnrealEd/PlayMapInPIETests.cpp
+++ b/Engine/Plugins/Tests/EditorTests/Source/EditorTests/Private/UnrealEd/PlayMapInPIETests.cpp
@@ -1,6 +1,8 @@
 // Copyright Epic Games, Inc. All Rights Reserved.
 
 #include "CoreMinimal.h"
+#include "Editor.h"
+#include "Engine/Texture.h"
 #include "Misc/AutomationTest.h"
 #include "Modules/ModuleManager.h"
 #include "Engine/Engine.h"
diff --git a/Engine/Plugins/VirtualProduction/CameraCalibration/Source/CameraCalibrationEditor/Private/AssetEditor/CameraCalibrationStepsController.cpp b/Engine/Plugins/VirtualProduction/CameraCalibration/Source/CameraCalibrationEditor/Private/AssetEditor/CameraCalibrationStepsController.cpp
--- a/Engine/Plugins/VirtualProduction/CameraCalibration/Source/CameraCalibrationEditor/Private/AssetEditor/CameraCalibrationStepsController.cpp
+++ b/Engine/Plugins/VirtualProduction/CameraCalibration/Source/CameraCalibrationEditor/Private/AssetEditor/CameraCalibrationStepsController.cpp
@@ -41,7 +41,9 @@
 #include "Profile/IMediaProfileManager.h"
 #include "Profile/MediaProfile.h"
 #include "SCameraCalibrationSteps.h"
+#include "TextureResource.h"
 #include "TimeSynchronizableMediaSource.h"
+#include "UnrealClient.h"
 #include "UObject/ObjectMacros.h"
 #include "UObject/UObjectGlobals.h"
 #include "UObject/UnrealType.h"
diff --git a/Engine/Plugins/VirtualProduction/CameraCalibration/Source/CameraCalibrationEditor/Private/AssetEditor/CameraCalibrationStepsController.h b/Engine/Plugins/VirtualProduction/CameraCalibration/Source/CameraCalibrationEditor/Private/AssetEditor/CameraCalibrationStepsController.h
--- a/Engine/Plugins/VirtualProduction/CameraCalibration/Source/CameraCalibrationEditor/Private/AssetEditor/CameraCalibrationStepsController.h
+++ b/Engine/Plugins/VirtualProduction/CameraCalibration/Source/CameraCalibrationEditor/Private/AssetEditor/CameraCalibrationStepsController.h
@@ -7,6 +7,7 @@
 #include "Containers/Ticker.h"
 #include "Engine/TextureRenderTarget2D.h"
 #include "Engine/World.h"
+#include "Input/Reply.h"
 #include "LensFile.h"
 #include "UObject/StrongObjectPtr.h"
 
diff --git a/Engine/Plugins/VirtualProduction/CameraCalibration/Source/CameraCalibrationEditor/Private/AssetEditor/NodalOffsetTool.cpp b/Engine/Plugins/VirtualProduction/CameraCalibration/Source/CameraCalibrationEditor/Private/AssetEditor/NodalOffsetTool.cpp
--- a/Engine/Plugins/VirtualProduction/CameraCalibration/Source/CameraCalibrationEditor/Private/AssetEditor/NodalOffsetTool.cpp
+++ b/Engine/Plugins/VirtualProduction/CameraCalibration/Source/CameraCalibrationEditor/Private/AssetEditor/NodalOffsetTool.cpp
@@ -11,6 +11,8 @@
 #include "CameraCalibrationTypes.h"
 #include "DesktopPlatformModule.h"
 #include "Dom/JsonObject.h"
+#include "Engine/Engine.h"
+#include "Framework/Application/SlateApplication.h"
 #include "HAL/FileManager.h"
 #include "IDesktopPlatform.h"
 #include "ImageCenterTool.h"
diff --git a/Engine/Plugins/VirtualProduction/CameraCalibration/Source/CameraCalibrationEditor/Private/AssetEditor/SCameraCalibrationSteps.h b/Engine/Plugins/VirtualProduction/CameraCalibration/Source/CameraCalibrationEditor/Private/AssetEditor/SCameraCalibrationSteps.h
--- a/Engine/Plugins/VirtualProduction/CameraCalibration/Source/CameraCalibrationEditor/Private/AssetEditor/SCameraCalibrationSteps.h
+++ b/Engine/Plugins/VirtualProduction/CameraCalibration/Source/CameraCalibrationEditor/Private/AssetEditor/SCameraCalibrationSteps.h
@@ -14,7 +14,8 @@ class SCheckBox;
 
 template<typename OptionType>
 class SComboBox;
-
+class SHorizontalBox;
+class SVerticalBox;
 class SWidgetSwitcher;
 
 struct FArguments;
diff --git a/Engine/Plugins/VirtualProduction/CameraCalibration/Source/CameraCalibrationEditor/Private/AssetEditor/SLensDataAddPointDialog.cpp b/Engine/Plugins/VirtualProduction/CameraCalibration/Source/CameraCalibrationEditor/Private/AssetEditor/SLensDataAddPointDialog.cpp
--- a/Engine/Plugins/VirtualProduction/CameraCalibration/Source/CameraCalibrationEditor/Private/AssetEditor/SLensDataAddPointDialog.cpp
+++ b/Engine/Plugins/VirtualProduction/CameraCalibration/Source/CameraCalibrationEditor/Private/AssetEditor/SLensDataAddPointDialog.cpp
@@ -18,7 +18,9 @@
 #include "Styling/AppStyle.h"
 #include "Widgets/Input/SButton.h"
 #include "Widgets/Input/SCheckBox.h"
+#include "Widgets/Input/SComboButton.h"
 #include "Widgets/Input/SNumericEntryBox.h"
+#include "Widgets/Layout/SSplitter.h"
 
 #define LOCTEXT_NAMESPACE "LensDataAddPointDialog"
 
diff --git a/Engine/Plugins/VirtualProduction/CameraCalibration/Source/CameraCalibrationEditor/Private/AssetEditor/SLensDataAddPointDialog.h b/Engine/Plugins/VirtualProduction/CameraCalibration/Source/CameraCalibrationEditor/Private/AssetEditor/SLensDataAddPointDialog.h
--- a/Engine/Plugins/VirtualProduction/CameraCalibration/Source/CameraCalibrationEditor/Private/AssetEditor/SLensDataAddPointDialog.h
+++ b/Engine/Plugins/VirtualProduction/CameraCalibration/Source/CameraCalibrationEditor/Private/AssetEditor/SLensDataAddPointDialog.h
@@ -2,6 +2,7 @@
 
 #pragma once
 
+#include "Styling/SlateTypes.h"
 #include "Widgets/SCompoundWidget.h"
 
 #include "CameraCalibrationEditorCommon.h"
@@ -9,6 +10,8 @@
 #include "SLensEvaluation.h"
 #include "UObject/StrongObjectPtr.h"
 
+class SBorder;
+
 class SLensDataAddPointDialog : public SCompoundWidget
 {
 private:
diff --git a/Engine/Plugins/VirtualProduction/CameraCalibration/Source/CameraCalibrationEditor/Private/AssetEditor/SLensDataListItem.cpp b/Engine/Plugins/VirtualProduction/CameraCalibration/Source/CameraCalibrationEditor/Private/AssetEditor/SLensDataListItem.cpp
--- a/Engine/Plugins/VirtualProduction/CameraCalibration/Source/CameraCalibrationEditor/Private/AssetEditor/SLensDataListItem.cpp
+++ b/Engine/Plugins/VirtualProduction/CameraCalibration/Source/CameraCalibrationEditor/Private/AssetEditor/SLensDataListItem.cpp
@@ -9,6 +9,8 @@
 #include "Widgets/Images/SImage.h"
 #include "Widgets/Input/SButton.h"
 #include "SCameraCalibrationRemovePointDialog.h"
+#include "Widgets/Views/SListView.h"
+#include "Widgets/Views/STableRow.h"
 
 #define LOCTEXT_NAMESPACE "LensDataListItem"
 
diff --git a/Engine/Plugins/VirtualProduction/CameraCalibration/Source/CameraCalibrationEditor/Private/AssetEditor/Viewport/SimulcamEditorViewportClient.cpp b/Engine/Plugins/VirtualProduction/CameraCalibration/Source/CameraCalibrationEditor/Private/AssetEditor/Viewport/SimulcamEditorViewportClient.cpp
--- a/Engine/Plugins/VirtualProduction/CameraCalibration/Source/CameraCalibrationEditor/Private/AssetEditor/Viewport/SimulcamEditorViewportClient.cpp
+++ b/Engine/Plugins/VirtualProduction/CameraCalibration/Source/CameraCalibrationEditor/Private/AssetEditor/Viewport/SimulcamEditorViewportClient.cpp
@@ -9,9 +9,11 @@
 #include "Editor/UnrealEdEngine.h"
 #include "EngineModule.h"
 #include "EngineUtils.h"
+#include "Framework/Application/SlateApplication.h"
 #include "ImageUtils.h"
 #include "Slate/SceneViewport.h"
 #include "Texture2DPreview.h"
+#include "TextureResource.h"
 #include "UnrealEdGlobals.h"
 
 
diff --git a/Engine/Plugins/VirtualProduction/CameraCalibration/Source/CameraCalibrationEditor/Private/AssetEditor/Viewport/SimulcamEditorViewportClient.h b/Engine/Plugins/VirtualProduction/CameraCalibration/Source/CameraCalibrationEditor/Private/AssetEditor/Viewport/SimulcamEditorViewportClient.h
--- a/Engine/Plugins/VirtualProduction/CameraCalibration/Source/CameraCalibrationEditor/Private/AssetEditor/Viewport/SimulcamEditorViewportClient.h
+++ b/Engine/Plugins/VirtualProduction/CameraCalibration/Source/CameraCalibrationEditor/Private/AssetEditor/Viewport/SimulcamEditorViewportClient.h
@@ -5,6 +5,7 @@
 #include "CoreMinimal.h"
 #include "InputCoreTypes.h"
 #include "UObject/GCObject.h"
+#include "ViewportClient.h"
 
 class FCanvas;
 class ITextureEditorToolkit;
diff --git a/Engine/Plugins/VirtualProduction/CameraCalibration/Source/CameraCalibrationEditor/Private/Calibrators/CameraNodalOffsetAlgoPoints.cpp b/Engine/Plugins/VirtualProduction/CameraCalibration/Source/CameraCalibrationEditor/Private/Calibrators/CameraNodalOffsetAlgoPoints.cpp
--- a/Engine/Plugins/VirtualProduction/CameraCalibration/Source/CameraCalibrationEditor/Private/Calibrators/CameraNodalOffsetAlgoPoints.cpp
+++ b/Engine/Plugins/VirtualProduction/CameraCalibration/Source/CameraCalibrationEditor/Private/Calibrators/CameraNodalOffsetAlgoPoints.cpp
@@ -29,7 +29,9 @@
 #include "UI/SFilterableActorPicker.h"
 #include "ScopedTransaction.h"
 #include "UI/CameraCalibrationWidgetHelpers.h"
+#include "UObject/UObjectIterator.h"
 #include "Widgets/Input/SButton.h"
+#include "Widgets/Input/SComboBox.h"
 #include "Widgets/Views/SListView.h"
 #include "Widgets/SWidget.h"
 
diff --git a/Engine/Plugins/VirtualProduction/CameraCalibrationCore/Source/CameraCalibrationCore/Private/CameraCalibrationSettings.cpp b/Engine/Plugins/VirtualProduction/CameraCalibrationCore/Source/CameraCalibrationCore/Private/CameraCalibrationSettings.cpp
--- a/Engine/Plugins/VirtualProduction/CameraCalibrationCore/Source/CameraCalibrationCore/Private/CameraCalibrationSettings.cpp
+++ b/Engine/Plugins/VirtualProduction/CameraCalibrationCore/Source/CameraCalibrationCore/Private/CameraCalibrationSettings.cpp
@@ -3,6 +3,8 @@
 #include "CameraCalibrationSettings.h"
 #include "CameraCalibrationSubsystem.h"
 
+#include "Engine/Engine.h"
+#include "Misc/CoreDelegates.h"
 #include "SphericalLensDistortionModelHandler.h"
 
 UCameraCalibrationSettings::UCameraCalibrationSettings()
diff --git a/Engine/Plugins/VirtualProduction/CameraCalibrationCore/Source/CameraCalibrationCore/Private/LensDistortionModelHandlerBase.cpp b/Engine/Plugins/VirtualProduction/CameraCalibrationCore/Source/CameraCalibrationCore/Private/LensDistortionModelHandlerBase.cpp
--- a/Engine/Plugins/VirtualProduction/CameraCalibrationCore/Source/CameraCalibrationCore/Private/LensDistortionModelHandlerBase.cpp
+++ b/Engine/Plugins/VirtualProduction/CameraCalibrationCore/Source/CameraCalibrationCore/Private/LensDistortionModelHandlerBase.cpp
@@ -4,7 +4,9 @@
 
 #include "Algo/MaxElement.h"
 #include "CameraCalibrationSettings.h"
+#include "Engine/World.h"
 #include "Kismet/KismetRenderingLibrary.h"
+#include "MaterialShared.h"
 
 bool FLensDistortionState::operator==(const FLensDistortionState& Other) const
 {
diff --git a/Engine/Plugins/VirtualProduction/CameraCalibrationCore/Source/CameraCalibrationCore/Public/LensFile.h b/Engine/Plugins/VirtualProduction/CameraCalibrationCore/Source/CameraCalibrationCore/Public/LensFile.h
--- a/Engine/Plugins/VirtualProduction/CameraCalibrationCore/Source/CameraCalibrationCore/Public/LensFile.h
+++ b/Engine/Plugins/VirtualProduction/CameraCalibrationCore/Source/CameraCalibrationCore/Public/LensFile.h
@@ -25,6 +25,7 @@
 
 class UCineCameraComponent;
 class ULensDistortionModelHandlerBase;
+class UTextureRenderTarget2D;
 struct FBaseLensTable;
 
 
diff --git a/Engine/Plugins/VirtualProduction/CameraCalibrationCore/Source/CameraCalibrationCore/Public/Tables/STMapTable.h b/Engine/Plugins/VirtualProduction/CameraCalibrationCore/Source/CameraCalibrationCore/Public/Tables/STMapTable.h
--- a/Engine/Plugins/VirtualProduction/CameraCalibrationCore/Source/CameraCalibrationCore/Public/Tables/STMapTable.h
+++ b/Engine/Plugins/VirtualProduction/CameraCalibrationCore/Source/CameraCalibrationCore/Public/Tables/STMapTable.h
@@ -8,7 +8,7 @@
 
 #include "STMapTable.generated.h"
 
-
+class UTextureRenderTarget2D;
 
 /**
  * Derived data computed from parameters or stmap
diff --git a/Engine/Plugins/VirtualProduction/DMX/DMXEngine/Source/DMXEditor/Private/Customizations/DMXLibraryPortReferencesCustomization.cpp b/Engine/Plugins/VirtualProduction/DMX/DMXEngine/Source/DMXEditor/Private/Customizations/DMXLibraryPortReferencesCustomization.cpp
--- a/Engine/Plugins/VirtualProduction/DMX/DMXEngine/Source/DMXEditor/Private/Customizations/DMXLibraryPortReferencesCustomization.cpp
+++ b/Engine/Plugins/VirtualProduction/DMX/DMXEngine/Source/DMXEditor/Private/Customizations/DMXLibraryPortReferencesCustomization.cpp
@@ -23,6 +23,7 @@
 #include "Widgets/Layout/SBorder.h"
 #include "Widgets/Layout/SBox.h"
 #include "Widgets/Layout/SSeparator.h"
+#include "Widgets/SWidget.h"
 #include "Widgets/Text/STextBlock.h" 
 #include "Widgets/Layout/SWrapBox.h"
 
diff --git a/Engine/Plugins/VirtualProduction/DMX/DMXEngine/Source/DMXEditor/Private/Customizations/DMXLibraryPortReferencesCustomization.h b/Engine/Plugins/VirtualProduction/DMX/DMXEngine/Source/DMXEditor/Private/Customizations/DMXLibraryPortReferencesCustomization.h
--- a/Engine/Plugins/VirtualProduction/DMX/DMXEngine/Source/DMXEditor/Private/Customizations/DMXLibraryPortReferencesCustomization.h
+++ b/Engine/Plugins/VirtualProduction/DMX/DMXEngine/Source/DMXEditor/Private/Customizations/DMXLibraryPortReferencesCustomization.h
@@ -11,7 +11,7 @@ class FDetailWidgetRow;
 class IPropertyHandle;
 class IPropertyUtilities;
 class SBorder;
-
+class SWidget;
 
 /** Property type customization for the Library Port References struct*/
 class FDMXLibraryPortReferencesCustomization
diff --git a/Engine/Plugins/VirtualProduction/DMX/DMXEngine/Source/DMXEditor/Private/Customizations/DMXPixelMappingDistributionCustomization.h b/Engine/Plugins/VirtualProduction/DMX/DMXEngine/Source/DMXEditor/Private/Customizations/DMXPixelMappingDistributionCustomization.h
--- a/Engine/Plugins/VirtualProduction/DMX/DMXEngine/Source/DMXEditor/Private/Customizations/DMXPixelMappingDistributionCustomization.h
+++ b/Engine/Plugins/VirtualProduction/DMX/DMXEngine/Source/DMXEditor/Private/Customizations/DMXPixelMappingDistributionCustomization.h
@@ -6,6 +6,7 @@
 
 #include "IPropertyTypeCustomization.h"
 
+class FReply;
 
 /**
  * Customization for the DMXPixelMappingDistribution enum
diff --git a/Engine/Plugins/VirtualProduction/DMX/DMXEngine/Source/DMXEditor/Private/Widgets/FixturePatch/SDMXPatchedUniverse.cpp b/Engine/Plugins/VirtualProduction/DMX/DMXEngine/Source/DMXEditor/Private/Widgets/FixturePatch/SDMXPatchedUniverse.cpp
--- a/Engine/Plugins/VirtualProduction/DMX/DMXEngine/Source/DMXEditor/Private/Widgets/FixturePatch/SDMXPatchedUniverse.cpp
+++ b/Engine/Plugins/VirtualProduction/DMX/DMXEngine/Source/DMXEditor/Private/Widgets/FixturePatch/SDMXPatchedUniverse.cpp
@@ -7,6 +7,7 @@
 #include "DMXFixturePatchEditorDefinitions.h"
 #include "DMXFixturePatchNode.h"
 #include "DMXFixturePatchSharedData.h"
+#include "Framework/Application/SlateApplication.h"
 #include "SDMXChannelConnector.h"
 #include "SDMXFixturePatchFragment.h"
 #include "DragDrop/DMXEntityFixturePatchDragDropOp.h"
@@ -21,6 +22,7 @@
 #include "Framework/MultiBox/MultiBoxBuilder.h"
 #include "Styling/AppStyle.h"
 #include "SlateOptMacros.h"
+#include "Widgets/Layout/SBorder.h"
 #include "Widgets/Layout/SGridPanel.h"
 
 
diff --git a/Engine/Plugins/VirtualProduction/DMX/DMXEngine/Source/DMXEditor/Private/Widgets/FixtureType/SDMXFixtureFunctionEditor.cpp b/Engine/Plugins/VirtualProduction/DMX/DMXEngine/Source/DMXEditor/Private/Widgets/FixtureType/SDMXFixtureFunctionEditor.cpp
--- a/Engine/Plugins/VirtualProduction/DMX/DMXEngine/Source/DMXEditor/Private/Widgets/FixtureType/SDMXFixtureFunctionEditor.cpp
+++ b/Engine/Plugins/VirtualProduction/DMX/DMXEngine/Source/DMXEditor/Private/Widgets/FixtureType/SDMXFixtureFunctionEditor.cpp
@@ -10,6 +10,8 @@
 #include "PropertyEditorModule.h"
 #include "ScopedTransaction.h"
 #include "Modules/ModuleManager.h"
+#include "Widgets/Layout/SBox.h"
+#include "Widgets/Text/STextBlock.h"
 
 
 #define LOCTEXT_NAMESPACE "SDMXFixtureFunctionEditor"
diff --git a/Engine/Plugins/VirtualProduction/DMX/DMXEngine/Source/DMXEditor/Private/Widgets/FixtureType/SDMXFixtureTypeEditor.cpp b/Engine/Plugins/VirtualProduction/DMX/DMXEngine/Source/DMXEditor/Private/Widgets/FixtureType/SDMXFixtureTypeEditor.cpp
--- a/Engine/Plugins/VirtualProduction/DMX/DMXEngine/Source/DMXEditor/Private/Widgets/FixtureType/SDMXFixtureTypeEditor.cpp
+++ b/Engine/Plugins/VirtualProduction/DMX/DMXEngine/Source/DMXEditor/Private/Widgets/FixtureType/SDMXFixtureTypeEditor.cpp
@@ -5,6 +5,8 @@
 #include "DMXEditor.h"
 #include "DMXEditorUtils.h"
 #include "DMXFixtureTypeSharedData.h"
+#include "IDetailsView.h"
+#include "PropertyEditorModule.h"
 #include "SDMXFixtureFunctionEditor.h"
 #include "SDMXFixtureModeEditor.h"
 #include "SDMXFixtureTypeFunctionsEditor.h"
diff --git a/Engine/Plugins/VirtualProduction/DMX/DMXEngine/Source/DMXRuntime/Private/MVR/DMXMVRFixtureActorLibrary.cpp b/Engine/Plugins/VirtualProduction/DMX/DMXEngine/Source/DMXRuntime/Private/MVR/DMXMVRFixtureActorLibrary.cpp
--- a/Engine/Plugins/VirtualProduction/DMX/DMXEngine/Source/DMXRuntime/Private/MVR/DMXMVRFixtureActorLibrary.cpp
+++ b/Engine/Plugins/VirtualProduction/DMX/DMXEngine/Source/DMXRuntime/Private/MVR/DMXMVRFixtureActorLibrary.cpp
@@ -4,6 +4,7 @@
 
 #include "DMXMVRFixtureActorInterface.h"
 #include "DMXRuntimeLog.h"
+#include "Engine/BlueprintGeneratedClass.h"
 #include "Library/DMXEntityFixturePatch.h"
 
 #include "AssetRegistry/AssetRegistryModule.h"
@@ -11,6 +12,7 @@
 #include "Engine/SkeletalMesh.h"
 #include "Interfaces/IPluginManager.h"
 #include "Misc/ScopedSlowTask.h"
+#include "UObject/UObjectIterator.h"
 
 
 class FDMXMVRFixtureActorAssetHierarchy
diff --git a/Engine/Plugins/VirtualProduction/DMX/DMXPixelMapping/Source/DMXPixelMappingEditor/Private/Customizations/DMXPixelMappingDetailCustomization_FixtureGroupItem.cpp b/Engine/Plugins/VirtualProduction/DMX/DMXPixelMapping/Source/DMXPixelMappingEditor/Private/Customizations/DMXPixelMappingDetailCustomization_FixtureGroupItem.cpp
--- a/Engine/Plugins/VirtualProduction/DMX/DMXPixelMapping/Source/DMXPixelMappingEditor/Private/Customizations/DMXPixelMappingDetailCustomization_FixtureGroupItem.cpp
+++ b/Engine/Plugins/VirtualProduction/DMX/DMXPixelMapping/Source/DMXPixelMappingEditor/Private/Customizations/DMXPixelMappingDetailCustomization_FixtureGroupItem.cpp
@@ -15,6 +15,10 @@
 #include "PropertyEditorModule.h"
 #include "Layout/Visibility.h"
 #include "Modules/ModuleManager.h"
+#include "Widgets/Layout/SBox.h"
+#include "Widgets/Text/STextBlock.h"
+#include "Widgets/Views/SListView.h"
+#include "Widgets/Views/STableRow.h"
 
 #define LOCTEXT_NAMESPACE "FixtureGroupItem"
 
diff --git a/Engine/Plugins/VirtualProduction/DMX/DMXPixelMapping/Source/DMXPixelMappingEditor/Private/Customizations/DMXPixelMappingDetailCustomization_Matrix.cpp b/Engine/Plugins/VirtualProduction/DMX/DMXPixelMapping/Source/DMXPixelMappingEditor/Private/Customizations/DMXPixelMappingDetailCustomization_Matrix.cpp
--- a/Engine/Plugins/VirtualProduction/DMX/DMXPixelMapping/Source/DMXPixelMappingEditor/Private/Customizations/DMXPixelMappingDetailCustomization_Matrix.cpp
+++ b/Engine/Plugins/VirtualProduction/DMX/DMXPixelMapping/Source/DMXPixelMappingEditor/Private/Customizations/DMXPixelMappingDetailCustomization_Matrix.cpp
@@ -21,6 +21,10 @@
 #include "Layout/Visibility.h"
 #include "Misc/Attribute.h"
 #include "Modules/ModuleManager.h"
+#include "Widgets/Layout/SBox.h"
+#include "Widgets/Text/STextBlock.h"
+#include "Widgets/Views/SListView.h"
+#include "Widgets/Views/STableRow.h"
 
 
 #define LOCTEXT_NAMESPACE "DMXPixelMappingDetailCustomization_Matrix"
diff --git a/Engine/Plugins/VirtualProduction/DMX/DMXPixelMapping/Source/DMXPixelMappingEditor/Private/DMXPixelMappingThumbnailRendering.cpp b/Engine/Plugins/VirtualProduction/DMX/DMXPixelMapping/Source/DMXPixelMappingEditor/Private/DMXPixelMappingThumbnailRendering.cpp
--- a/Engine/Plugins/VirtualProduction/DMX/DMXPixelMapping/Source/DMXPixelMappingEditor/Private/DMXPixelMappingThumbnailRendering.cpp
+++ b/Engine/Plugins/VirtualProduction/DMX/DMXPixelMapping/Source/DMXPixelMappingEditor/Private/DMXPixelMappingThumbnailRendering.cpp
@@ -3,6 +3,9 @@
 #include "DMXPixelMappingThumbnailRendering.h"
 #include "CanvasTypes.h"
 #include "DMXPixelMapping.h"
+#include "Engine/Texture.h"
+#include "SceneManagement.h"
+#include "TextureResource.h"
 
 bool UDMXPixelMappingThumbnailRendering::CanVisualizeAsset(UObject* Object)
 {
diff --git a/Engine/Plugins/VirtualProduction/DMX/DMXPixelMapping/Source/DMXPixelMappingEditor/Private/Viewports/DMXPixelMappingPreviewViewportClient.h b/Engine/Plugins/VirtualProduction/DMX/DMXPixelMapping/Source/DMXPixelMappingEditor/Private/Viewports/DMXPixelMappingPreviewViewportClient.h
--- a/Engine/Plugins/VirtualProduction/DMX/DMXPixelMapping/Source/DMXPixelMappingEditor/Private/Viewports/DMXPixelMappingPreviewViewportClient.h
+++ b/Engine/Plugins/VirtualProduction/DMX/DMXPixelMapping/Source/DMXPixelMappingEditor/Private/Viewports/DMXPixelMappingPreviewViewportClient.h
@@ -4,6 +4,7 @@
 
 #include "CoreMinimal.h"
 #include "UnrealClient.h"
+#include "ViewportClient.h"
 
 class SDMXPixelMappingPreviewViewport;
 class FDMXPixelMappingToolkit;
diff --git a/Engine/Plugins/VirtualProduction/DMX/DMXPixelMapping/Source/DMXPixelMappingEditor/Private/Views/SDMXPixelMappingPreviewView.cpp b/Engine/Plugins/VirtualProduction/DMX/DMXPixelMapping/Source/DMXPixelMappingEditor/Private/Views/SDMXPixelMappingPreviewView.cpp
--- a/Engine/Plugins/VirtualProduction/DMX/DMXPixelMapping/Source/DMXPixelMappingEditor/Private/Views/SDMXPixelMappingPreviewView.cpp
+++ b/Engine/Plugins/VirtualProduction/DMX/DMXPixelMapping/Source/DMXPixelMappingEditor/Private/Views/SDMXPixelMappingPreviewView.cpp
@@ -3,7 +3,9 @@
 #include "Views/SDMXPixelMappingPreviewView.h"
 
 #include "DMXPixelMappingEditorStyle.h"
+#include "Framework/Application/SlateApplication.h"
 #include "Toolkits/DMXPixelMappingToolkit.h"
+#include "Widgets/Layout/SBox.h"
 #include "Widgets/SDMXPixelMappingPreviewViewport.h"
 #include "Widgets/SDMXPixelMappingZoomPan.h"
 #include "Components/DMXPixelMappingRendererComponent.h"
diff --git a/Engine/Plugins/VirtualProduction/DMX/DMXPixelMapping/Source/DMXPixelMappingEditor/Private/Widgets/SDMXPixelMappingOutputComponent.cpp b/Engine/Plugins/VirtualProduction/DMX/DMXPixelMapping/Source/DMXPixelMappingEditor/Private/Widgets/SDMXPixelMappingOutputComponent.cpp
--- a/Engine/Plugins/VirtualProduction/DMX/DMXPixelMapping/Source/DMXPixelMappingEditor/Private/Widgets/SDMXPixelMappingOutputComponent.cpp
+++ b/Engine/Plugins/VirtualProduction/DMX/DMXPixelMapping/Source/DMXPixelMappingEditor/Private/Widgets/SDMXPixelMappingOutputComponent.cpp
@@ -15,9 +15,11 @@
 #include "PropertyEditorModule.h"
 #include "TimerManager.h"
 #include "Modules/ModuleManager.h"
+#include "Widgets/Layout/SBorder.h"
 #include "Widgets/Layout/SBox.h"
 #include "Widgets/Layout/SScaleBox.h"
 #include "Widgets/Layout/SUniformGridPanel.h"
+#include "Widgets/Text/STextBlock.h"
 
 
 #define LOCTEXT_NAMESPACE "SDMXPixelMappingOutputComponent"
diff --git a/Engine/Plugins/VirtualProduction/DMX/DMXPixelMapping/Source/DMXPixelMappingEditor/Private/Widgets/SDMXPixelMappingOutputComponent.h b/Engine/Plugins/VirtualProduction/DMX/DMXPixelMapping/Source/DMXPixelMappingEditor/Private/Widgets/SDMXPixelMappingOutputComponent.h
--- a/Engine/Plugins/VirtualProduction/DMX/DMXPixelMapping/Source/DMXPixelMappingEditor/Private/Widgets/SDMXPixelMappingOutputComponent.h
+++ b/Engine/Plugins/VirtualProduction/DMX/DMXPixelMapping/Source/DMXPixelMappingEditor/Private/Widgets/SDMXPixelMappingOutputComponent.h
@@ -4,6 +4,7 @@
 
 #include "CoreMinimal.h"
 
+#include "Styling/SlateBrush.h"
 #include "Widgets/SCompoundWidget.h"
 #include "Widgets/Layout/SConstraintCanvas.h"
 
diff --git a/Engine/Plugins/VirtualProduction/DMX/DMXPixelMapping/Source/DMXPixelMappingEditorWidgets/Private/DMXPixelMappingComponentWidget.cpp b/Engine/Plugins/VirtualProduction/DMX/DMXPixelMapping/Source/DMXPixelMappingEditorWidgets/Private/DMXPixelMappingComponentWidget.cpp
--- a/Engine/Plugins/VirtualProduction/DMX/DMXPixelMapping/Source/DMXPixelMappingEditorWidgets/Private/DMXPixelMappingComponentWidget.cpp
+++ b/Engine/Plugins/VirtualProduction/DMX/DMXPixelMapping/Source/DMXPixelMappingEditorWidgets/Private/DMXPixelMappingComponentWidget.cpp
@@ -4,6 +4,7 @@
 
 #include "SDMXPixelMappingComponentBox.h"
 #include "SDMXPixelMappingComponentLabel.h"
+#include "Widgets/SCompoundWidget.h"
 
 
 PRAGMA_DISABLE_DEPRECATION_WARNINGS // The whole class and several child widgets are deprecated 5.1
diff --git a/Engine/Plugins/VirtualProduction/DMX/DMXPixelMapping/Source/DMXPixelMappingEditorWidgets/Private/SDMXPixelMappingEditorWidgets.cpp b/Engine/Plugins/VirtualProduction/DMX/DMXPixelMapping/Source/DMXPixelMappingEditorWidgets/Private/SDMXPixelMappingEditorWidgets.cpp
--- a/Engine/Plugins/VirtualProduction/DMX/DMXPixelMapping/Source/DMXPixelMappingEditorWidgets/Private/SDMXPixelMappingEditorWidgets.cpp
+++ b/Engine/Plugins/VirtualProduction/DMX/DMXPixelMapping/Source/DMXPixelMappingEditorWidgets/Private/SDMXPixelMappingEditorWidgets.cpp
@@ -10,6 +10,7 @@
 #include "Styling/AppStyle.h"
 #include "Components/HorizontalBox.h"
 #include "Components/HorizontalBoxSlot.h"
+#include "Widgets/Layout/SBox.h"
 #include "Widgets/SOverlay.h"
 #include "Widgets/Images/SImage.h"
 #include "Widgets/Layout/SScaleBox.h"
diff --git a/Engine/Plugins/VirtualProduction/DMX/DMXPixelMapping/Source/DMXPixelMappingEditorWidgets/Private/SDMXPixelMappingScreenComponentBox.cpp b/Engine/Plugins/VirtualProduction/DMX/DMXPixelMapping/Source/DMXPixelMappingEditorWidgets/Private/SDMXPixelMappingScreenComponentBox.cpp
--- a/Engine/Plugins/VirtualProduction/DMX/DMXPixelMapping/Source/DMXPixelMappingEditorWidgets/Private/SDMXPixelMappingScreenComponentBox.cpp
+++ b/Engine/Plugins/VirtualProduction/DMX/DMXPixelMapping/Source/DMXPixelMappingEditorWidgets/Private/SDMXPixelMappingScreenComponentBox.cpp
@@ -5,6 +5,7 @@
 #include "SDMXPixelMappingEditorWidgets.h"
 
 #include "Widgets/Layout/SBox.h"
+#include "Widgets/SWidget.h"
 
 
 PRAGMA_DISABLE_DEPRECATION_WARNINGS // The whole class and several child widgets are deprecated 5.1
diff --git a/Engine/Plugins/VirtualProduction/DMX/DMXProtocol/Source/DMXProtocolEditor/Private/DetailsCustomizations/DMXOutputPortDestinationAddressCustomization.h b/Engine/Plugins/VirtualProduction/DMX/DMXProtocol/Source/DMXProtocolEditor/Private/DetailsCustomizations/DMXOutputPortDestinationAddressCustomization.h
--- a/Engine/Plugins/VirtualProduction/DMX/DMXProtocol/Source/DMXProtocolEditor/Private/DetailsCustomizations/DMXOutputPortDestinationAddressCustomization.h
+++ b/Engine/Plugins/VirtualProduction/DMX/DMXProtocol/Source/DMXProtocolEditor/Private/DetailsCustomizations/DMXOutputPortDestinationAddressCustomization.h
@@ -5,9 +5,9 @@
 #include "CoreMinimal.h"
 #include "IPropertyTypeCustomization.h"
 
-class SDMXIPAddressEditWidget;
-
+class FReply;
 class IPropertyHandle;
+class SDMXIPAddressEditWidget;
 
 
 /** Customization for the DMXOutputPortConfigCustomization struct */
diff --git a/Engine/Plugins/VirtualProduction/EpicStageApp/Source/EpicStageApp/Private/StageAppRouteHandler.cpp b/Engine/Plugins/VirtualProduction/EpicStageApp/Source/EpicStageApp/Private/StageAppRouteHandler.cpp
--- a/Engine/Plugins/VirtualProduction/EpicStageApp/Source/EpicStageApp/Private/StageAppRouteHandler.cpp
+++ b/Engine/Plugins/VirtualProduction/EpicStageApp/Source/EpicStageApp/Private/StageAppRouteHandler.cpp
@@ -8,6 +8,7 @@
 #include "DisplayClusterLightCardActor.h"
 #include "DisplayClusterLightCardEditorHelper.h"
 #include "DisplayClusterRootActor.h"
+#include "Engine/Canvas.h"
 #include "Engine/TextureRenderTarget2D.h"
 #include "IDisplayClusterScenePreview.h"
 #include "IImageWrapperModule.h"
@@ -16,6 +17,7 @@
 #include "Modules/ModuleManager.h"
 #include "RemoteControlWebsocketRoute.h"
 #include "StageAppResponse.h"
+#include "TextureResource.h"
 #include "WebRemoteControlUtils.h"
 #include "StageActor/DisplayClusterStageActorTemplate.h"
 #include "StageActor/DisplayClusterWeakStageActorPtr.h"
diff --git a/Engine/Plugins/VirtualProduction/LevelSnapshots/Source/LevelSnapshotsEditor/Private/Util/TakeSnapshotUtil.cpp b/Engine/Plugins/VirtualProduction/LevelSnapshots/Source/LevelSnapshotsEditor/Private/Util/TakeSnapshotUtil.cpp
--- a/Engine/Plugins/VirtualProduction/LevelSnapshots/Source/LevelSnapshotsEditor/Private/Util/TakeSnapshotUtil.cpp
+++ b/Engine/Plugins/VirtualProduction/LevelSnapshots/Source/LevelSnapshotsEditor/Private/Util/TakeSnapshotUtil.cpp
@@ -2,6 +2,7 @@
 
 #include "Util/TakeSnapshotUtil.h"
 
+#include "Components/Viewport.h"
 #include "Data/LevelSnapshotsEditorData.h"
 #include "LevelSnapshot.h"
 #include "LevelSnapshotsEditorSettings.h"
diff --git a/Engine/Plugins/VirtualProduction/LiveLinkCamera/Source/LiveLinkCamera/Private/LiveLinkCameraController.cpp b/Engine/Plugins/VirtualProduction/LiveLinkCamera/Source/LiveLinkCamera/Private/LiveLinkCameraController.cpp
--- a/Engine/Plugins/VirtualProduction/LiveLinkCamera/Source/LiveLinkCamera/Private/LiveLinkCameraController.cpp
+++ b/Engine/Plugins/VirtualProduction/LiveLinkCamera/Source/LiveLinkCamera/Private/LiveLinkCameraController.cpp
@@ -6,6 +6,7 @@
 #include "CameraCalibrationSubsystem.h"
 #include "CineCameraComponent.h"
 #include "Controllers/LiveLinkTransformController.h"
+#include "GameFramework/Actor.h"
 #include "LensComponent.h"
 #include "LensFile.h"
 #include "LiveLinkLog.h"
diff --git a/Engine/Plugins/VirtualProduction/RemoteControl/Source/RemoteControlLogic/Private/Behaviour/Builtin/Path/RCSetAssetByPathBehaviour.cpp b/Engine/Plugins/VirtualProduction/RemoteControl/Source/RemoteControlLogic/Private/Behaviour/Builtin/Path/RCSetAssetByPathBehaviour.cpp
--- a/Engine/Plugins/VirtualProduction/RemoteControl/Source/RemoteControlLogic/Private/Behaviour/Builtin/Path/RCSetAssetByPathBehaviour.cpp
+++ b/Engine/Plugins/VirtualProduction/RemoteControl/Source/RemoteControlLogic/Private/Behaviour/Builtin/Path/RCSetAssetByPathBehaviour.cpp
@@ -8,11 +8,13 @@
 #include "Components/StaticMeshComponent.h"
 #include "Controller/RCController.h"
 #include "Engine/StaticMesh.h"
+#include "Engine/Texture2D.h"
 #include "Engine/Texture2DDynamic.h"
 #include "Engine/World.h"
 #include "IRemoteControlModule.h"
 #include "IRemoteControlPropertyHandle.h"
 #include "Kismet/KismetRenderingLibrary.h"
+#include "Materials/Material.h"
 #include "Materials/MaterialInstance.h"
 #include "RCVirtualProperty.h"
 #include "RCVirtualPropertyContainer.h"
diff --git a/Engine/Plugins/VirtualProduction/RemoteControl/Source/RemoteControlUI/Private/RemoteControlUIModule.h b/Engine/Plugins/VirtualProduction/RemoteControl/Source/RemoteControlUI/Private/RemoteControlUIModule.h
--- a/Engine/Plugins/VirtualProduction/RemoteControl/Source/RemoteControlUI/Private/RemoteControlUIModule.h
+++ b/Engine/Plugins/VirtualProduction/RemoteControl/Source/RemoteControlUI/Private/RemoteControlUIModule.h
@@ -12,6 +12,7 @@
 
 class IToolkitHost;
 class SRemoteControlPanel;
+class UMeshComponent;
 class URemoteControlPreset;
 
 /**
diff --git a/Engine/Plugins/VirtualProduction/RemoteControl/Source/RemoteControlUI/Private/UI/RCUIHelpers.cpp b/Engine/Plugins/VirtualProduction/RemoteControl/Source/RemoteControlUI/Private/UI/RCUIHelpers.cpp
--- a/Engine/Plugins/VirtualProduction/RemoteControl/Source/RemoteControlUI/Private/UI/RCUIHelpers.cpp
+++ b/Engine/Plugins/VirtualProduction/RemoteControl/Source/RemoteControlUI/Private/UI/RCUIHelpers.cpp
@@ -15,6 +15,8 @@
 #include "PropertyHandle.h"
 #include "RCVirtualProperty.h"
 #include "TimerManager.h"
+#include "Widgets/Layout/SBorder.h"
+#include "Widgets/Layout/SBox.h"
 
 FLinearColor UE::RCUIHelpers::GetFieldClassTypeColor(const FProperty* InProperty)
 {
diff --git a/Engine/Plugins/VirtualProduction/Switchboard/Source/SwitchboardEditor/Private/DisplayClusterRootActorReferenceDetailCustomization.cpp b/Engine/Plugins/VirtualProduction/Switchboard/Source/SwitchboardEditor/Private/DisplayClusterRootActorReferenceDetailCustomization.cpp
--- a/Engine/Plugins/VirtualProduction/Switchboard/Source/SwitchboardEditor/Private/DisplayClusterRootActorReferenceDetailCustomization.cpp
+++ b/Engine/Plugins/VirtualProduction/Switchboard/Source/SwitchboardEditor/Private/DisplayClusterRootActorReferenceDetailCustomization.cpp
@@ -2,8 +2,10 @@
 
 #include "DisplayClusterRootActorReferenceDetailCustomization.h"
 
+#include "AssetRegistry/AssetData.h"
 #include "DetailLayoutBuilder.h"
 #include "DetailWidgetRow.h"
+#include "GameFramework/Actor.h"
 #include "IPropertyUtilities.h"
 #include "IPropertyTypeCustomization.h"
 #include "PropertyCustomizationHelpers.h"
diff --git a/Engine/Plugins/VirtualProduction/TimedDataMonitor/Source/TimedDataMonitor/Private/TimedDataMonitorSubsystem.cpp b/Engine/Plugins/VirtualProduction/TimedDataMonitor/Source/TimedDataMonitor/Private/TimedDataMonitorSubsystem.cpp
--- a/Engine/Plugins/VirtualProduction/TimedDataMonitor/Source/TimedDataMonitor/Private/TimedDataMonitorSubsystem.cpp
+++ b/Engine/Plugins/VirtualProduction/TimedDataMonitor/Source/TimedDataMonitor/Private/TimedDataMonitorSubsystem.cpp
@@ -9,6 +9,7 @@
 #include "ITimeManagementModule.h"
 #include "LatentActions.h"
 #include "Misc/App.h"
+#include "Misc/CoreDelegates.h"
 #include "Misc/QualifiedFrameTime.h"
 #include "Stats/Stats2.h"
 #include "TimedDataInputCollection.h"
diff --git a/Engine/Plugins/VirtualProduction/TimedDataMonitor/Source/TimedDataMonitorEditor/Private/STimedDataListView.cpp b/Engine/Plugins/VirtualProduction/TimedDataMonitor/Source/TimedDataMonitorEditor/Private/STimedDataListView.cpp
--- a/Engine/Plugins/VirtualProduction/TimedDataMonitor/Source/TimedDataMonitorEditor/Private/STimedDataListView.cpp
+++ b/Engine/Plugins/VirtualProduction/TimedDataMonitor/Source/TimedDataMonitorEditor/Private/STimedDataListView.cpp
@@ -4,6 +4,7 @@
 
 #include "EditorFontGlyphs.h"
 #include "Engine/Engine.h"
+#include "Framework/Application/SlateUser.h"
 #include "Framework/MultiBox/MultiBoxBuilder.h"
 #include "ITimedDataInput.h"
 #include "Misc/App.h"
diff --git a/Engine/Plugins/VirtualProduction/TimedDataMonitor/Source/TimedDataMonitorEditor/Private/STimedDataMonitorBufferVisualizer.cpp b/Engine/Plugins/VirtualProduction/TimedDataMonitor/Source/TimedDataMonitorEditor/Private/STimedDataMonitorBufferVisualizer.cpp
--- a/Engine/Plugins/VirtualProduction/TimedDataMonitor/Source/TimedDataMonitorEditor/Private/STimedDataMonitorBufferVisualizer.cpp
+++ b/Engine/Plugins/VirtualProduction/TimedDataMonitor/Source/TimedDataMonitorEditor/Private/STimedDataMonitorBufferVisualizer.cpp
@@ -11,6 +11,8 @@
 #include "TimedDataMonitorSubsystem.h"
 #include "Widgets/Layout/SBox.h"
 #include "Widgets/Layout/SScrollBorder.h"
+#include "Widgets/Text/STextBlock.h"
+#include "Widgets/Views/SListView.h"
 #include "Widgets/Views/STableRow.h"
 
 #define LOCTEXT_NAMESPACE "TimedDataBufferVisualizer"
diff --git a/Engine/Plugins/VirtualProduction/TimedDataMonitor/Source/TimedDataMonitorEditor/Private/STimedDataMonitorPanel.cpp b/Engine/Plugins/VirtualProduction/TimedDataMonitor/Source/TimedDataMonitorEditor/Private/STimedDataMonitorPanel.cpp
--- a/Engine/Plugins/VirtualProduction/TimedDataMonitor/Source/TimedDataMonitorEditor/Private/STimedDataMonitorPanel.cpp
+++ b/Engine/Plugins/VirtualProduction/TimedDataMonitor/Source/TimedDataMonitorEditor/Private/STimedDataMonitorPanel.cpp
@@ -5,6 +5,7 @@
 #include "Editor.h"
 #include "Editor/EditorPerformanceSettings.h"
 #include "Engine/Engine.h"
+#include "Framework/Application/SlateApplication.h"
 #include "IMessageLogListing.h"
 #include "ISettingsModule.h"
 #include "ITimedDataInput.h"
@@ -32,6 +33,7 @@
 #include "Widgets/Images/SThrobber.h"
 #include "Widgets/Input/SButton.h"
 #include "Widgets/Input/SCheckBox.h"
+#include "Widgets/Input/SComboButton.h"
 #include "Widgets/Layout/SBorder.h"
 #include "Widgets/Layout/SScrollBox.h"
 #include "Widgets/Layout/SSpacer.h"
diff --git a/Engine/Source/Runtime/Engine/Public/TraceFilter.h b/Engine/Source/Runtime/Engine/Public/TraceFilter.h
--- a/Engine/Source/Runtime/Engine/Public/TraceFilter.h
+++ b/Engine/Source/Runtime/Engine/Public/TraceFilter.h
@@ -12,6 +12,10 @@
 #include "Templates/UnrealTemplate.h"
 #include "UObject/Object.h"
 
+#if PLATFORM_LINUX
+#include "Components/SkeletalMeshComponent.h" // Needed because AnimInstance::GetOwningComponent. (this only fails on linux)
+#endif
+
 
 #if UE_TRACE_ENABLED && !IS_PROGRAM && !UE_BUILD_SHIPPING
 #define TRACE_FILTERING_ENABLED 1
diff --git a/Engine/Source/Runtime/GameMenuBuilder/Private/SGameMenuPageWidget.cpp b/Engine/Source/Runtime/GameMenuBuilder/Private/SGameMenuPageWidget.cpp
--- a/Engine/Source/Runtime/GameMenuBuilder/Private/SGameMenuPageWidget.cpp
+++ b/Engine/Source/Runtime/GameMenuBuilder/Private/SGameMenuPageWidget.cpp
@@ -1,6 +1,7 @@
 // Copyright Epic Games, Inc. All Rights Reserved.
 
 #include "SGameMenuPageWidget.h"
+#include "Components/Viewport.h"
 #include "Widgets/SBoxPanel.h"
 #include "Widgets/SOverlay.h"
 #include "Framework/Application/SlateApplication.h"
diff --git a/Engine/Source/Runtime/SlateNullRenderer/Private/SlateNullRenderer.h b/Engine/Source/Runtime/SlateNullRenderer/Private/SlateNullRenderer.h
--- a/Engine/Source/Runtime/SlateNullRenderer/Private/SlateNullRenderer.h
+++ b/Engine/Source/Runtime/SlateNullRenderer/Private/SlateNullRenderer.h
@@ -7,6 +7,8 @@
 #include "Rendering/SlateRenderer.h"
 
 class FSlateDrawBuffer;
+class FSlateFontServices;
+class FSlateShaderResourceManager;
 class FSlateUpdatableTexture;
 class SWindow;
 struct Rect;
