Add verbose logging to remote control rebinding

diff --git a/Engine/Plugins/VirtualProduction/RemoteControl/Source/RemoteControl/Private/RemoteControlPresetRebindingManager.cpp b/Engine/Plugins/VirtualProduction/RemoteControl/Source/RemoteControl/Private/RemoteControlPresetRebindingManager.cpp
--- a/Engine/Plugins/VirtualProduction/RemoteControl/Source/RemoteControl/Private/RemoteControlPresetRebindingManager.cpp
+++ b/Engine/Plugins/VirtualProduction/RemoteControl/Source/RemoteControl/Private/RemoteControlPresetRebindingManager.cpp
@@ -11,6 +11,7 @@
 #include "Engine/Engine.h"
 #include "EngineUtils.h"
 #include "Misc/Char.h"
+#include "IRemoteControlModule.h"
 #include "RemoteControlBinding.h"
 #include "RemoteControlEntity.h"
 #include "UObject/SoftObjectPath.h"
@@ -365,16 +366,22 @@ void FRemoteControlPresetRebindingManager::Rebind_NewAlgo(URemoteControlPreset*
 			{
 				TArray<UActorComponent*> Components;
 				Cast<AActor>(InActorMatch)->GetComponents(InBinding->BindingContext.SupportedClass.LoadSynchronous(), Components);
-				
+
+				UE_LOG(LogRemoteControl, Verbose, TEXT("Potential Subobject candidates:"));
 				for (UActorComponent* Component : Components)
 				{
+					UE_LOG(LogRemoteControl, Verbose, TEXT("\t%s"), *Component->GetName());
+
 					if (TryRebindPair(InBinding, Component->GetOwner(), Component))
 					{
+						UE_LOG(LogRemoteControl, Verbose, TEXT("Found a matching subobject (%s) with a matching class (%s) on actor %s"), *Component->GetName(), *Component->GetClass()->GetName(), *InActorMatch->GetName());
 						return true;
 					}
 				}
 			}
 		}
+
+		UE_LOG(LogRemoteControl, Verbose, TEXT("Could not find an object to rebind using the binding class."));
 		return false;
 	};
 
@@ -401,12 +408,16 @@ void FRemoteControlPresetRebindingManager::Rebind_NewAlgo(URemoteControlPreset*
 
 	auto TryRebindSubObject = [&](URemoteControlLevelDependantBinding* BindingToRebind, UObject* PotentialMatch)
 	{
+		UE_LOG(LogRemoteControl, Verbose, TEXT("Attempting to find subobject %s on object %s."), *BindingToRebind->BindingContext.ComponentName.ToString(), *PotentialMatch->GetName());
+		
 		if (UObject* TargetSubObject = GetSubObjectBasedOnName(BindingToRebind, PotentialMatch))
 		{
+			UE_LOG(LogRemoteControl, Verbose, TEXT("Found a matching subobject (%s) with a matching name on actor %s"), *TargetSubObject->GetName(), *PotentialMatch->GetName());
 			return TryRebindPair(BindingToRebind, TargetSubObject->GetTypedOuter<AActor>(), TargetSubObject);
 		}
 		else
 		{
+			UE_LOG(LogRemoteControl, Verbose, TEXT("Could not find a subobject using its name, trying to find one with a matching class (%s)"), *BindingToRebind->BindingContext.OwnerActorClass->GetClass()->GetName());
 			return RebindComponentBasedOnClass(BindingToRebind, Cast<AActor>(PotentialMatch));
 		}
 	};
@@ -415,12 +426,13 @@ void FRemoteControlPresetRebindingManager::Rebind_NewAlgo(URemoteControlPreset*
 	{
 		if (UObject* Match = FindByName(ObjectsWithSupportedClass, InvalidBinding->BindingContext.OwnerActorName))
 		{
-			// Found an owner object with matching name. Try to find a component under it with a matching name.
+			UE_LOG(LogRemoteControl, Verbose, TEXT("Found owner with a matching Owner Actor Name. (%s)"), *Match->GetName());
 			return TryRebindSubObject(InvalidBinding, Match);
 		}
 		else
 		{
 			// Could not find an actor with the same name as the initial binding, rely only on class instead.
+			UE_LOG(LogRemoteControl, Verbose, TEXT("Could not find an actor with the same name, defaulting to the first actor that supports the binding."));
 			for (UObject* Object : ObjectsWithSupportedClass)
 			{
 				if (TryRebindSubObject(InvalidBinding, Object))
@@ -447,6 +459,17 @@ void FRemoteControlPresetRebindingManager::Rebind_NewAlgo(URemoteControlPreset*
 		
 		if (GetPotentialActors(InvalidBinding, ObjectsWithSupportedClass))
 		{
+			UE_LOG(LogRemoteControl, Verbose, TEXT("Potential Actor Matches: "));
+			for (AActor* Actor : ObjectsWithSupportedClass)
+			{
+#if WITH_EDITOR
+				UE_LOG(LogRemoteControl, Verbose, TEXT("\t%s (%s)"), *Actor->GetActorLabel(), *Actor->GetName());
+#else
+				UE_LOG(LogRemoteControl, Verbose, TEXT("\t%s"), *Actor->GetName());
+#endif
+
+			}
+
 			if (BindingContext.HasValidSubObjectPath() || BindingContext.HasValidComponentName())
 			{
 				return FindSubObjectToRebind(InvalidBinding, ObjectsWithSupportedClass);
@@ -468,14 +491,16 @@ void FRemoteControlPresetRebindingManager::Rebind_NewAlgo(URemoteControlPreset*
 
 	auto TryNDisplayRebind = [&](URemoteControlLevelDependantBinding* InvalidBinding)
 	{
-		// Used when the preset does not contain the full component path, we try to find an ndisplay viewport on a new object.
+		UE_LOG(LogRemoteControl, Verbose, TEXT("Binding did not contain a full component path, checking if it's a ndisplay actor."));
 		if (UBlueprintGeneratedClass* BPClass = Cast <UBlueprintGeneratedClass>(InvalidBinding->BindingContext.OwnerActorClass.Get()))
 		{
 			if (BPClass->GetSuperStruct() && BPClass->GetSuperStruct()->GetName() == TEXT("DisplayClusterRootActor"))
 			{
+				UE_LOG(LogRemoteControl, Verbose, TEXT("Binding was on a ndisplay actor."));
 				TArray<AActor*> ObjectsWithSupportedClass;
 				if (!GetPotentialActors(InvalidBinding, ObjectsWithSupportedClass))
 				{
+					UE_LOG(LogRemoteControl, Verbose, TEXT("No potential ndisplay actor to rebind to."));
 					return;
 				}
 
@@ -485,6 +510,7 @@ void FRemoteControlPresetRebindingManager::Rebind_NewAlgo(URemoteControlPreset*
 				FString LastObjectPath = InvalidBinding->GetLastBoundObject().ToSoftObjectPath().ToString();
 				if (LastObjectPath.IsEmpty())
 				{
+					UE_LOG(LogRemoteControl, Verbose, TEXT("Could not find the last bound object path on the binding."));
 					return;
 				}
 
@@ -499,9 +525,9 @@ void FRemoteControlPresetRebindingManager::Rebind_NewAlgo(URemoteControlPreset*
 
 				for (AActor* Actor : ObjectsWithSupportedClass)
 				{
-					FTopLevelAssetPath AssetPath = FTopLevelAssetPath(Actor->GetPathName() + SubObjectSubPath);
 					if (UObject* SubObject = FindObject<UObject>(Actor, *SubObjectSubPath))
 					{
+						UE_LOG(LogRemoteControl, Verbose, TEXT("Found object %s that contains subobject %s"), *Actor->GetName(), *SubObjectSubPath);
 						TryRebindPair(InvalidBinding, Actor, SubObject);
 					}
 				}
@@ -509,12 +535,23 @@ void FRemoteControlPresetRebindingManager::Rebind_NewAlgo(URemoteControlPreset*
 		}
 	};
 
+	UE_LOG(LogRemoteControl, Verbose, TEXT("============================== Rebinding All for preset %s ==============================="), *Preset->GetName());
+
 	for (URemoteControlLevelDependantBinding* InvalidBinding : InvalidBindings)
 	{
+		UE_LOG(LogRemoteControl, Verbose, TEXT("Rebinding Invalid Binding"));
+		UE_LOG(LogRemoteControl, Verbose, TEXT("\tOwnerActorName: %s"), *InvalidBinding->BindingContext.OwnerActorName.ToString());
+		UE_LOG(LogRemoteControl, Verbose, TEXT("\tOwnerActorClass: %s"), *InvalidBinding->BindingContext.OwnerActorClass.ToString());
+		UE_LOG(LogRemoteControl, Verbose, TEXT("\tSupportedClass: %s"), *InvalidBinding->BindingContext.SupportedClass.ToString());
+		UE_LOG(LogRemoteControl, Verbose, TEXT("\tComponentName: %s"), *InvalidBinding->BindingContext.ComponentName.ToString());
+		UE_LOG(LogRemoteControl, Verbose, TEXT("\tSubObjectPath: %s"), *InvalidBinding->BindingContext.SubObjectPath);
+		
 		if (!Rebind(InvalidBinding))
 		{
 			TryNDisplayRebind(InvalidBinding);
 		}
+
+		UE_LOG(LogRemoteControl, Verbose, TEXT("==================================================================="));
 	}
 
 	Preset->RemoveUnusedBindings();
