Fix CIS NoPCH unity error in RemoteControlPreset

diff --git a/Engine/Plugins/VirtualProduction/RemoteControl/Source/RemoteControl/Private/RemoteControlPreset.cpp b/Engine/Plugins/VirtualProduction/RemoteControl/Source/RemoteControl/Private/RemoteControlPreset.cpp
--- a/Engine/Plugins/VirtualProduction/RemoteControl/Source/RemoteControl/Private/RemoteControlPreset.cpp
+++ b/Engine/Plugins/VirtualProduction/RemoteControl/Source/RemoteControl/Private/RemoteControlPreset.cpp
@@ -8,6 +8,8 @@
 #include "Components/ActorComponent.h"
 #include "Components/StaticMeshComponent.h"
 #include "Engine/BlueprintGeneratedClass.h"
+#include "Engine/Engine.h"
+#include "Engine/Level.h"
 #include "GameFramework/Actor.h"
 #include "HAL/IConsoleManager.h"
 #include "IRemoteControlModule.h"
