Add owner actor name column
Simplify exposed property label
Display last valid bindings in details if its invalid

diff --git a/Engine/Plugins/VirtualProduction/RemoteControl/Source/RemoteControl/Private/RemoteControlBinding.cpp b/Engine/Plugins/VirtualProduction/RemoteControl/Source/RemoteControl/Private/RemoteControlBinding.cpp
--- a/Engine/Plugins/VirtualProduction/RemoteControl/Source/RemoteControl/Private/RemoteControlBinding.cpp
+++ b/Engine/Plugins/VirtualProduction/RemoteControl/Source/RemoteControl/Private/RemoteControlBinding.cpp
@@ -130,6 +130,11 @@ static FAutoConsoleCommand CCmdDumpBindings = FAutoConsoleCommand(
 	FConsoleCommandWithArgsDelegate::CreateStatic(&DumpBindings)
 );
 
+FSoftObjectPath URemoteControlBinding::GetLastBoundObjectPath() const
+{
+	return LastBoundObjectPath;
+}
+
 void URemoteControlLevelIndependantBinding::SetBoundObject(const TSoftObjectPtr<UObject>& InObject)
 {
 	BoundObject = InObject;
@@ -163,13 +168,22 @@ bool URemoteControlLevelIndependantBinding::PruneDeletedObjects()
 	if (!BoundObject.IsValid())
 	{
 		Modify();
-		BoundObject.Reset();
+		BoundObject.ResetWeakPtr();
 		return true;
 	}
 
 	return false;
 }
 
+void URemoteControlLevelIndependantBinding::PostLoad()
+{
+	Super::PostLoad();
+	if (!LastBoundObjectPath.IsValid())
+	{
+		LastBoundObjectPath = BoundObject.ToSoftObjectPath();
+	}
+}
+
 void URemoteControlLevelDependantBinding::SetBoundObject(const TSoftObjectPtr<UObject>& InObject)
 {
 	if (ensure(InObject))
@@ -180,6 +194,7 @@ void URemoteControlLevelDependantBinding::SetBoundObject(const TSoftObjectPtr<UO
 		SubLevelSelectionMapByPath.FindOrAdd(OuterLevel->GetWorld()) = OuterLevel;
 		
 		Name = EditorObject->GetName();
+		LastBoundObjectPath = InObject.ToSoftObjectPath();
 
 		const bool bShouldSetSubObjectContext = !BindingContext.HasValidSubObjectPath() && !EditorObject->IsA<AActor>();
 		
@@ -260,6 +275,10 @@ UObject* URemoteControlLevelDependantBinding::Resolve() const
 		}
 
 		LevelWithLastSuccessfulResolve = Object->GetTypedOuter<ULevel>();
+		if (!LastBoundObjectPath.IsValid())
+		{
+			LastBoundObjectPath = FSoftObjectPath(Object);
+		}
 
 		if (BindingContext.OwnerActorName.IsNone() || (!Object->IsA<AActor>() && !BindingContext.HasValidSubObjectPath()))
 		{
@@ -340,6 +359,11 @@ void URemoteControlLevelDependantBinding::PostLoad()
 
 	PRAGMA_ENABLE_DEPRECATION_WARNINGS
 #endif
+
+	if (!LastBoundObjectPath.IsValid())
+	{
+		LastBoundObjectPath = BoundObjectMapByPath.FindRef(LevelWithLastSuccessfulResolve.ToSoftObjectPath()).ToSoftObjectPath();
+	}
 }
 
 TSoftObjectPtr<UObject> URemoteControlLevelDependantBinding::ResolveForCurrentWorld() const
@@ -404,11 +428,6 @@ void URemoteControlLevelDependantBinding::InitializeBindingContext(UObject* InOb
 	}
 }
 
-TSoftObjectPtr<UObject> URemoteControlLevelDependantBinding::GetLastBoundObject() const
-{
-	return BoundObjectMapByPath.FindRef(LevelWithLastSuccessfulResolve.ToSoftObjectPath());
-}
-
 UClass* URemoteControlLevelDependantBinding::GetSupportedOwnerClass() const
 {
 	return BindingContext.OwnerActorClass.LoadSynchronous();
@@ -433,6 +452,8 @@ void URemoteControlLevelDependantBinding::SetBoundObject(const TSoftObjectPtr<UL
 	SubPath.Split(TEXT("."), &LeftPart, &ObjectName, ESearchCase::Type::IgnoreCase, ESearchDir::FromEnd);
 	ensure(ObjectName.Len());
 	Name = MoveTemp(ObjectName);
+	LastBoundObjectPath = BoundObject.ToSoftObjectPath();
 }
 
-#undef LOCTEXT_NAMESPACE 
+#undef LOCTEXT_NAMESPACE
+
diff --git a/Engine/Plugins/VirtualProduction/RemoteControl/Source/RemoteControl/Private/RemoteControlEntity.cpp b/Engine/Plugins/VirtualProduction/RemoteControl/Source/RemoteControl/Private/RemoteControlEntity.cpp
--- a/Engine/Plugins/VirtualProduction/RemoteControl/Source/RemoteControl/Private/RemoteControlEntity.cpp
+++ b/Engine/Plugins/VirtualProduction/RemoteControl/Source/RemoteControl/Private/RemoteControlEntity.cpp
@@ -88,6 +88,21 @@ bool FRemoteControlEntity::IsBound() const
 	return GetBoundObjects().Num() > 0;
 }
 
+FSoftObjectPath FRemoteControlEntity::GetLastBindingPath() const
+{
+	FSoftObjectPath Path;
+	for (const TWeakObjectPtr<URemoteControlBinding>& Binding : Bindings)
+	{
+		if (Binding.IsValid())
+		{
+			Path = Binding->GetLastBoundObjectPath();
+			break;
+		}
+	}
+
+	return Path;
+}
+
 bool FRemoteControlEntity::operator==(const FRemoteControlEntity& InEntity) const
 {
 	return Id == InEntity.Id;
diff --git a/Engine/Plugins/VirtualProduction/RemoteControl/Source/RemoteControl/Private/RemoteControlPreset.cpp b/Engine/Plugins/VirtualProduction/RemoteControl/Source/RemoteControl/Private/RemoteControlPreset.cpp
--- a/Engine/Plugins/VirtualProduction/RemoteControl/Source/RemoteControl/Private/RemoteControlPreset.cpp
+++ b/Engine/Plugins/VirtualProduction/RemoteControl/Source/RemoteControl/Private/RemoteControlPreset.cpp
@@ -108,7 +108,7 @@ namespace
 				ObjectName = FieldOwner->HasAnyFlags(RF_ClassDefaultObject | RF_ArchetypeObject) ? FieldOwner->GetClass()->GetName() : FieldOwner->GetName();
 			}
 
-			OutputName = *FString::Printf(TEXT("%s (%s)"), *FieldName, *ObjectName);
+			OutputName = *FString::Printf(TEXT("%s"), *FieldName);
 		}
 		return OutputName;
 	}
@@ -1244,22 +1244,6 @@ FName URemoteControlPreset::GetEntityName(const FName InDesiredName, UObject* In
 	
 	if (DesiredName == NAME_None)
 	{
-		FString ObjectName;
-#if WITH_EDITOR
-		if (AActor* Actor = Cast<AActor>(InObject))
-		{
-			ObjectName = Actor->GetActorLabel();
-		}
-		else if(UActorComponent* Component = Cast<UActorComponent>(InObject))
-		{
-			ObjectName = Component->GetOwner()->GetActorLabel();
-		}
-		else
-#endif
-		{
-			ObjectName = InObject->GetName();
-		}
-
 		FProperty* Property = InFieldPath.GetResolvedData().Field;
 		check(Property);
 
@@ -1271,7 +1255,7 @@ FName URemoteControlPreset::GetEntityName(const FName InDesiredName, UObject* In
 		FieldPath = InFieldPath.GetFieldName().ToString();
 #endif
 			
-		DesiredName = *FString::Printf(TEXT("%s (%s)"), *FieldPath, *ObjectName);
+		DesiredName = *FString::Printf(TEXT("%s"), *FieldPath);
 	}
 
 	return DesiredName;
diff --git a/Engine/Plugins/VirtualProduction/RemoteControl/Source/RemoteControl/Private/RemoteControlPresetRebindingManager.cpp b/Engine/Plugins/VirtualProduction/RemoteControl/Source/RemoteControl/Private/RemoteControlPresetRebindingManager.cpp
--- a/Engine/Plugins/VirtualProduction/RemoteControl/Source/RemoteControl/Private/RemoteControlPresetRebindingManager.cpp
+++ b/Engine/Plugins/VirtualProduction/RemoteControl/Source/RemoteControl/Private/RemoteControlPresetRebindingManager.cpp
@@ -507,7 +507,7 @@ void FRemoteControlPresetRebindingManager::Rebind_NewAlgo(URemoteControlPreset*
 				// We manually get the path from the last bound object since it was not saved.
 				// We are looking for a path like DisplayClusterConfigurationData_0.DisplayClusterConfigurationCluster_0.Node_1.VP_1
 				// We will then combine this with display cluster root actors in the scene to attempt finding a matching object.
-				FString LastObjectPath = InvalidBinding->GetLastBoundObject().ToSoftObjectPath().ToString();
+				FString LastObjectPath = InvalidBinding->GetLastBoundObjectPath().ToString();
 				if (LastObjectPath.IsEmpty())
 				{
 					UE_LOG(LogRemoteControl, Verbose, TEXT("Could not find the last bound object path on the binding."));
diff --git a/Engine/Plugins/VirtualProduction/RemoteControl/Source/RemoteControl/Public/RemoteControlBinding.h b/Engine/Plugins/VirtualProduction/RemoteControl/Source/RemoteControl/Public/RemoteControlBinding.h
--- a/Engine/Plugins/VirtualProduction/RemoteControl/Source/RemoteControl/Public/RemoteControlBinding.h
+++ b/Engine/Plugins/VirtualProduction/RemoteControl/Source/RemoteControl/Public/RemoteControlBinding.h
@@ -69,12 +69,24 @@ public:
 	 */
 	virtual bool PruneDeletedObjects() PURE_VIRTUAL(URemoteControlBinding::PruneDeletedObjects, return false;);
 
+	/**
+	 *  Get the last object this binding was pointing to.
+	 */
+	FSoftObjectPath GetLastBoundObjectPath() const;
+
 public:
 	/**
 	 * The name of this binding. Defaults to the bound object's name.
 	 */
 	UPROPERTY(EditAnywhere, Category=Default)
 	FString Name;
+
+protected:
+	/**
+	 * The path to the last object that was bound.
+	 */
+	UPROPERTY()
+	mutable FSoftObjectPath LastBoundObjectPath;
 };
 
 UCLASS(BlueprintType)
@@ -83,12 +95,19 @@ class REMOTECONTROL_API URemoteControlLevelIndependantBinding : public URemoteCo
 public:
 	GENERATED_BODY()
 
+	//~ Begin UObject Interface
+	virtual void PostLoad() override;
+	//~ End UObject Interface
+
+	//~ Begin URemoteControlBinding Interface
 	virtual void SetBoundObject(const TSoftObjectPtr<UObject>& InObject) override;
 	virtual void UnbindObject(const TSoftObjectPtr<UObject>& InBoundObject) override;
 	virtual UObject* Resolve() const override;
 	virtual bool IsValid() const override;
 	virtual bool IsBound(const TSoftObjectPtr<UObject>& Object) const override;
 	virtual bool PruneDeletedObjects() override;
+	//~ End URemoteControlBinding interface
+
 
 private:
 	/**
@@ -172,6 +191,10 @@ class REMOTECONTROL_API URemoteControlLevelDependantBinding : public URemoteCont
 public:
 	GENERATED_BODY()
 
+	//~ Begin UObject Interface
+	virtual void PostLoad() override;
+	//~ End UObject Interface
+	
 	//~ Begin URemoteControlBinding Interface
 	virtual void SetBoundObject(const TSoftObjectPtr<UObject>& BoundObject) override;
 	virtual void UnbindObject(const TSoftObjectPtr<UObject>& BoundObject) override;
@@ -181,8 +204,6 @@ public:
 	virtual bool PruneDeletedObjects() override;
 	//~ Begin URemoteControlBinding Interface
 
-	virtual void PostLoad() override;
-
 	/**
 	 *	Set the bound object by specifying the level it resides in.
 	 *	@Note Useful is you want to set the bound object without loading the level/object.
@@ -213,7 +234,6 @@ private:
 
 	void InitializeBindingContext(UObject* InObject) const;
 
-	TSoftObjectPtr<UObject> GetLastBoundObject() const;
 
 private:
 #if WITH_EDITORONLY_DATA
diff --git a/Engine/Plugins/VirtualProduction/RemoteControl/Source/RemoteControl/Public/RemoteControlEntity.h b/Engine/Plugins/VirtualProduction/RemoteControl/Source/RemoteControl/Public/RemoteControlEntity.h
--- a/Engine/Plugins/VirtualProduction/RemoteControl/Source/RemoteControl/Public/RemoteControlEntity.h
+++ b/Engine/Plugins/VirtualProduction/RemoteControl/Source/RemoteControl/Public/RemoteControlEntity.h
@@ -108,6 +108,11 @@ struct REMOTECONTROL_API FRemoteControlEntity
 	 */
 	virtual UClass* GetSupportedBindingClass() const { return UObject::StaticClass(); }
 
+	/**
+	 * Get the last valid binding path for this entity.
+	 */
+	FSoftObjectPath GetLastBindingPath() const;
+
 	bool operator==(const FRemoteControlEntity& InEntity) const;
 	bool operator==(FGuid InEntityId) const;
 	friend uint32 REMOTECONTROL_API GetTypeHash(const FRemoteControlEntity& InEntity);
diff --git a/Engine/Plugins/VirtualProduction/RemoteControl/Source/RemoteControlUI/Private/UI/Customizations/RemoteControlEntityCustomization.cpp b/Engine/Plugins/VirtualProduction/RemoteControl/Source/RemoteControlUI/Private/UI/Customizations/RemoteControlEntityCustomization.cpp
--- a/Engine/Plugins/VirtualProduction/RemoteControl/Source/RemoteControlUI/Private/UI/Customizations/RemoteControlEntityCustomization.cpp
+++ b/Engine/Plugins/VirtualProduction/RemoteControl/Source/RemoteControlUI/Private/UI/Customizations/RemoteControlEntityCustomization.cpp
@@ -7,6 +7,7 @@
 #include "DetailWidgetRow.h"
 #include "IDocumentation.h"
 #include "IPropertyRowGenerator.h"
+#include "RemoteControlBinding.h"
 #include "RemoteControlEntity.h"
 #include "RemoteControlField.h"
 #include "RemoteControlPreset.h"
@@ -128,25 +129,27 @@ void FRemoteControlEntityCustomization::CreateBindingWidget(const FRemoteControl
 	const FSlateFontInfo FontInfo = LayoutBuilder.GetDetailFont();
 	FText ToolTipText = LOCTEXT("RCBindingToolTip", "The object this is currently bound to.");
 
-	FString PathName;
+	FText PathNameText;
 
 	if (ensure(RCEntity))
 	{
 		if (UObject* Obj = RCEntity->GetBoundObject())
 		{
-			PathName = Obj->GetPathName();
+			PathNameText = FText::FromString(Obj->GetPathName());
+		}
+		else
+		{
+			FString BindingPath = RCEntity->GetLastBindingPath().ToString();
+			if (!BindingPath.IsEmpty())
+			{
+				PathNameText = FText::Format(LOCTEXT("InvalidBindingPath", "INVALID_BINDING: {0}"), FText::FromString(BindingPath));
+			}
 		}
 	}
 
-	FText PathNameText;
-
-	if (PathName.IsEmpty())
-	{
-		PathNameText = LOCTEXT("BindingPath", "This entity is not bound.");
-	}
-	else
+	if (PathNameText.IsEmpty())
 	{
-		PathNameText = FText::FromString(MoveTemp(PathName));
+		PathNameText = LOCTEXT("NoBindingInformation", "INVALID_BINDING: - This entity has no binding information.");
 	}
 
 	CategoryBuilder.AddCustomRow(LOCTEXT("BindingLabel", "Binding"))
diff --git a/Engine/Plugins/VirtualProduction/RemoteControl/Source/RemoteControlUI/Private/UI/SRCPanelExposedEntitiesList.cpp b/Engine/Plugins/VirtualProduction/RemoteControl/Source/RemoteControlUI/Private/UI/SRCPanelExposedEntitiesList.cpp
--- a/Engine/Plugins/VirtualProduction/RemoteControl/Source/RemoteControlUI/Private/UI/SRCPanelExposedEntitiesList.cpp
+++ b/Engine/Plugins/VirtualProduction/RemoteControl/Source/RemoteControlUI/Private/UI/SRCPanelExposedEntitiesList.cpp
@@ -281,28 +281,43 @@ void SRCPanelExposedEntitiesList::Construct(const FArguments& InArgs, URemoteCon
 		.HeaderRow(
 			SAssignNew(FieldsHeaderRow, SHeaderRow)
 			.Style(&RCPanelStyle->HeaderRowStyle)
+			.CanSelectGeneratedColumn(true) //To show/hide columns
 
 			+ SHeaderRow::Column(RemoteControlPresetColumns::DragDropHandle)
 			.DefaultLabel(LOCTEXT("RCPresetDragDropHandleColumnHeader", ""))
 			.FixedWidth(25.f)
 			.HeaderContentPadding(RCPanelStyle->HeaderRowPadding)
+			.ShouldGenerateWidget(true)
+
+			+ SHeaderRow::Column(RemoteControlPresetColumns::OwnerName)
+			.DefaultLabel(LOCTEXT("RCPresetOwnerNameColumnHeader", "Owner Name"))
+			.HAlignHeader(HAlign_Center)
+			.FillWidth(0.125f)
+			.HeaderContentPadding(RCPanelStyle->HeaderRowPadding)
+
+			+ SHeaderRow::Column(RemoteControlPresetColumns::SubobjectPath)
+			.DefaultLabel(LOCTEXT("RCPresetSubobjectPathColumnHeader", "Subobject Path"))
+			.HAlignHeader(HAlign_Center)
+			.FillWidth(0.125f)
+			.HeaderContentPadding(RCPanelStyle->HeaderRowPadding)
 
 			+ SHeaderRow::Column(RemoteControlPresetColumns::Description)
 			.DefaultLabel(LOCTEXT("RCPresetDescColumnHeader", "Description"))
 			.HAlignHeader(HAlign_Center)
-			.FillWidth(0.5f)
+			.FillWidth(0.35f)
 			.HeaderContentPadding(RCPanelStyle->HeaderRowPadding)
 
 			+ SHeaderRow::Column(RemoteControlPresetColumns::Value)
 			.DefaultLabel(LOCTEXT("RCPresetValueColumnHeader", "Value"))
 			.HAlignHeader(HAlign_Center)
-			.FillWidth(0.5f)
+			.FillWidth(0.4f)
 			.HeaderContentPadding(RCPanelStyle->HeaderRowPadding)
 
 			+ SHeaderRow::Column(RemoteControlPresetColumns::Reset)
 			.DefaultLabel(LOCTEXT("RCPresetResetButtonColumnHeader", ""))
 			.FixedWidth(48.f)
 			.HeaderContentPadding(RCPanelStyle->HeaderRowPadding)
+			.ShouldGenerateWidget(true)
 		);
 
 	// Exposed Entities Dock Panel
diff --git a/Engine/Plugins/VirtualProduction/RemoteControl/Source/RemoteControlUI/Private/UI/SRCPanelExposedEntity.cpp b/Engine/Plugins/VirtualProduction/RemoteControl/Source/RemoteControlUI/Private/UI/SRCPanelExposedEntity.cpp
--- a/Engine/Plugins/VirtualProduction/RemoteControl/Source/RemoteControlUI/Private/UI/SRCPanelExposedEntity.cpp
+++ b/Engine/Plugins/VirtualProduction/RemoteControl/Source/RemoteControlUI/Private/UI/SRCPanelExposedEntity.cpp
@@ -24,6 +24,7 @@
 #include "SceneOutlinerFilters.h"
 #include "SceneOutlinerModule.h"
 #include "ScopedTransaction.h"
+#include "Styling/CoreStyle.h"
 #include "Styling/RemoteControlStyles.h"
 #include "Styling/SlateIconFinder.h"
 #include "Modules/ModuleManager.h"
@@ -34,6 +35,7 @@
 #include "Widgets/Layout/SBox.h"
 #include "Widgets/Layout/SBorder.h"
 #include "Widgets/SBoxPanel.h"
+#include "Widgets/Text/STextBlock.h"
 #include "Widgets/Text/SInlineEditableTextBlock.h"
 
 #define LOCTEXT_NAMESPACE "RemoteControlPanel"
@@ -149,7 +151,56 @@ void SRCPanelExposedEntity::Initialize(const FGuid& InEntityId, URemoteControlPr
 	{
 		if (const TSharedPtr<FRemoteControlEntity> RCEntity = InPreset->GetExposedEntity(InEntityId).Pin())
 		{
+			const FString BindingPath = RCEntity->GetLastBindingPath().ToString();
 			CachedLabel = RCEntity->GetLabel();
+			CachedBindingPath = *BindingPath;
+
+			FName OwnerFName;
+			// If the binding is valid, display the actor label if possible
+			if (UObject* Object = RCEntity->GetBoundObject())
+			{
+				bValidBinding = true;
+				
+				if (AActor* OwnerActor = Object->GetTypedOuter<AActor>())
+				{
+					CachedOwnerName = *OwnerActor->GetActorLabel();
+					OwnerFName = OwnerActor->GetFName();
+				}
+				else if (AActor* Actor = Cast<AActor>(Object))
+				{
+					CachedOwnerName = *Actor->GetActorLabel();
+					OwnerFName = Object->GetFName();
+				}
+				else
+				{
+					CachedOwnerName = Object->GetFName();
+					OwnerFName = Object->GetFName();
+				}
+			}
+			else
+			{
+				// If not, default to the owner fname
+				bValidBinding = false;
+
+				static const FString PersistentLevelString = TEXT(":PersistentLevel.");
+				int32 PersistentLevelIndex = BindingPath.Find(PersistentLevelString);
+				if (PersistentLevelIndex != INDEX_NONE)
+				{
+					OwnerFName = *BindingPath.RightChop(PersistentLevelIndex + PersistentLevelString.Len());
+					CachedOwnerName = OwnerFName;
+				}
+			}
+			
+			const int32 OwnerNameIndex = BindingPath.Find(OwnerFName.ToString() + TEXT("."));
+			if (OwnerNameIndex != INDEX_NONE)
+			{
+				CachedSubobjectPath = *BindingPath.RightChop(OwnerNameIndex + OwnerFName.GetStringLength() + 1);
+			}
+
+			if (CachedOwnerName.IsNone())
+			{
+				CachedOwnerName = *LOCTEXT("InvalidOwner", "Invalid Owner").ToString();
+			}
 		}
 	}
 }
@@ -395,6 +446,40 @@ TSharedRef<SWidget> SRCPanelExposedEntity::CreateEntityWidget(TSharedPtr<SWidget
 			.Widget(Widget)
 		];
 
+	const FSlateBrush* TrashBrush = FAppStyle::Get().GetBrush("Icons.Delete");
+
+	Args.OwnerNameWidget = SNew(SBox)
+		.Visibility(this, &SRCPanelExposedEntity::GetVisibilityAccordingToLiveMode, EVisibility::Collapsed)
+		[
+			SNew(SHorizontalBox)
+			+ SHorizontalBox::Slot()
+			.Padding(FMargin(0.0, 0.0, 4.0, 0.0))
+			.AutoWidth()
+			[
+				SNew(SImage)
+				.ColorAndOpacity(FSlateColor::UseSubduedForeground())
+				.Visibility_Lambda([this]() { return bValidBinding ? EVisibility::Collapsed : EVisibility::Visible; })
+				.Image(TrashBrush)
+			]
+			+ SHorizontalBox::Slot()
+			.AutoWidth()
+			[
+				SNew(STextBlock)
+				.ColorAndOpacity_Lambda([this]() { return bValidBinding ? FSlateColor::UseForeground() : FSlateColor::UseSubduedForeground(); })
+				.Text(FText::FromName(CachedOwnerName))
+				.ToolTipText(FText::FromName(CachedBindingPath))
+			]
+		];
+
+	Args.SubObjectPathWidget = SNew(SBox)
+		.Visibility(this, &SRCPanelExposedEntity::GetVisibilityAccordingToLiveMode, EVisibility::Collapsed)
+		[
+			SNew(STextBlock)
+			.ColorAndOpacity_Lambda([this]() { return bValidBinding ? FSlateColor::UseForeground() : FSlateColor::UseSubduedForeground(); })
+			.Text(CachedSubobjectPath.IsNone() ? FText::GetEmpty() : FText::FromName(CachedSubobjectPath))
+			.ToolTipText(LOCTEXT("SubobjectPathToolTip", "The path from the owner actor to the uobject holding the exposed property."))
+		];
+
 	Args.NameWidget = SNew(SHorizontalBox)
 		.Clipping(EWidgetClipping::OnDemand)
 		+ SHorizontalBox::Slot()
diff --git a/Engine/Plugins/VirtualProduction/RemoteControl/Source/RemoteControlUI/Private/UI/SRCPanelExposedEntity.h b/Engine/Plugins/VirtualProduction/RemoteControl/Source/RemoteControlUI/Private/UI/SRCPanelExposedEntity.h
--- a/Engine/Plugins/VirtualProduction/RemoteControl/Source/RemoteControlUI/Private/UI/SRCPanelExposedEntity.h
+++ b/Engine/Plugins/VirtualProduction/RemoteControl/Source/RemoteControlUI/Private/UI/SRCPanelExposedEntity.h
@@ -49,6 +49,14 @@ protected:
 	TWeakObjectPtr<URemoteControlPreset> Preset;
 	/** Whether the panel is in live mode or not. */
 	TAttribute<bool> bLiveMode;
+	/** Name of the entity's owner (What actor it's bound to). */
+	FName CachedOwnerName;
+	/** Path from the entity owner to the object holding the exposed property. */
+	FName CachedSubobjectPath;
+	/** Path of the entity's binding. */
+	FName CachedBindingPath;
+	/** Whether the binding is valid at the moment. */
+	bool bValidBinding = false;
 	/** Display name of the entity. */
 	FName CachedLabel;
 	/** Text to be highlighted while searching. */
diff --git a/Engine/Plugins/VirtualProduction/RemoteControl/Source/RemoteControlUI/Private/UI/SRCPanelTreeNode.cpp b/Engine/Plugins/VirtualProduction/RemoteControl/Source/RemoteControlUI/Private/UI/SRCPanelTreeNode.cpp
--- a/Engine/Plugins/VirtualProduction/RemoteControl/Source/RemoteControlUI/Private/UI/SRCPanelTreeNode.cpp
+++ b/Engine/Plugins/VirtualProduction/RemoteControl/Source/RemoteControlUI/Private/UI/SRCPanelTreeNode.cpp
@@ -12,6 +12,8 @@
 
 TSet<FName> SRCPanelTreeNode::DefaultColumns = {
 	RemoteControlPresetColumns::DragDropHandle,
+	RemoteControlPresetColumns::OwnerName,
+	RemoteControlPresetColumns::SubobjectPath,
 	RemoteControlPresetColumns::Description,
 	RemoteControlPresetColumns::Value,
 	RemoteControlPresetColumns::Reset
@@ -43,6 +45,14 @@ TSharedRef<SWidget> SRCPanelTreeNode::GetWidget(const FName ForColumnName, const
 	{
 		return DragHandleWidget.ToSharedRef();
 	}
+	else if (ForColumnName == RemoteControlPresetColumns::OwnerName)
+	{
+		return NodeOwnerNameWidget.ToSharedRef();
+	}
+    else if (ForColumnName == RemoteControlPresetColumns::SubobjectPath)
+    {
+    	return SubobjectPathWidget.ToSharedRef();
+    }
 	else if (ForColumnName == RemoteControlPresetColumns::Description)
 	{
 		return NodeNameWidget.ToSharedRef();
@@ -98,6 +108,20 @@ TSharedRef<SWidget> SRCPanelTreeNode::MakeNodeWidget(const FMakeNodeWidgetArgs&
 		[
 			DragHandleWidget.ToSharedRef()
 		]
+		// Owner name
+		+ SHorizontalBox::Slot()
+		.VAlign(VAlign_Center)
+		.AutoWidth()
+		[
+			NodeOwnerNameWidget.ToSharedRef()
+		]
+		// Subobject Path
+		+ SHorizontalBox::Slot()
+		.VAlign(VAlign_Center)
+		.AutoWidth()
+		[
+			SubobjectPathWidget.ToSharedRef()
+		]
 		// Field name
 		+ SHorizontalBox::Slot()
 		.VAlign(VAlign_Center)
@@ -136,6 +160,10 @@ void SRCPanelTreeNode::MakeNodeWidgets(const FMakeNodeWidgetArgs& Args)
 
 	DragHandleWidget = WidgetOrNull(Args.DragHandle);
 
+	NodeOwnerNameWidget = WidgetOrNull(Args.OwnerNameWidget);
+	
+	SubobjectPathWidget = WidgetOrNull(Args.SubObjectPathWidget);
+
 	NodeNameWidget = WidgetOrNull(Args.NameWidget);
 
 	NodeValueWidget = WidgetOrNull(Args.ValueWidget);
@@ -167,4 +195,4 @@ void SRCPanelTreeNode::SetColumnWidth(float InWidth)
 	ColumnSizeData.SetColumnWidth(FMath::Clamp(InWidth + SplitterOffset, 0.f, 1.f));
 }
 
-#undef LOCTEXT_NAMESPACE /*RemoteControlPanelNode*/
\ No newline at end of file
+#undef LOCTEXT_NAMESPACE /*RemoteControlPanelNode*/
diff --git a/Engine/Plugins/VirtualProduction/RemoteControl/Source/RemoteControlUI/Private/UI/SRCPanelTreeNode.h b/Engine/Plugins/VirtualProduction/RemoteControl/Source/RemoteControlUI/Private/UI/SRCPanelTreeNode.h
--- a/Engine/Plugins/VirtualProduction/RemoteControl/Source/RemoteControlUI/Private/UI/SRCPanelTreeNode.h
+++ b/Engine/Plugins/VirtualProduction/RemoteControl/Source/RemoteControlUI/Private/UI/SRCPanelTreeNode.h
@@ -19,6 +19,8 @@ namespace RemoteControlPresetColumns
 {
 	static FName DragDropHandle = TEXT("DragDropHandle");
 	static FName LinkIdentifier = TEXT("LinkID");
+	static FName OwnerName = TEXT("OwnerName");
+	static FName SubobjectPath = TEXT("Subobject Path");
 	static FName Description = TEXT("Description");
 	static FName Mask = TEXT("Mask");
 	static FName Value = TEXT("Value");
@@ -86,6 +88,8 @@ protected:
 	struct FMakeNodeWidgetArgs
 	{
 		TSharedPtr<SWidget> DragHandle;
+		TSharedPtr<SWidget> OwnerNameWidget;
+		TSharedPtr<SWidget> SubObjectPathWidget;
 		TSharedPtr<SWidget> NameWidget;
 		TSharedPtr<SWidget> ValueWidget;
 		TSharedPtr<SWidget> ResetButton;
@@ -121,13 +125,19 @@ private:
 	/** Cached widget of drag handle. */
 	TSharedPtr<SWidget> DragHandleWidget;
 
-	/** Cached widget of node name widget. */
+	/** Cached widget of node owner name. */
+	TSharedPtr<SWidget> NodeOwnerNameWidget;
+	
+	/** Cached widget of subobject path. */
+	TSharedPtr<SWidget> SubobjectPathWidget;
+
+	/** Cached widget of node name. */
 	TSharedPtr<SWidget> NodeNameWidget;
 
-	/** Cached widget of node value widget. */
+	/** Cached widget of node value. */
 	TSharedPtr<SWidget> NodeValueWidget;
 	
-	/** Cached widget of reset arrow widget. */
+	/** Cached widget of reset arrow. */
 	TSharedPtr<SWidget> ResetValueWidget;
 
 private:
diff --git a/Engine/Plugins/VirtualProduction/RemoteControl/Source/RemoteControlUI/RemoteControlUI.Build.cs b/Engine/Plugins/VirtualProduction/RemoteControl/Source/RemoteControlUI/RemoteControlUI.Build.cs
--- a/Engine/Plugins/VirtualProduction/RemoteControl/Source/RemoteControlUI/RemoteControlUI.Build.cs
+++ b/Engine/Plugins/VirtualProduction/RemoteControl/Source/RemoteControlUI/RemoteControlUI.Build.cs
@@ -27,6 +27,7 @@ public class RemoteControlUI : ModuleRules
 				"GraphEditor",
 				"HotReload",
 				"InputCore",
+				"LevelEditor",
 				"MainFrame",
 				"MessageLog",
 				"Projects",
@@ -42,6 +43,8 @@ public class RemoteControlUI : ModuleRules
 				"StructUtils",
 				"ToolMenus",
 				"ToolWidgets",
+				"TypedElementFramework",
+				"TypedElementRuntime",
 				"UnrealEd"
 			}
 		);
