Highlight property in detail panels when clicking it in the RCP

diff --git a/Engine/Plugins/VirtualProduction/RemoteControl/Source/RemoteControl/Private/RemoteControlFieldPath.cpp b/Engine/Plugins/VirtualProduction/RemoteControl/Source/RemoteControl/Private/RemoteControlFieldPath.cpp
--- a/Engine/Plugins/VirtualProduction/RemoteControl/Source/RemoteControl/Private/RemoteControlFieldPath.cpp
+++ b/Engine/Plugins/VirtualProduction/RemoteControl/Source/RemoteControl/Private/RemoteControlFieldPath.cpp
@@ -369,6 +369,49 @@ void FRCFieldPathInfo::ToEditPropertyChain(FEditPropertyChain& OutPropertyChain)
 	}
 }
 
+#if WITH_EDITOR
+TSharedRef<FPropertyPath> FRCFieldPathInfo::ToPropertyPath() const
+{
+	TSharedRef<FPropertyPath> Path = FPropertyPath::CreateEmpty();
+	check(IsResolved());
+
+	for (const FRCFieldPathSegment& Segment : Segments)
+	{
+		if (!Segment.Name.IsNone())
+		{
+			// Property path map entries contain duplicates.
+			if (!Segment.MapKey.IsEmpty())
+			{
+				Path->AddProperty(FPropertyInfo(Segment.ResolvedData.Field));
+			}
+
+			if (Segment.ArrayIndex != -1)
+			{
+				if (FArrayProperty* Property = CastField<FArrayProperty>(Segment.ResolvedData.Field))
+				{
+					FPropertyInfo Info;
+					Info.Property = Segment.ResolvedData.Field;
+
+					Path->AddProperty(FPropertyInfo(Property));
+					Path->AddProperty(FPropertyInfo(Property->Inner, Segment.ArrayIndex));
+
+					continue;
+				}
+				
+			}
+
+			FPropertyInfo Info;
+			Info.Property = Segment.ResolvedData.Field;
+			Info.ArrayIndex = Segment.ArrayIndex;
+
+			Path->AddProperty(Info);
+		}
+	}
+
+	return Path;
+}
+#endif
+
 void FRCFieldPathInfo::Initialize(const FString& PathInfo, bool bCleanDuplicates)
 {
 	TArray<FString> PathSegment;
diff --git a/Engine/Plugins/VirtualProduction/RemoteControl/Source/RemoteControl/Public/RemoteControlFieldPath.h b/Engine/Plugins/VirtualProduction/RemoteControl/Source/RemoteControl/Public/RemoteControlFieldPath.h
--- a/Engine/Plugins/VirtualProduction/RemoteControl/Source/RemoteControl/Public/RemoteControlFieldPath.h
+++ b/Engine/Plugins/VirtualProduction/RemoteControl/Source/RemoteControl/Public/RemoteControlFieldPath.h
@@ -2,6 +2,11 @@
 #pragma once
 
 #include "CoreMinimal.h"
+
+#if WITH_EDITOR
+#include "PropertyPath.h"
+#endif
+
 #include "RemoteControlFieldPath.generated.h"
 
 
@@ -194,6 +199,12 @@ public:
 	/** Builds an EditPropertyChain from the segments */
 	void ToEditPropertyChain(FEditPropertyChain& OutPropertyChain) const;
 
+#if WITH_EDITOR
+	/** Converts this RCPath to a property path. */
+	TSharedRef<FPropertyPath> ToPropertyPath() const;
+#endif
+
+
 private:
 	/**
 	 * Initialize from a string of format with '.' delimiters
diff --git a/Engine/Plugins/VirtualProduction/RemoteControl/Source/RemoteControlUI/Private/RemoteControlUIModule.cpp b/Engine/Plugins/VirtualProduction/RemoteControl/Source/RemoteControlUI/Private/RemoteControlUIModule.cpp
--- a/Engine/Plugins/VirtualProduction/RemoteControl/Source/RemoteControlUI/Private/RemoteControlUIModule.cpp
+++ b/Engine/Plugins/VirtualProduction/RemoteControl/Source/RemoteControlUI/Private/RemoteControlUIModule.cpp
@@ -8,6 +8,8 @@
 #include "Commands/RemoteControlCommands.h"
 #include "Components/SkeletalMeshComponent.h"
 #include "Components/StaticMeshComponent.h"
+#include "Elements/Framework/TypedElementRegistry.h"
+#include "Elements/Framework/TypedElementSelectionSet.h"
 #include "Engine/SkeletalMesh.h"
 #include "Engine/StaticMesh.h"
 #include "Framework/MultiBox/MultiBoxBuilder.h"
@@ -15,6 +17,7 @@
 #include "IRemoteControlModule.h"
 #include "Interfaces/IMainFrameModule.h"
 #include "Kismet2/ComponentEditorUtils.h"
+#include "LevelEditorSubsystem.h"
 #include "MaterialEditor/DEditorParameterValue.h"
 #include "Materials/Material.h"
 #include "PropertyHandle.h"
@@ -340,6 +343,7 @@ TSharedRef<SRemoteControlPanel> FRemoteControlUIModule::CreateRemoteControlPanel
 		for (const FName& DetailsTabIdentifier : DetailsTabIdentifiers)
 		{
 			SharedDetailsPanel = PropertyEditor.FindDetailView(DetailsTabIdentifier);
+			SharedDetailsPanel->SetRightColumnMinWidth(50);
 			
 			if (SharedDetailsPanel.IsValid())
 			{
@@ -881,6 +885,42 @@ void FRemoteControlUIModule::UnregisterWidgetFactoryForType(UScriptStruct* Remot
 	GenerateWidgetDelegates.Remove(RemoteControlEntityType);
 }
 
+void FRemoteControlUIModule::HighlightPropertyInDetailsPanel(const FPropertyPath& Path) const
+{
+	if (SharedDetailsPanel)
+	{
+		SharedDetailsPanel->HighlightProperty(Path);
+	}
+}
+
+void FRemoteControlUIModule::SelectObjects(const TArray<UObject*>& Objects) const
+{
+	if (ULevelEditorSubsystem* LevelEditorSubsystem = GEditor->GetEditorSubsystem<ULevelEditorSubsystem>())
+	{
+		FTypedElementListRef ElementHandles = UTypedElementRegistry::GetInstance()->CreateElementList();
+		for (UObject* Object : Objects)
+		{
+			if (UActorComponent* Component = Cast<UActorComponent>(Object))
+			{
+				if (FTypedElementHandle Handle = UEngineElementsLibrary::AcquireEditorComponentElementHandle(Component))
+				{
+					ElementHandles->Add(Handle);
+				}
+			}
+			else if (AActor* Actor = Cast<AActor>(Object))
+			{
+				if (FTypedElementHandle Handle = UEngineElementsLibrary::AcquireEditorActorElementHandle(Actor))
+				{
+					ElementHandles->Add(Handle);
+				}
+			}
+		}
+
+		UTypedElementSelectionSet* SelectionSet = LevelEditorSubsystem->GetSelectionSet();
+		SelectionSet->SetSelection(ElementHandles->GetElementHandles(), FTypedElementSelectionOptions());
+	}
+}
+
 void FRemoteControlUIModule::RegisterWidgetFactories()
 {
 	RegisterWidgetFactoryForType(FRemoteControlActor::StaticStruct(), FOnGenerateRCWidget::CreateStatic(&SRCPanelExposedActor::MakeInstance));
diff --git a/Engine/Plugins/VirtualProduction/RemoteControl/Source/RemoteControlUI/Private/RemoteControlUIModule.h b/Engine/Plugins/VirtualProduction/RemoteControl/Source/RemoteControlUI/Private/RemoteControlUIModule.h
--- a/Engine/Plugins/VirtualProduction/RemoteControl/Source/RemoteControlUI/Private/RemoteControlUIModule.h
+++ b/Engine/Plugins/VirtualProduction/RemoteControl/Source/RemoteControlUI/Private/RemoteControlUIModule.h
@@ -9,6 +9,7 @@
 #include "Modules/ModuleManager.h"
 #include "PropertyEditorDelegates.h"
 #include "PropertyHandle.h"
+#include "PropertyPath.h"
 
 class IToolkitHost;
 class SRemoteControlPanel;
@@ -54,6 +55,8 @@ public:
 	virtual uint32 GetRemoteControlAssetCategory() const override;
 	virtual void RegisterWidgetFactoryForType(UScriptStruct* RemoteControlEntityType, const FOnGenerateRCWidget& OnGenerateRCWidgetDelegate) override;
 	virtual void UnregisterWidgetFactoryForType(UScriptStruct* RemoteControlEntityType) override;
+	virtual void HighlightPropertyInDetailsPanel(const FPropertyPath& Path) const override;
+	virtual void SelectObjects(const TArray<UObject*>& Objects) const override;
 	//~ End IRemoteControlUIModule interface
 
 	/**
diff --git a/Engine/Plugins/VirtualProduction/RemoteControl/Source/RemoteControlUI/Private/UI/SRCPanelExposedEntitiesList.cpp b/Engine/Plugins/VirtualProduction/RemoteControl/Source/RemoteControlUI/Private/UI/SRCPanelExposedEntitiesList.cpp
--- a/Engine/Plugins/VirtualProduction/RemoteControl/Source/RemoteControlUI/Private/UI/SRCPanelExposedEntitiesList.cpp
+++ b/Engine/Plugins/VirtualProduction/RemoteControl/Source/RemoteControlUI/Private/UI/SRCPanelExposedEntitiesList.cpp
@@ -17,6 +17,7 @@
 #include "Misc/Attribute.h"
 #include "Misc/Guid.h"
 #include "Misc/MessageDialog.h"
+#include "PropertyPath.h"
 #include "RemoteControlPanelStyle.h"
 #include "RemoteControlPreset.h"
 #include "RemoteControlUIModule.h"
@@ -35,6 +36,7 @@
 #include "Widgets/Views/STableRow.h"
 #include "Widgets/Views/STreeView.h"
 
+
 #define LOCTEXT_NAMESPACE "RemoteControlPanelEntitiesList"
 
 TSet<FName> SRCPanelExposedEntitiesList::DefaultProtocolColumns = {
@@ -292,13 +294,13 @@ void SRCPanelExposedEntitiesList::Construct(const FArguments& InArgs, URemoteCon
 			+ SHeaderRow::Column(RemoteControlPresetColumns::OwnerName)
 			.DefaultLabel(LOCTEXT("RCPresetOwnerNameColumnHeader", "Owner Name"))
 			.HAlignHeader(HAlign_Center)
-			.FillWidth(0.125f)
+			.FillWidth(0.15f)
 			.HeaderContentPadding(RCPanelStyle->HeaderRowPadding)
 
 			+ SHeaderRow::Column(RemoteControlPresetColumns::SubobjectPath)
 			.DefaultLabel(LOCTEXT("RCPresetSubobjectPathColumnHeader", "Subobject Path"))
 			.HAlignHeader(HAlign_Center)
-			.FillWidth(0.125f)
+			.FillWidth(0.15f)
 			.HeaderContentPadding(RCPanelStyle->HeaderRowPadding)
 
 			+ SHeaderRow::Column(RemoteControlPresetColumns::Description)
@@ -838,7 +840,7 @@ TSharedRef<ITableRow> SRCPanelExposedEntitiesList::OnGenerateRow(TSharedPtr<SRCP
 
 		return FReply::Unhandled();
 	};
-	
+
 	if (Node->GetRCType() == SRCPanelTreeNode::Group)
 	{
 		return SNew(STableRow<TSharedPtr<SWidget>>, OwnerTable)
@@ -889,13 +891,13 @@ void SRCPanelExposedEntitiesList::OnSelectionChanged(TSharedPtr<SRCPanelTreeNode
 			CurrentlySelectedGroup = RCGroup->Id;
 
 			GenerateListWidgets(*RCGroup);
-			
+
 			if (BackendFilter.HasAnyActiveFilters())
 			{
 				bFilterApplicationRequested = true;
 			}
 		}
-	
+
 		FieldsListView->ClearSelection();
 	}
 	else
@@ -903,7 +905,8 @@ void SRCPanelExposedEntitiesList::OnSelectionChanged(TSharedPtr<SRCPanelTreeNode
 		// todo call handler on node itself
 		if (TSharedPtr<FRemoteControlField> RCField = Preset->GetExposedEntity<FRemoteControlField>(Node->GetRCId()).Pin())
 		{
-			TSet<UObject*> Objects = TSet<UObject*>{ RCField->GetBoundObjects() };
+			TArray<UObject*> BoundObjects = RCField->GetBoundObjects();
+			TSet<UObject*> Objects = TSet<UObject*>{ BoundObjects };
 
 			TArray<UObject*> OwnerActors;
 			for (UObject* Object : Objects)
@@ -918,7 +921,55 @@ void SRCPanelExposedEntitiesList::OnSelectionChanged(TSharedPtr<SRCPanelTreeNode
 				}
 			}
 
-			SelectActorsInlevel(OwnerActors);
+			if (RCField->GetStruct() == FRemoteControlProperty::StaticStruct())
+			{
+				TSharedPtr<FRemoteControlProperty> RCProp = StaticCastSharedPtr<FRemoteControlProperty>(RCField);
+				TSharedRef<FPropertyPath> PropertyPath = RCProp->FieldPathInfo.ToPropertyPath();
+
+				for (auto It = BoundObjects.CreateIterator(); It; ++It)
+				{
+					if (AActor* OwnerActor = (*It)->GetTypedOuter<AActor>())
+					{
+						UObject* BindingObject = *It;
+
+						// When we encounter a non-component object, 
+						if (!BindingObject->IsA<UActorComponent>())
+						{
+							BoundObjects.Add(OwnerActor);
+							It.RemoveCurrent();
+						}
+
+						// --- Special NDisplay handling because of their customization ---
+						// Since display cluster config data is not created as a defaeult subobject, therefore we have no way of retrieving the CurrentConfigData property from the config object itself.
+						static FName DisplayClusterConfigDataClassName = "DisplayClusterConfigurationData";
+						if (BindingObject->GetClass()->GetFName() == DisplayClusterConfigDataClassName)
+						{
+							if (FProperty* Property = OwnerActor->GetClass()->FindPropertyByName("CurrentConfigData"))
+							{
+								// Append "CurrentConfigData" to the beginning of the path.
+								TSharedRef<FPropertyPath> NewPropertyPath = FPropertyPath::Create(Property);
+								for (int32 Index = 0; Index < PropertyPath->GetNumProperties(); Index++)
+								{
+									NewPropertyPath->AddProperty(PropertyPath->GetPropertyInfo(Index));
+								}
+
+								PropertyPath = NewPropertyPath;
+							}
+						}
+					}
+				}
+
+				FRemoteControlUIModule::Get().SelectObjects(BoundObjects);
+				
+				FTimerHandle Handle;
+				FTimerDelegate Delegate = FTimerDelegate::CreateLambda(([PropertyPath]()
+					{
+						FRemoteControlUIModule::Get().HighlightPropertyInDetailsPanel(*PropertyPath);
+					}));
+
+				// Needed because modifying the selection set is asynchronous.
+				GEditor->GetTimerManager()->SetTimer(Handle, Delegate, 0.1, 0, -1);
+			}
 		}
 	}
 
diff --git a/Engine/Plugins/VirtualProduction/RemoteControl/Source/RemoteControlUI/Public/IRemoteControlUIModule.h b/Engine/Plugins/VirtualProduction/RemoteControl/Source/RemoteControlUI/Public/IRemoteControlUIModule.h
--- a/Engine/Plugins/VirtualProduction/RemoteControl/Source/RemoteControlUI/Public/IRemoteControlUIModule.h
+++ b/Engine/Plugins/VirtualProduction/RemoteControl/Source/RemoteControlUI/Public/IRemoteControlUIModule.h
@@ -19,6 +19,7 @@ class IDetailCategoryBuilder;
 class IDetailLayoutBuilder;
 struct SRCPanelTreeNode;
 class URemoteControlPreset;
+class FPropertyPath;
 
 struct FRCColumnSizeData
 {
@@ -195,4 +196,14 @@ public:
 
 	/** The section of EditorPerProjectUserSettings in which to save filter settings */
 	static const FString SettingsIniSection;
+
+	/**
+	 * If a details panel is open, highlight a property specified by the path in it.
+	 */
+	virtual void HighlightPropertyInDetailsPanel(const FPropertyPath& Path) const = 0;
+
+	/**
+	 * Set the list of selected objects.
+	 */
+	virtual void SelectObjects(const TArray<UObject*>& Objects) const = 0;
 };
