Update binding context every time we rebind

diff --git a/Engine/Plugins/VirtualProduction/RemoteControl/Source/RemoteControl/Private/RemoteControlBinding.cpp b/Engine/Plugins/VirtualProduction/RemoteControl/Source/RemoteControl/Private/RemoteControlBinding.cpp
--- a/Engine/Plugins/VirtualProduction/RemoteControl/Source/RemoteControl/Private/RemoteControlBinding.cpp
+++ b/Engine/Plugins/VirtualProduction/RemoteControl/Source/RemoteControl/Private/RemoteControlBinding.cpp
@@ -195,23 +195,14 @@ void URemoteControlLevelDependantBinding::SetBoundObject(const TSoftObjectPtr<UO
 		
 		Name = EditorObject->GetName();
 		LastBoundObjectPath = InObject.ToSoftObjectPath();
-
-		const bool bShouldSetSubObjectContext = !BindingContext.HasValidSubObjectPath() && !EditorObject->IsA<AActor>();
 		
-		if (BindingContext.OwnerActorName.IsNone() || bShouldSetSubObjectContext || !GetDefault<URemoteControlSettings>()->bUseRebindingContext)
-		{
-			InitializeBindingContext(InObject.Get());
-		}
+		UpdateBindingContext(InObject.Get());
 	}
 }
 
 void URemoteControlLevelDependantBinding::SetBoundObject_OverrideContext(const TSoftObjectPtr<UObject>& InObject)
 {
 	SetBoundObject(InObject);
-	if (InObject)
-	{
-		InitializeBindingContext(InObject.Get());
-	}
 }
 
 void URemoteControlLevelDependantBinding::InitializeForNewLevel()
@@ -282,7 +273,7 @@ UObject* URemoteControlLevelDependantBinding::Resolve() const
 
 		if (BindingContext.OwnerActorName.IsNone() || (!Object->IsA<AActor>() && !BindingContext.HasValidSubObjectPath()))
 		{
-			InitializeBindingContext(Object);
+			UpdateBindingContext(Object);
 		}
 	}
 
@@ -405,7 +396,7 @@ TSoftObjectPtr<UObject> URemoteControlLevelDependantBinding::ResolveForCurrentWo
 	return nullptr;
 }
 
-void URemoteControlLevelDependantBinding::InitializeBindingContext(UObject* InObject) const
+void URemoteControlLevelDependantBinding::UpdateBindingContext(UObject* InObject) const
 {
 	if (InObject)
 	{
diff --git a/Engine/Plugins/VirtualProduction/RemoteControl/Source/RemoteControl/Public/RemoteControlBinding.h b/Engine/Plugins/VirtualProduction/RemoteControl/Source/RemoteControl/Public/RemoteControlBinding.h
--- a/Engine/Plugins/VirtualProduction/RemoteControl/Source/RemoteControl/Public/RemoteControlBinding.h
+++ b/Engine/Plugins/VirtualProduction/RemoteControl/Source/RemoteControl/Public/RemoteControlBinding.h
@@ -232,7 +232,7 @@ public:
 private:
 	TSoftObjectPtr<UObject> ResolveForCurrentWorld() const;
 
-	void InitializeBindingContext(UObject* InObject) const;
+	void UpdateBindingContext(UObject* InObject) const;
 
 
 private:
