Fix assertion when clicking unbound property

diff --git a/Engine/Plugins/VirtualProduction/RemoteControl/Source/RemoteControlUI/Private/UI/SRCPanelExposedEntitiesList.cpp b/Engine/Plugins/VirtualProduction/RemoteControl/Source/RemoteControlUI/Private/UI/SRCPanelExposedEntitiesList.cpp
--- a/Engine/Plugins/VirtualProduction/RemoteControl/Source/RemoteControlUI/Private/UI/SRCPanelExposedEntitiesList.cpp
+++ b/Engine/Plugins/VirtualProduction/RemoteControl/Source/RemoteControlUI/Private/UI/SRCPanelExposedEntitiesList.cpp
@@ -924,51 +924,55 @@ void SRCPanelExposedEntitiesList::OnSelectionChanged(TSharedPtr<SRCPanelTreeNode
 			if (RCField->GetStruct() == FRemoteControlProperty::StaticStruct())
 			{
 				TSharedPtr<FRemoteControlProperty> RCProp = StaticCastSharedPtr<FRemoteControlProperty>(RCField);
-				TSharedRef<FPropertyPath> PropertyPath = RCProp->FieldPathInfo.ToPropertyPath();
-
-				for (auto It = BoundObjects.CreateIterator(); It; ++It)
+				RCProp->FieldPathInfo.Resolve(RCProp->GetBoundObject());
+				if (RCProp->FieldPathInfo.IsResolved())
 				{
-					if (AActor* OwnerActor = (*It)->GetTypedOuter<AActor>())
-					{
-						UObject* BindingObject = *It;
+					TSharedRef<FPropertyPath> PropertyPath = RCProp->FieldPathInfo.ToPropertyPath();
 
-						// When we encounter a non-component object, 
-						if (!BindingObject->IsA<UActorComponent>())
+					for (auto It = BoundObjects.CreateIterator(); It; ++It)
+					{
+						if (AActor* OwnerActor = (*It)->GetTypedOuter<AActor>())
 						{
-							BoundObjects.Add(OwnerActor);
-							It.RemoveCurrent();
-						}
+							UObject* BindingObject = *It;
 
-						// --- Special NDisplay handling because of their customization ---
-						// Since display cluster config data is not created as a defaeult subobject, therefore we have no way of retrieving the CurrentConfigData property from the config object itself.
-						static FName DisplayClusterConfigDataClassName = "DisplayClusterConfigurationData";
-						if (BindingObject->GetClass()->GetFName() == DisplayClusterConfigDataClassName)
-						{
-							if (FProperty* Property = OwnerActor->GetClass()->FindPropertyByName("CurrentConfigData"))
+							// When we encounter a non-component object, 
+							if (!BindingObject->IsA<UActorComponent>())
+							{
+								BoundObjects.Add(OwnerActor);
+								It.RemoveCurrent();
+							}
+
+							// --- Special NDisplay handling because of their customization ---
+							// Since display cluster config data is not created as a defaeult subobject, therefore we have no way of retrieving the CurrentConfigData property from the config object itself.
+							static FName DisplayClusterConfigDataClassName = "DisplayClusterConfigurationData";
+							if (BindingObject->GetClass()->GetFName() == DisplayClusterConfigDataClassName)
 							{
-								// Append "CurrentConfigData" to the beginning of the path.
-								TSharedRef<FPropertyPath> NewPropertyPath = FPropertyPath::Create(Property);
-								for (int32 Index = 0; Index < PropertyPath->GetNumProperties(); Index++)
+								if (FProperty* Property = OwnerActor->GetClass()->FindPropertyByName("CurrentConfigData"))
 								{
-									NewPropertyPath->AddProperty(PropertyPath->GetPropertyInfo(Index));
-								}
+									// Append "CurrentConfigData" to the beginning of the path.
+									TSharedRef<FPropertyPath> NewPropertyPath = FPropertyPath::Create(Property);
+									for (int32 Index = 0; Index < PropertyPath->GetNumProperties(); Index++)
+									{
+										NewPropertyPath->AddProperty(PropertyPath->GetPropertyInfo(Index));
+									}
 
-								PropertyPath = NewPropertyPath;
+									PropertyPath = NewPropertyPath;
+								}
 							}
 						}
 					}
-				}
 
-				FRemoteControlUIModule::Get().SelectObjects(BoundObjects);
+					FRemoteControlUIModule::Get().SelectObjects(BoundObjects);
 				
-				FTimerHandle Handle;
-				FTimerDelegate Delegate = FTimerDelegate::CreateLambda(([PropertyPath]()
-					{
-						FRemoteControlUIModule::Get().HighlightPropertyInDetailsPanel(*PropertyPath);
-					}));
+					FTimerHandle Handle;
+					FTimerDelegate Delegate = FTimerDelegate::CreateLambda(([PropertyPath]()
+						{
+							FRemoteControlUIModule::Get().HighlightPropertyInDetailsPanel(*PropertyPath);
+						}));
 
-				// Needed because modifying the selection set is asynchronous.
-				GEditor->GetTimerManager()->SetTimer(Handle, Delegate, 0.1, 0, -1);
+					// Needed because modifying the selection set is asynchronous.
+					GEditor->GetTimerManager()->SetTimer(Handle, Delegate, 0.1, 0, -1);
+				}
 			}
 		}
 	}
