Improve time for opening remote control preset in huge maps or when it has a lot of exposed properties
Significantly improve OnObjectPropertyChanged callback time
Fix hang when recompiling materials
Add traces

diff --git a/Engine/Plugins/VirtualProduction/RemoteControl/Source/RemoteControl/Private/RemoteControlEntity.cpp b/Engine/Plugins/VirtualProduction/RemoteControl/Source/RemoteControl/Private/RemoteControlEntity.cpp
--- a/Engine/Plugins/VirtualProduction/RemoteControl/Source/RemoteControl/Private/RemoteControlEntity.cpp
+++ b/Engine/Plugins/VirtualProduction/RemoteControl/Source/RemoteControl/Private/RemoteControlEntity.cpp
@@ -8,6 +8,7 @@
 
 TArray<UObject*> FRemoteControlEntity::GetBoundObjects() const
 {
+	TRACE_CPUPROFILER_EVENT_SCOPE(FRemoteControlEntity::GetBoundObjects);
 	TArray<UObject*> ResolvedObjects;
 	ResolvedObjects.Reserve(Bindings.Num());
 	Algo::TransformIf(Bindings, ResolvedObjects,
diff --git a/Engine/Plugins/VirtualProduction/RemoteControl/Source/RemoteControl/Private/RemoteControlPreset.cpp b/Engine/Plugins/VirtualProduction/RemoteControl/Source/RemoteControl/Private/RemoteControlPreset.cpp
--- a/Engine/Plugins/VirtualProduction/RemoteControl/Source/RemoteControl/Private/RemoteControlPreset.cpp
+++ b/Engine/Plugins/VirtualProduction/RemoteControl/Source/RemoteControl/Private/RemoteControlPreset.cpp
@@ -1574,6 +1574,8 @@ void URemoteControlPreset::OnObjectPropertyChanged(UObject* Object, struct FProp
 		HandleDisplayClusterConfigChange(Object);
 		return;
 	}
+
+	TMap<URemoteControlBinding*, UObject*> BoundObjectsCache;
  
 	if (Event.Property == nullptr)
 	{
@@ -1582,9 +1584,34 @@ void URemoteControlPreset::OnObjectPropertyChanged(UObject* Object, struct FProp
 			// When no property is passed to OnObjectPropertyChanged (such as by LevelSnapshot->Restore()), let's assume they all changed since we don't have more context.
 			for (TSharedPtr<FRemoteControlProperty> Property : Registry->GetExposedEntities<FRemoteControlProperty>())
 			{
-				if (Property->GetBoundObjects().Contains(Object))
+				bool bPropertyModified = false;
+
+				for (TWeakObjectPtr<URemoteControlBinding> Binding : Property->GetBindings())
 				{
-					PerFrameModifiedProperties.Add(Property->GetId());
+					if (!Binding.IsValid())
+					{
+						continue;
+					}
+
+					if (UObject** BoundObject = BoundObjectsCache.Find(Binding.Get()))
+					{
+						if (Object == *BoundObject)
+						{
+							PerFrameModifiedProperties.Add(Property->GetId());
+							break;
+						}
+					}
+					else
+					{
+						UObject* ResolvedObject = Binding->Resolve();
+						BoundObjectsCache.Add(Binding.Get(), ResolvedObject);
+
+						if (Object == ResolvedObject)
+						{
+							PerFrameModifiedProperties.Add(Property->GetId());
+							break;
+						}
+					}
 				}
 			}
 		}
diff --git a/Engine/Plugins/VirtualProduction/RemoteControl/Source/RemoteControlUI/Private/RemoteControlUIModule.cpp b/Engine/Plugins/VirtualProduction/RemoteControl/Source/RemoteControlUI/Private/RemoteControlUIModule.cpp
--- a/Engine/Plugins/VirtualProduction/RemoteControl/Source/RemoteControlUI/Private/RemoteControlUIModule.cpp
+++ b/Engine/Plugins/VirtualProduction/RemoteControl/Source/RemoteControlUI/Private/RemoteControlUIModule.cpp
@@ -1120,6 +1120,8 @@ void FRemoteControlUIModule::RefreshPanels()
 
 TSharedPtr<SRCPanelTreeNode> FRemoteControlUIModule::GenerateEntityWidget(const FGenerateWidgetArgs& Args)
 {
+	TRACE_CPUPROFILER_EVENT_SCOPE(FRemoteControlUIModule::GenerateEntityWidget);
+
 	if (Args.Preset && Args.Entity)
 	{
 		const UScriptStruct* EntityType = Args.Preset->GetExposedEntityType(Args.Entity->GetId());
diff --git a/Engine/Plugins/VirtualProduction/RemoteControl/Source/RemoteControlUI/Private/UI/SRCPanelExposedEntitiesList.cpp b/Engine/Plugins/VirtualProduction/RemoteControl/Source/RemoteControlUI/Private/UI/SRCPanelExposedEntitiesList.cpp
--- a/Engine/Plugins/VirtualProduction/RemoteControl/Source/RemoteControlUI/Private/UI/SRCPanelExposedEntitiesList.cpp
+++ b/Engine/Plugins/VirtualProduction/RemoteControl/Source/RemoteControlUI/Private/UI/SRCPanelExposedEntitiesList.cpp
@@ -477,6 +477,12 @@ void SRCPanelExposedEntitiesList::Tick(const FGeometry& AllottedGeometry, const
 
 		bFilterApplicationRequested = false;
 	}
+
+	if (bRefreshRequested)
+	{
+		ProcessRefresh();
+		bRefreshRequested = false;
+	}
 }
 
 TSharedPtr<SRCPanelTreeNode> SRCPanelExposedEntitiesList::GetSelectedGroup() const
@@ -672,26 +678,7 @@ void SRCPanelExposedEntitiesList::OnObjectPropertyChange(UObject* InObject, FPro
 
 void SRCPanelExposedEntitiesList::Refresh()
 {
-	GenerateListWidgets();
-
-	RefreshGroups();
-
-	if (Preset.IsValid())
-	{
-		constexpr bool bForceMouseClick = true;
-
-		if (const FRemoteControlPresetGroup* SelectedGroup = Preset->Layout.GetGroup(CurrentlySelectedGroup))
-		{
-			GenerateListWidgets(*SelectedGroup);
-			SetSelection(FindGroupById(SelectedGroup->Id), bForceMouseClick);
-		}
-		else
-		{
-			const FRemoteControlPresetGroup& DefaultGroup = Preset->Layout.GetDefaultGroup();
-			GenerateListWidgets(DefaultGroup);
-			SetSelection(FindGroupById(DefaultGroup.Id), bForceMouseClick);
-		}
-	}
+	bRefreshRequested = true;
 }
 
 void SRCPanelExposedEntitiesList::TryRefreshingSearch(const FText& InSearchText, bool bApplyFilter)
@@ -746,6 +733,8 @@ void SRCPanelExposedEntitiesList::TryRefreshingSearch(const FText& InSearchText,
 
 void SRCPanelExposedEntitiesList::GenerateListWidgets()
 {
+	TRACE_CPUPROFILER_EVENT_SCOPE(SRCPanelExposedEntitiesList::GenerateListWidgets);
+
 	FieldWidgetMap.Reset();
 
 	for (TWeakPtr<FRemoteControlEntity> WeakEntity : Preset->GetExposedEntities())
@@ -766,6 +755,8 @@ void SRCPanelExposedEntitiesList::GenerateListWidgets()
 
 void SRCPanelExposedEntitiesList::GenerateListWidgets(const FRemoteControlPresetGroup& FromGroup)
 {
+	TRACE_CPUPROFILER_EVENT_SCOPE(SRCPanelExposedEntitiesList::GenerateListWidgetsFromGroup);
+
 	FieldEntities.Reset();
 
 	if (FieldWidgetMap.IsEmpty())
@@ -924,6 +915,8 @@ void SRCPanelExposedEntitiesList::OnSelectionChanged(TSharedPtr<SRCPanelTreeNode
 			if (RCField->GetStruct() == FRemoteControlProperty::StaticStruct())
 			{
 				TSharedPtr<FRemoteControlProperty> RCProp = StaticCastSharedPtr<FRemoteControlProperty>(RCField);
+
+				// Resolve it to get the property path.
 				RCProp->FieldPathInfo.Resolve(RCProp->GetBoundObject());
 				if (RCProp->FieldPathInfo.IsResolved())
 				{
@@ -1673,6 +1666,32 @@ void SRCPanelExposedEntitiesList::OnWidgetRegistryRefreshed(const TArray<UObject
 	Refresh();
 }
 
+void SRCPanelExposedEntitiesList::ProcessRefresh()
+{
+	TRACE_CPUPROFILER_EVENT_SCOPE(SRCPanelExposedEntitiesList::Refresh);
+
+	GenerateListWidgets();
+
+	RefreshGroups();
+
+	if (Preset.IsValid())
+	{
+		constexpr bool bForceMouseClick = true;
+
+		if (const FRemoteControlPresetGroup* SelectedGroup = Preset->Layout.GetGroup(CurrentlySelectedGroup))
+		{
+			GenerateListWidgets(*SelectedGroup);
+			SetSelection(FindGroupById(SelectedGroup->Id), bForceMouseClick);
+		}
+		else
+		{
+			const FRemoteControlPresetGroup& DefaultGroup = Preset->Layout.GetDefaultGroup();
+			GenerateListWidgets(DefaultGroup);
+			SetSelection(FindGroupById(DefaultGroup.Id), bForceMouseClick);
+		}
+	}
+}
+
 bool FGroupDragEvent::IsDraggedFromSameGroup() const
 {
 	return DragOriginGroup->GetId() == DragTargetGroup->GetId();
diff --git a/Engine/Plugins/VirtualProduction/RemoteControl/Source/RemoteControlUI/Private/UI/SRCPanelExposedEntitiesList.h b/Engine/Plugins/VirtualProduction/RemoteControl/Source/RemoteControlUI/Private/UI/SRCPanelExposedEntitiesList.h
--- a/Engine/Plugins/VirtualProduction/RemoteControl/Source/RemoteControlUI/Private/UI/SRCPanelExposedEntitiesList.h
+++ b/Engine/Plugins/VirtualProduction/RemoteControl/Source/RemoteControlUI/Private/UI/SRCPanelExposedEntitiesList.h
@@ -185,6 +185,8 @@ private:
 	/** Called when the widget registry is refreshed by an underlying property generator. */
 	void OnWidgetRegistryRefreshed(const TArray<UObject*>& Objects);
 
+	void ProcessRefresh();
+
 private:
 	/** Holds the Groups list view. */
 	TSharedPtr<SListView<TSharedPtr<SRCPanelTreeNode>>> GroupsListView;
@@ -236,4 +238,6 @@ private:
 	static TSet<FName> DefaultProtocolColumns;
 	/** Holds identifier of the selected group. */
 	FGuid CurrentlySelectedGroup;
+
+	bool bRefreshRequested = false;
 };
\ No newline at end of file
diff --git a/Engine/Plugins/VirtualProduction/RemoteControl/Source/RemoteControlUI/Private/UI/SRCPanelFunctionPicker.cpp b/Engine/Plugins/VirtualProduction/RemoteControl/Source/RemoteControlUI/Private/UI/SRCPanelFunctionPicker.cpp
--- a/Engine/Plugins/VirtualProduction/RemoteControl/Source/RemoteControlUI/Private/UI/SRCPanelFunctionPicker.cpp
+++ b/Engine/Plugins/VirtualProduction/RemoteControl/Source/RemoteControlUI/Private/UI/SRCPanelFunctionPicker.cpp
@@ -27,7 +27,7 @@
 struct FRCFunctionPickerTreeNode
 {
 	virtual ~FRCFunctionPickerTreeNode() = default;
-	virtual const FString& GetName() const = 0;
+	virtual FText GetName() const = 0;
 	virtual bool IsFunctionNode() const = 0;
 	virtual UObject* GetObject() const = 0;
 	virtual UFunction* GetFunction() const = 0;
@@ -100,61 +100,59 @@ namespace FunctionPickerUtils
 
 	struct FRCFunctionNode : public FRCFunctionPickerTreeNode
 	{
-		FRCFunctionNode(UObject* InOwner, UFunction* InFunction)
-			: Owner(InOwner)
-			, Function(InFunction)
+		FRCFunctionNode(UFunction* InFunction)
+			: Function(InFunction)
 		{
-			if (Function.IsValid())
+			if (InFunction)
 			{
-				Name = Function->GetDisplayNameText().ToString();
+				Name = Function->GetDisplayNameText();
 			}
 		}
 
 		//~ Begin FRCBlueprintPickerTreeNode interface
-		virtual const FString& GetName() const override { return Name; }
+		virtual FText GetName() const override { return Name; }
 		virtual bool IsFunctionNode() const override { return true; }
 		virtual UFunction* GetFunction() const override { return Function.Get(); }
-		virtual UObject* GetObject() const override { return Owner.Get(); }
+		virtual UObject* GetObject() const override { return nullptr; }
 		virtual void GetChildNodes(TArray<TSharedPtr<FRCFunctionPickerTreeNode>>& OutNodes) const {}
 		//~ End FRCBlueprintPickerTreeNode interface
 
 	private:
-		/** The object that owns this node's function. */
-		TWeakObjectPtr<UObject> Owner = nullptr;
 		/** This node's function. */
 		TWeakObjectPtr<UFunction> Function = nullptr;
 		/** This node's function's name. */
-		FString Name;
+		FText Name;
 	};
 
 	struct FRCObjectNode : public FRCFunctionPickerTreeNode
 	{
-		FRCObjectNode(UObject* InObject, const TArray<UFunction*>& InFunctions)
+		FRCObjectNode(UObject* InObject, const TArray<TSharedPtr<FRCFunctionPickerTreeNode>>& InFunctionNodes)
 			: WeakObject(InObject)
+			, ChildNodes(InFunctionNodes)
 		{
+			TRACE_CPUPROFILER_EVENT_SCOPE(FRCObjectNode::FRCObjectNode());
+
 			if (ensure(InObject && InObject->GetFName().IsValid() && InObject->GetFName() != NAME_None))
 			{
+				FString NameString;
+
 				if (InObject->IsA<AActor>())
 				{
-					Name = Cast<AActor>(InObject)->GetActorLabel();
+					NameString = Cast<AActor>(InObject)->GetActorLabel();
 				}
 				else
 				{
-					Name = InObject->HasAnyFlags(RF_ClassDefaultObject | RF_ArchetypeObject) || InObject->IsA<USubsystem>() ? InObject->GetClass()->GetName() : InObject->GetName();
-					Name.RemoveFromStart(DefaultPrefix, ESearchCase::CaseSensitive);
-					Name.RemoveFromEnd(CPostfix, ESearchCase::CaseSensitive);
+					NameString = InObject->HasAnyFlags(RF_ClassDefaultObject | RF_ArchetypeObject) || InObject->IsA<USubsystem>() ? InObject->GetClass()->GetName() : InObject->GetName();
+					NameString.RemoveFromStart(DefaultPrefix, ESearchCase::CaseSensitive);
+					NameString.RemoveFromEnd(CPostfix, ESearchCase::CaseSensitive);
 				}
 
-				for (UFunction* Function : InFunctions)
-				{
-					TSharedPtr<FRCFunctionNode> FunctionNode = MakeShared<FRCFunctionNode>(InObject, Function);
-					ChildNodes.Add(FunctionNode);
-				}
+				Name = FText::FromString(NameString);
 			}
 		}
 
 		//~ Begin FRCBlueprintPickerTreeNode interface
-		virtual const FString& GetName() const override { return Name; }
+		virtual FText GetName() const override { return Name; }
 		virtual bool IsFunctionNode() const override { return false; }
 		virtual UObject* GetObject() const override { return WeakObject.Get(); }
 		virtual UFunction* GetFunction() const override { return nullptr; }
@@ -170,7 +168,7 @@ namespace FunctionPickerUtils
 		/** This node's child functions. */
 		mutable TArray<TSharedPtr<FRCFunctionPickerTreeNode>> ChildNodes;
 		/** This object's name. */
-		FString Name;
+		FText Name;
 	};
 }
 
@@ -192,6 +190,8 @@ void SRCPanelFunctionPicker::Construct(const FArguments& InArgs)
 	OnSelectFunction = InArgs._OnSelectFunction;
 	check(OnSelectFunction.IsBound());
 
+	static FText InvalidNameText = LOCTEXT("InvalidEntryText", "INVALID_ENTRY");
+
 	ChildSlot
 	[
 		SNew(SBox)
@@ -201,12 +201,12 @@ void SRCPanelFunctionPicker::Construct(const FArguments& InArgs)
 			SAssignNew(ObjectsTreeView, SSearchableTreeView<TSharedPtr<FRCFunctionPickerTreeNode>>)
 			.Items(&ObjectNodes)
 			.OnGetChildren_Lambda(OnGetChildren)
-			.OnGetDisplayName_Lambda([](TSharedPtr<FRCFunctionPickerTreeNode> InEntry) { return InEntry ? InEntry->GetName() : TEXT("INVALID_ENTRY"); })
+			.OnGetDisplayName_Lambda([](TSharedPtr<FRCFunctionPickerTreeNode> InEntry) { return InEntry ? InEntry->GetName() : InvalidNameText; })
 			.IsSelectable_Lambda([](TSharedPtr<FRCFunctionPickerTreeNode> TreeNode) { return TreeNode ? TreeNode->IsFunctionNode() : false; })
 			.OnItemSelected_Lambda(
-				[this] (TSharedPtr<FRCFunctionPickerTreeNode> TreeNode) 
+				[this] (TSharedPtr<FRCFunctionPickerTreeNode> TreeNode, TSharedPtr<FRCFunctionPickerTreeNode> ParentNode)
 				{
-					UObject* Owner = TreeNode->GetObject();
+					UObject* Owner = ParentNode->GetObject();
 					UFunction* Function = TreeNode->GetFunction();
 					if (Owner && Function)
 					{
@@ -267,24 +267,53 @@ void SRCPanelFunctionPicker::Refresh()
 		Class = UObject::StaticClass();
 	}
 
-	auto MakeNode = [this, Filter](UObject* Object)
+	TMap<UClass*, TArray<TSharedPtr<FRCFunctionPickerTreeNode>>> CachedNodesPerClass;
+
+	auto MakeNode = [this, Filter, &CachedNodesPerClass](UObject* Object)
 	{
+		TRACE_CPUPROFILER_EVENT_SCOPE(SRCPanelFunctionPicker::MakeNode);
 		if (Filter(Object))
 		{
+			if (TArray<TSharedPtr<FRCFunctionPickerTreeNode>>* FunctionNodes = CachedNodesPerClass.Find(Object->GetClass()))
+			{
+				if (FunctionNodes->Num())
+				{
+					ObjectNodes.Add(MakeShared<FRCObjectNode>(Object, *FunctionNodes));
+				}
+				return;
+			}
+
+
 			TArray<UFunction*> Functions = GetExposableFunctions(Object->GetClass());
-			if (Functions.Num())
+			TArray<TSharedPtr<FRCFunctionPickerTreeNode>> FunctionNodes;
+			FunctionNodes.Reserve(FunctionNodes.Num());
+			for (UFunction* Function : Functions)
+			{
+				FunctionNodes.Add(MakeShared<FRCFunctionNode>(Function));
+			}
+
+			CachedNodesPerClass.Add(Object->GetClass(), FunctionNodes);
+
+			if (FunctionNodes.Num())
 			{
-				ObjectNodes.Add(MakeShared<FRCObjectNode>(Object, Functions));
+				ObjectNodes.Add(MakeShared<FRCObjectNode>(Object, FunctionNodes));
 			}
 		}
 	};
 
+
 	// When displaying functions for actors, only display actors in the current level.
 	if (Class->IsChildOf(AActor::StaticClass()) && GEditor)
 	{
 		TArray<AActor*> ActorList;
 		constexpr bool bIgnorePIE = false;
-		UGameplayStatics::GetAllActorsOfClass(GetWorld(bIgnorePIE), Class, ActorList);
+
+		{
+			TRACE_CPUPROFILER_EVENT_SCOPE(SRCPanelFunctionPicker::GetAllActorsOfClass);
+			UGameplayStatics::GetAllActorsOfClass(GetWorld(bIgnorePIE), Class, ActorList);
+		}
+
+		ObjectNodes.Reserve(ActorList.Num());
 
 		for (AActor* Actor : ActorList)
 		{
diff --git a/Engine/Plugins/VirtualProduction/RemoteControl/Source/RemoteControlUI/Private/UI/SRemoteControlPanel.cpp b/Engine/Plugins/VirtualProduction/RemoteControl/Source/RemoteControlUI/Private/UI/SRemoteControlPanel.cpp
--- a/Engine/Plugins/VirtualProduction/RemoteControl/Source/RemoteControlUI/Private/UI/SRemoteControlPanel.cpp
+++ b/Engine/Plugins/VirtualProduction/RemoteControl/Source/RemoteControlUI/Private/UI/SRemoteControlPanel.cpp
@@ -794,14 +794,13 @@ void SRemoteControlPanel::Shutdown()
 	NoneSelectedWidget.Reset();
 }
 
-void SRemoteControlPanel::PostUndo(bool bSuccess)
+void SRemoteControlPanel::Tick(const FGeometry& AllottedGeometry, const double InCurrentTime, const float InDeltaTime)
 {
-	Refresh();
-}
-
-void SRemoteControlPanel::PostRedo(bool bSuccess)
-{
-	Refresh();
+	if (bMaterialsCompiledThisFrame)
+	{
+		TriggerMaterialCompiledRefresh();
+		bMaterialsCompiledThisFrame = false;
+	}
 }
 
 bool SRemoteControlPanel::IsExposed(const FRCExposesPropertyArgs& InPropertyArgs)
@@ -1534,9 +1533,21 @@ void SRemoteControlPanel::UnregisterPanels()
 
 void SRemoteControlPanel::Refresh()
 {
-	BlueprintPicker->Refresh();
-	ActorFunctionPicker->Refresh();
-	SubsystemFunctionPicker->Refresh();
+	TRACE_CPUPROFILER_EVENT_SCOPE(SRemoteControlPanel::Refresh);
+
+	{
+		TRACE_CPUPROFILER_EVENT_SCOPE(SRemoteControlPanel::RefreshBlueprintPicker);
+		BlueprintPicker->Refresh();
+	}
+	{
+		TRACE_CPUPROFILER_EVENT_SCOPE(SRemoteControlPanel::RefreshActorFunctionPicker);
+		ActorFunctionPicker->Refresh();
+	}
+	{
+		TRACE_CPUPROFILER_EVENT_SCOPE(SRemoteControlPanel::RefreshSubsystemFunctionPicker);
+		SubsystemFunctionPicker->Refresh();
+	}
+
 	EntityList->Refresh();
 }
 
@@ -1844,6 +1855,11 @@ FReply SRemoteControlPanel::OnClickSettingsButton()
 }
 
 void SRemoteControlPanel::OnMaterialCompiled(UMaterialInterface* MaterialInterface)
+{
+	bMaterialsCompiledThisFrame = true;
+}
+
+void SRemoteControlPanel::TriggerMaterialCompiledRefresh()
 {
 	bool bTriggerRefresh = true;
 
@@ -1861,13 +1877,13 @@ void SRemoteControlPanel::OnMaterialCompiled(UMaterialInterface* MaterialInterfa
 			if (!ChildWindows[0]->GetOnWindowClosedEvent().IsBound())
 			{
 				ChildWindows[0]->GetOnWindowClosedEvent().AddLambda([WeakThis = TWeakPtr<SRemoteControlPanel>(SharedThis(this))](const TSharedRef<SWindow>&)
-				{
-					if (TSharedPtr<SRemoteControlPanel> Panel = WeakThis.Pin())
 					{
-						Panel->WidgetRegistry->Clear();
-						Panel->Refresh();
-					}
-				});
+						if (TSharedPtr<SRemoteControlPanel> Panel = WeakThis.Pin())
+						{
+							Panel->WidgetRegistry->Clear();
+							Panel->Refresh();
+						}
+					});
 			}
 		}
 	}
@@ -1879,6 +1895,7 @@ void SRemoteControlPanel::OnMaterialCompiled(UMaterialInterface* MaterialInterfa
 	}
 }
 
+
 void SRemoteControlPanel::RegisterDefaultToolBar()
 {
 	UToolMenus* ToolMenus = UToolMenus::Get();
diff --git a/Engine/Plugins/VirtualProduction/RemoteControl/Source/RemoteControlUI/Private/UI/SRemoteControlPanel.h b/Engine/Plugins/VirtualProduction/RemoteControl/Source/RemoteControlUI/Private/UI/SRemoteControlPanel.h
--- a/Engine/Plugins/VirtualProduction/RemoteControl/Source/RemoteControlUI/Private/UI/SRemoteControlPanel.h
+++ b/Engine/Plugins/VirtualProduction/RemoteControl/Source/RemoteControlUI/Private/UI/SRemoteControlPanel.h
@@ -2,7 +2,6 @@
 #pragma once
 
 #include "Engine/EngineTypes.h"
-#include "EditorUndoClient.h"
 #include "Engine/EngineTypes.h"
 #include "Filters/SFilterBar.h"
 #include "IDetailTreeNode.h"
@@ -55,7 +54,7 @@ DECLARE_DELEGATE_TwoParams(FOnLiveModeChange, TSharedPtr<SRemoteControlPanel> /*
  * UI representation of a remote control preset.
  * Allows a user to expose/unexpose properties and functions from actors and blueprint libraries.
  */
-class SRemoteControlPanel : public SCompoundWidget, public FSelfRegisteringEditorUndoClient
+class SRemoteControlPanel : public SCompoundWidget
 {
 	SLATE_BEGIN_ARGS(SRemoteControlPanel) {}
 		SLATE_EVENT(FOnLiveModeChange, OnLiveModeChange)
@@ -78,9 +77,9 @@ public:
 	~SRemoteControlPanel();
 	static void Shutdown();
 
-	//~ FEditorUndoClient interface
-	virtual void PostUndo(bool bSuccess) override;
-	virtual void PostRedo(bool bSuccess) override;
+	//~ Begin SWidget interface
+	virtual void Tick(const FGeometry& AllottedGeometry, const double InCurrentTime, const float InDeltaTime);
+	//~ End SWidget interface
 
 	/**
 	 * @return The preset represented by the panel.
@@ -251,6 +250,9 @@ private:
 
 	/** Handle triggers each time when any material has been recompiled */
 	void OnMaterialCompiled(class UMaterialInterface* MaterialInterface);
+
+	/* Refreshed called at max once per frame when a material is compiled. */
+	void TriggerMaterialCompiledRefresh();
 	
 	/** Registers default tool bar */
 	static void RegisterDefaultToolBar();
@@ -348,8 +350,6 @@ private:
 	bool bIsInLiveMode = false;
 	/** Whether the logic panel is enabled or not. */
 	bool bIsLogicPanelEnabled = false;
-	/** Whether objects need to be re-resolved because PIE Started or ended. */
-	bool bTriggerRefreshForPIE = false;
 	/** Delegate called when the live mode changes. */
 	FOnLiveModeChange OnLiveModeChange;
 	/** Holds the blueprint library picker */
@@ -434,6 +434,9 @@ private:
 	UPROPERTY(Transient)
 	TObjectPtr<UObject> LogicClipboardItem;
 
+	/** Keeps track of whether materials were compiled from the current frame. Used to limit the number of UI refresh to once per frame. */
+	bool bMaterialsCompiledThisFrame = false;
+
 public:
 
 	static const float MinimumPanelWidth;
diff --git a/Engine/Plugins/VirtualProduction/RemoteControl/Source/RemoteControlUI/Private/UI/SSearchableTreeView.h b/Engine/Plugins/VirtualProduction/RemoteControl/Source/RemoteControlUI/Private/UI/SSearchableTreeView.h
--- a/Engine/Plugins/VirtualProduction/RemoteControl/Source/RemoteControlUI/Private/UI/SSearchableTreeView.h
+++ b/Engine/Plugins/VirtualProduction/RemoteControl/Source/RemoteControlUI/Private/UI/SSearchableTreeView.h
@@ -20,8 +20,8 @@ template <typename ItemType>
 class SSearchableTreeView : public SCompoundWidget
 {
 public:
-	DECLARE_DELEGATE_OneParam(FOnItemSelected, ItemType /* Item */)
-	DECLARE_DELEGATE_RetVal_OneParam(FString, FOnGetDisplayName, ItemType /* Item */)
+	DECLARE_DELEGATE_TwoParams(FOnItemSelected, ItemType /* Item */, ItemType /* ItemParent */)
+	DECLARE_DELEGATE_RetVal_OneParam(FText, FOnGetDisplayName, ItemType /* Item */)
 	DECLARE_DELEGATE_TwoParams(FOnGetChildren, ItemType /* Item */, TArray<ItemType>& /* OutChildren */)
 	DECLARE_DELEGATE_RetVal_OneParam(bool, FIsSelectable, ItemType /* Item */)
 	
@@ -145,7 +145,7 @@ private:
 			[
 				SNew(STextBlock)
 				.HighlightText(this, &SSearchableTreeView::GetFilterHighlightText)
-				.Text(FText::FromString(OnGetDisplayName.Execute(MoveTemp(InObject))))
+				.Text(OnGetDisplayName.Execute(MoveTemp(InObject)))
 			]
 		];
 	}
@@ -208,7 +208,7 @@ private:
 		{
 			if (IsSelectable.IsBound() && IsSelectable.Execute(SelectedObject))
 			{
-				OnItemSelected.ExecuteIfBound(SelectedObject);
+				OnItemSelected.ExecuteIfBound(SelectedObject, GetParent(SelectedObject));
 				TreeView->ClearSelection();
 			}
 			else
@@ -236,7 +236,7 @@ private:
 	/** Create the row's display name using the delegate. */
 	void TransformElementToString(ItemType InObject, TArray<FString>& OutStrings)
 	{
-		OutStrings.Add(OnGetDisplayName.Execute(InObject));
+		OutStrings.Add(OnGetDisplayName.Execute(InObject).ToString());
 	}
 
 	/** Handler for getting the rows highlight text. */
@@ -270,6 +270,28 @@ private:
 		return FReply::Unhandled();
 	}
 
+	ItemType GetParent(const ItemType& InItem)
+	{
+		if (Items)
+		{
+			for (const ItemType& Item : *Items)
+			{
+				TArray<ItemType> Children;
+				OnGetChildrenDelegate.ExecuteIfBound(Item, Children);
+
+				for (const ItemType& Child : Children)
+				{
+					if (Child == InItem)
+					{
+						return Item;
+					}
+				}
+			}
+		}
+
+		return ItemType();
+	}
+
 private:
 	/** Item selected delegate. */
 	FOnItemSelected OnItemSelected;
diff --git a/Engine/Plugins/VirtualProduction/RemoteControl/Source/WebRemoteControl/Private/RemoteControlModels.h b/Engine/Plugins/VirtualProduction/RemoteControl/Source/WebRemoteControl/Private/RemoteControlModels.h
--- a/Engine/Plugins/VirtualProduction/RemoteControl/Source/WebRemoteControl/Private/RemoteControlModels.h
+++ b/Engine/Plugins/VirtualProduction/RemoteControl/Source/WebRemoteControl/Private/RemoteControlModels.h
@@ -5,7 +5,7 @@
 #include "Algo/Transform.h"
 #include "AssetRegistry/AssetData.h"
 #include "AssetRegistry/ARFilter.h"
-#include "Controller/RCController.h"
+#include "RCVirtualProperty.h"
 #include "GameFramework/Actor.h"
 #include "RemoteControlActor.h"
 #include "RemoteControlField.h"
