Add error message when field path cannot be resolved
Add edit condition widget to RC preset
Add field path to the field name tool tip

diff --git a/Engine/Plugins/VirtualProduction/RemoteControl/Source/RemoteControlUI/Private/UI/RCPanelWidgetRegistry.cpp b/Engine/Plugins/VirtualProduction/RemoteControl/Source/RemoteControlUI/Private/UI/RCPanelWidgetRegistry.cpp
--- a/Engine/Plugins/VirtualProduction/RemoteControl/Source/RemoteControlUI/Private/UI/RCPanelWidgetRegistry.cpp
+++ b/Engine/Plugins/VirtualProduction/RemoteControl/Source/RemoteControlUI/Private/UI/RCPanelWidgetRegistry.cpp
@@ -71,6 +71,8 @@ namespace WidgetRegistryUtils
 	/** Find a node by its name in a detail tree node hierarchy. */
 	TSharedPtr<IDetailTreeNode> FindNode(const TArray<TSharedRef<IDetailTreeNode>>& RootNodes, const FString& QualifiedPropertyName, ERCFindNodeMethod FindMethod)
 	{
+		TRACE_CPUPROFILER_EVENT_SCOPE(FRCPanelWidgetRegistry::FindNode);
+
 		for (const TSharedRef<IDetailTreeNode>& CategoryNode : RootNodes)
 		{
 			TSharedPtr<IDetailTreeNode> FoundNode = FindTreeNodeRecursive(CategoryNode, QualifiedPropertyName, FindMethod);
@@ -107,6 +109,8 @@ FRCPanelWidgetRegistry::~FRCPanelWidgetRegistry()
 
 TSharedPtr<IDetailTreeNode> FRCPanelWidgetRegistry::GetObjectTreeNode(UObject* InObject, const FString& InField, ERCFindNodeMethod InFindMethod)
 {
+	TRACE_CPUPROFILER_EVENT_SCOPE(FRCPanelWidgetRegistry::GetObjectTreeNode);
+
 	const TPair<TWeakObjectPtr<UObject>, FString> CacheKey{InObject, InField};
 	if (TWeakPtr<IDetailTreeNode>* Node = TreeNodeCache.Find(CacheKey))
 	{
@@ -179,6 +183,13 @@ void FRCPanelWidgetRegistry::Refresh(UObject* InObject)
 	{
 		(*Generator)->SetObjects({InObject});
 	}
+	else if (IsNDisplayObject(InObject))
+	{
+		if (TSharedPtr<IPropertyRowGenerator>* GeneratorFromParent = ObjectToRowGenerator.Find({ InObject->GetTypedOuter<AActor>() }))
+		{
+			(*GeneratorFromParent)->SetObjects({ InObject->GetTypedOuter<AActor>() });
+		}
+	}
 }
 
 void FRCPanelWidgetRegistry::Refresh(const TSharedPtr<FStructOnScope>& InStruct)
@@ -275,6 +286,8 @@ void FRCPanelWidgetRegistry::OnRowsRefreshed(TSharedPtr<IPropertyRowGenerator> G
 
 TSharedPtr<IPropertyRowGenerator> FRCPanelWidgetRegistry::CreateGenerator(UObject* InObject)
 {
+	TRACE_CPUPROFILER_EVENT_SCOPE(FRCPanelWidgetRegistry::CreateGenerator);
+
 	// Since we must keep many PRG objects alive in order to access the handle data, validating the nodes each tick is very taxing.
 	// We can override the validation with a lambda since the validation function in PRG is not necessary for our implementation
 	auto ValidationLambda = ([](const FRootPropertyNodeList& PropertyNodeList) { return true; });
diff --git a/Engine/Plugins/VirtualProduction/RemoteControl/Source/RemoteControlUI/Private/UI/RCPanelWidgetRegistry.h b/Engine/Plugins/VirtualProduction/RemoteControl/Source/RemoteControlUI/Private/UI/RCPanelWidgetRegistry.h
--- a/Engine/Plugins/VirtualProduction/RemoteControl/Source/RemoteControlUI/Private/UI/RCPanelWidgetRegistry.h
+++ b/Engine/Plugins/VirtualProduction/RemoteControl/Source/RemoteControlUI/Private/UI/RCPanelWidgetRegistry.h
@@ -77,7 +77,7 @@ public:
 private:
 
 	/** Special handlers for ndisplay root actor and its details customizations */
-	bool IsNDisplayObject(UObject* InObject, const FString& InField, ERCFindNodeMethod InFindMethod);
+	bool IsNDisplayObject(UObject* InObject, const FString& InField = FString(), ERCFindNodeMethod InFindMethod = ERCFindNodeMethod::Path);
 	TSharedPtr<IDetailTreeNode> FindNDisplayTreeNode(UObject* InObject, const FString& InField, ERCFindNodeMethod InFindMethod);
 	void OnRowsRefreshed(TSharedPtr<IPropertyRowGenerator> Generator);
 	/** Create a property row generator */
diff --git a/Engine/Plugins/VirtualProduction/RemoteControl/Source/RemoteControlUI/Private/UI/SRCPanelExposedEntity.cpp b/Engine/Plugins/VirtualProduction/RemoteControl/Source/RemoteControlUI/Private/UI/SRCPanelExposedEntity.cpp
--- a/Engine/Plugins/VirtualProduction/RemoteControl/Source/RemoteControlUI/Private/UI/SRCPanelExposedEntity.cpp
+++ b/Engine/Plugins/VirtualProduction/RemoteControl/Source/RemoteControlUI/Private/UI/SRCPanelExposedEntity.cpp
@@ -34,10 +34,12 @@
 #include "Widgets/Input/SCheckBox.h"
 #include "Widgets/Layout/SBox.h"
 #include "Widgets/Layout/SBorder.h"
+#include "Widgets/Notifications/SPopUpErrorText.h"
 #include "Widgets/SBoxPanel.h"
 #include "Widgets/Text/STextBlock.h"
 #include "Widgets/Text/SInlineEditableTextBlock.h"
 
+
 #define LOCTEXT_NAMESPACE "RemoteControlPanel"
 
 namespace RebindingUtils
@@ -155,6 +157,11 @@ void SRCPanelExposedEntity::Initialize(const FGuid& InEntityId, URemoteControlPr
 			CachedLabel = RCEntity->GetLabel();
 			CachedBindingPath = *BindingPath;
 
+			if (RCEntity->GetStruct() == FRemoteControlProperty::StaticStruct())
+			{
+				CachedFieldPath = StaticCastSharedPtr<FRemoteControlProperty>(RCEntity)->FieldPathInfo.ToString();
+			}
+
 			FName OwnerFName;
 			// If the binding is valid, display the actor label if possible
 			if (UObject* Object = RCEntity->GetBoundObject())
@@ -313,9 +320,25 @@ TSharedRef<SWidget> SRCPanelExposedEntity::CreateRebindAllPropertiesForActorMenu
 		];
 }
 
-TSharedRef<SWidget> SRCPanelExposedEntity::CreateInvalidWidget()
+TSharedRef<SWidget> SRCPanelExposedEntity::CreateInvalidWidget(const FText& InErrorText)
 {
+	TRACE_CPUPROFILER_EVENT_SCOPE(SRCPanelExposedEntity::CreateInvalidWidget);
+
+	TSharedRef<SPopupErrorText> ErrorText = SNew(SPopupErrorText)
+		.Visibility(InErrorText.IsEmpty() ? EVisibility::Collapsed : EVisibility::Visible);
+
+	ErrorText->SetError(InErrorText);
+
 	return SNew(SHorizontalBox)
+		+ SHorizontalBox::Slot()
+		.AutoWidth()
+		[
+			SNew(SBox)
+			.WidthOverride(20.0)
+			[
+				ErrorText
+			]
+		]
 		+ SHorizontalBox::Slot()
 		.AutoWidth()
 		[
@@ -325,15 +348,15 @@ TSharedRef<SWidget> SRCPanelExposedEntity::CreateInvalidWidget()
 				SNew(STextBlock)
 				.Text(LOCTEXT("RebindLabel", "Rebind"))
 			]
-			.MenuContent()
-			[
-				SNew(SBox)
+			.OnGetMenuContent_Lambda([this]() 
+			{ 
+				return SNew(SBox)
 				.MaxDesiredHeight(400.0f)
 				.WidthOverride(300.0f)
 				[
 					CreateRebindMenuContent()
-				]
-			]
+				];
+			})
 		];
 }
 
@@ -344,6 +367,8 @@ EVisibility SRCPanelExposedEntity::GetVisibilityAccordingToLiveMode(EVisibility
 
 TSharedRef<SWidget> SRCPanelExposedEntity::CreateRebindMenuContent()
 {
+	TRACE_CPUPROFILER_EVENT_SCOPE(SRCPanelExposedEntity::CreateRebindMenuContent);
+
 	FSceneOutlinerModule& SceneOutlinerModule = FModuleManager::Get().LoadModuleChecked<FSceneOutlinerModule>("SceneOutliner");
 	FSceneOutlinerInitializationOptions Options;
 	Options.Filters = MakeShared<FSceneOutlinerFilters>();
@@ -431,7 +456,7 @@ bool SRCPanelExposedEntity::IsActorSelectable(const AActor* Actor) const
 	return false;
 }
 
-TSharedRef<SWidget> SRCPanelExposedEntity::CreateEntityWidget(TSharedPtr<SWidget> ValueWidget, TSharedPtr<SWidget> ResetWidget, const FText& OptionalWarningMessage)
+TSharedRef<SWidget> SRCPanelExposedEntity::CreateEntityWidget(TSharedPtr<SWidget> ValueWidget, TSharedPtr<SWidget> ResetWidget, const FText& OptionalWarningMessage, TSharedRef<SWidget> EditConditionWidget)
 {
 	FMakeNodeWidgetArgs Args;
 
@@ -495,11 +520,17 @@ TSharedRef<SWidget> SRCPanelExposedEntity::CreateEntityWidget(TSharedPtr<SWidget
             .Text(FEditorFontGlyphs::Exclamation_Triangle)
 		]
 		+ SHorizontalBox::Slot()
+		.AutoWidth()
+		[
+			EditConditionWidget
+		]
+		+ SHorizontalBox::Slot()
 		.Padding(4.f, 0.f, 2.0f, 0.f)
 		.AutoWidth()
 		[
 			SAssignNew(NameTextBox, SInlineEditableTextBlock)
 			.Text(FText::FromName(CachedLabel))
+			.ToolTipText(FText::FromString(CachedFieldPath))
 			.OnTextCommitted(this, &SRCPanelExposedEntity::OnLabelCommitted)
 			.OnVerifyTextChanged(this, &SRCPanelExposedEntity::OnVerifyItemLabelChanged)
 			.IsReadOnly_Lambda([this]() { return bLiveMode.Get(); })
diff --git a/Engine/Plugins/VirtualProduction/RemoteControl/Source/RemoteControlUI/Private/UI/SRCPanelExposedEntity.h b/Engine/Plugins/VirtualProduction/RemoteControl/Source/RemoteControlUI/Private/UI/SRCPanelExposedEntity.h
--- a/Engine/Plugins/VirtualProduction/RemoteControl/Source/RemoteControlUI/Private/UI/SRCPanelExposedEntity.h
+++ b/Engine/Plugins/VirtualProduction/RemoteControl/Source/RemoteControlUI/Private/UI/SRCPanelExposedEntity.h
@@ -34,13 +34,13 @@ protected:
 	void Initialize(const FGuid& InEntityId, URemoteControlPreset* InPreset, const TAttribute<bool>& InbLiveMode);
 	
 	/** Create a widget that displays the rebind button. */
-	TSharedRef<SWidget> CreateInvalidWidget();
+	TSharedRef<SWidget> CreateInvalidWidget(const FText& InErrorText = FText::GetEmpty());
 
 	/** Get the widget's visibility according to the panel's mode. */
 	EVisibility GetVisibilityAccordingToLiveMode(EVisibility NonEditModeVisibility) const;
 
 	/** Create an exposed entity widget with a drag handle and unexpose button. */
-	TSharedRef<SWidget> CreateEntityWidget(TSharedPtr<SWidget> ValueWidget, TSharedPtr<SWidget> ResetWidget = SNullWidget::NullWidget, const FText& OptionalWarningMessage = FText::GetEmpty());
+	TSharedRef<SWidget> CreateEntityWidget(TSharedPtr<SWidget> ValueWidget, TSharedPtr<SWidget> ResetWidget = SNullWidget::NullWidget, const FText& OptionalWarningMessage = FText::GetEmpty(), TSharedRef<SWidget> EditConditionWidget = SNullWidget::NullWidget);
 
 protected:
 	/** Id of the entity. */
@@ -59,6 +59,8 @@ protected:
 	bool bValidBinding = false;
 	/** Display name of the entity. */
 	FName CachedLabel;
+	/** Cached entity field path. */
+	FString CachedFieldPath;
 	/** Text to be highlighted while searching. */
 	TAttribute<FText> HighlightText;
 	/** Panel Style reference. */
diff --git a/Engine/Plugins/VirtualProduction/RemoteControl/Source/RemoteControlUI/Private/UI/SRCPanelExposedField.cpp b/Engine/Plugins/VirtualProduction/RemoteControl/Source/RemoteControlUI/Private/UI/SRCPanelExposedField.cpp
--- a/Engine/Plugins/VirtualProduction/RemoteControl/Source/RemoteControlUI/Private/UI/SRCPanelExposedField.cpp
+++ b/Engine/Plugins/VirtualProduction/RemoteControl/Source/RemoteControlUI/Private/UI/SRCPanelExposedField.cpp
@@ -24,6 +24,7 @@
 #include "Styling/RemoteControlStyles.h"
 #include "UObject/Object.h"
 #include "Widgets/Input/SButton.h"
+#include "Widgets/Input/SCheckBox.h"
 #include "Widgets/Layout/SBox.h"
 #include "Widgets/Masks/SRCProtocolMask.h"
 #include "Widgets/Layout/SBorder.h"
@@ -35,7 +36,7 @@
 
 namespace ExposedFieldUtils
 {
-	TSharedRef<SWidget> CreateNodeValueWidget(const TSharedPtr<IDetailTreeNode>& Node)
+	void CreateWidgets(const TSharedPtr<IDetailTreeNode>& Node, TSharedPtr<SWidget>& OutValueWidget, TSharedPtr<SWidget>& OutEditConditionWidget)
 	{
 		FNodeWidgets NodeWidgets = Node->CreateNodeWidgets();
 
@@ -59,7 +60,8 @@ namespace ExposedFieldUtils
 				];
 		}
 
-		return FieldWidget;
+		OutValueWidget = FieldWidget;
+		OutEditConditionWidget = NodeWidgets.EditConditionWidget;
 	}
 }
 
@@ -70,6 +72,9 @@ TSharedPtr<SRCPanelTreeNode> SRCPanelExposedField::MakeInstance(const FGenerateW
 
 void SRCPanelExposedField::Construct(const FArguments& InArgs, TWeakPtr<FRemoteControlField> InField, FRCColumnSizeData InColumnSizeData, TWeakPtr<FRCPanelWidgetRegistry> InWidgetRegistry)
 {
+	TRACE_CPUPROFILER_EVENT_SCOPE(SRCPanelExposedField::Construct);
+
+
 	WeakField = MoveTemp(InField);
 
 	ColumnSizeData = MoveTemp(InColumnSizeData);
@@ -318,15 +323,19 @@ void SRCPanelExposedField::GetBoundObjects(TSet<UObject*>& OutBoundObjects) cons
 
 TSharedRef<SWidget> SRCPanelExposedField::ConstructWidget()
 {
+	TRACE_CPUPROFILER_EVENT_SCOPE(SRCPanelExposedField::ConstructWidget);
+
 	if (TSharedPtr<FRemoteControlField> Field = WeakField.Pin())
 	{
 		// For the moment, just use the first object.
-		TArray<UObject*> Objects = Field->GetBoundObjects();
-		if (GetFieldType() == EExposedFieldType::Property && Objects.Num() > 0)
+		UObject* Object = Field->GetBoundObject();
+		if (Object && GetFieldType() == EExposedFieldType::Property)
 		{
 			if (TSharedPtr<FRCPanelWidgetRegistry> Registry = WidgetRegistry.Pin())
 			{
-				if (TSharedPtr<IDetailTreeNode> Node = Registry->GetObjectTreeNode(Objects[0], Field->FieldPathInfo.ToPathPropertyString(), ERCFindNodeMethod::Path))
+				TRACE_CPUPROFILER_EVENT_SCOPE(SRCPanelExposedField::BeforeGetObjectTreeNode);
+
+				if (TSharedPtr<IDetailTreeNode> Node = Registry->GetObjectTreeNode(Object, Field->FieldPathInfo.ToPathPropertyString(), ERCFindNodeMethod::Path))
 				{
 					TArray<TSharedRef<IDetailTreeNode>> ChildNodes;
 					Node->GetChildren(ChildNodes);
@@ -338,9 +347,14 @@ TSharedRef<SWidget> SRCPanelExposedField::ConstructWidget()
 					}
 
 					TSharedPtr<SHorizontalBox> FieldWidget = SNew(SHorizontalBox);
+					TSharedPtr<SWidget> ValueWidget = SNullWidget::NullWidget;
+					TSharedPtr<SWidget> EditConditionWidget = SNullWidget::NullWidget;
+
+					ExposedFieldUtils::CreateWidgets(Node, ValueWidget, EditConditionWidget);
+
 					FieldWidget->AddSlot()
 					[
-						ExposedFieldUtils::CreateNodeValueWidget(MoveTemp(Node))
+						ValueWidget.ToSharedRef()
 					];
 
 					TSharedPtr<IPropertyHandle> PropertyHandle = Node->CreatePropertyHandle();
@@ -408,7 +422,41 @@ TSharedRef<SWidget> SRCPanelExposedField::ConstructWidget()
 							];
 					}
 
-					return MakeFieldWidget(FieldWidget.ToSharedRef());
+					return MakeFieldWidget(FieldWidget.ToSharedRef(), EditConditionWidget.ToSharedRef());
+				}
+				else
+				{
+					FString PropertyName = Field->FieldPathInfo.ToPathPropertyString();
+					FString TargetName = Object->GetName();
+
+					FRCFieldPathInfo& FieldPath = Field->FieldPathInfo;
+
+					if (FieldPath.Resolve(Object))
+					{
+						FRCFieldResolvedData ResolvedData = FieldPath.GetResolvedData();
+						const FRCFieldPathSegment& LastSegment = FieldPath.GetFieldSegment(FieldPath.GetSegmentCount() - 1);
+						if (LastSegment.ArrayIndex != INDEX_NONE)
+						{
+							if (FArrayProperty* ArrayProperty = CastField<FArrayProperty>(ResolvedData.Field))
+							{
+								FScriptArrayHelper_InContainer Helper(ArrayProperty, ResolvedData.ContainerAddress);
+								int32 ArrayNum = Helper.Num();
+								if (!Helper.IsValidIndex(LastSegment.ArrayIndex))
+								{
+									PropertyName = LastSegment.ToString();
+
+									if (FieldPath.GetSegmentCount() > 1)
+									{
+										TargetName = Object->GetName() + TEXT(".") + FieldPath.ToString(FieldPath.GetSegmentCount() - 1);
+									}
+								}
+							}
+						}
+					}
+
+					FText ErrorText = FText::Format(LOCTEXT("ExposedPropertyInvalidErrorMessage", "Could not find property {0} on object {1}"), FText::FromString(PropertyName), FText::FromString(TargetName));
+					
+					return MakeFieldWidget(CreateInvalidWidget(ErrorText));
 				}
 			} 
 		}
@@ -419,8 +467,10 @@ TSharedRef<SWidget> SRCPanelExposedField::ConstructWidget()
 	return MakeFieldWidget(SNullWidget::NullWidget);
 }
 
-TSharedRef<SWidget> SRCPanelExposedField::MakeFieldWidget(const TSharedRef<SWidget>& InWidget)
+TSharedRef<SWidget> SRCPanelExposedField::MakeFieldWidget(const TSharedRef<SWidget>& InWidget, const TSharedRef<SWidget>& EditConditionWidget)
 {
+	TRACE_CPUPROFILER_EVENT_SCOPE(SRCPanelExposedField::MakeFieldWidget);
+
 	FText WarningMessage;
 
 	if (GetDefault<URemoteControlSettings>()->bDisplayInEditorOnlyWarnings)
@@ -461,7 +511,7 @@ TSharedRef<SWidget> SRCPanelExposedField::MakeFieldWidget(const TSharedRef<SWidg
 		WarningMessage = Builder.ToText();
 	}
 	
-	return CreateEntityWidget(InWidget, ResetButtonWidget.ToSharedRef(), WarningMessage);
+	return CreateEntityWidget(InWidget, ResetButtonWidget.ToSharedRef(), WarningMessage, EditConditionWidget);
 }
 
 void SRCPanelExposedField::ConstructPropertyWidget()
@@ -575,8 +625,8 @@ void SRCPanelFieldChildNode::Construct(const FArguments& InArgs, const TSharedRe
 	FNodeWidgets Widgets = InNode->CreateNodeWidgets();
 	FMakeNodeWidgetArgs Args;
 	Args.NameWidget = Widgets.NameWidget;
-	Args.ValueWidget = ExposedFieldUtils::CreateNodeValueWidget(InNode);
-
+	TSharedPtr<SWidget> EditConditionWidget;
+	ExposedFieldUtils::CreateWidgets(InNode, Args.ValueWidget, EditConditionWidget);
 	ChildSlot
 	[
 		MakeNodeWidget(Args)
