Add missing file from submit to fix CIS

diff --git a/Engine/Plugins/VirtualProduction/RemoteControl/Source/RemoteControlUI/Private/UI/SRCPanelExposedField.h b/Engine/Plugins/VirtualProduction/RemoteControl/Source/RemoteControlUI/Private/UI/SRCPanelExposedField.h
--- a/Engine/Plugins/VirtualProduction/RemoteControl/Source/RemoteControlUI/Private/UI/SRCPanelExposedField.h
+++ b/Engine/Plugins/VirtualProduction/RemoteControl/Source/RemoteControlUI/Private/UI/SRCPanelExposedField.h
@@ -77,7 +77,7 @@ private:
 	/** Construct a property widget. */
 	TSharedRef<SWidget> ConstructWidget();
 	/** Create the wrapper around the field value widget. */
-	TSharedRef<SWidget> MakeFieldWidget(const TSharedRef<SWidget>& InWidget);
+	TSharedRef<SWidget> MakeFieldWidget(const TSharedRef<SWidget>& InWidget, const TSharedRef<SWidget>& EditConditionWidget = SNullWidget::NullWidget);
 	/** Construct this field widget as a property widget. */
 	void ConstructPropertyWidget();
 	/** Construct this field widget as a function widget. */
