removed PRAGMA_DISABLE_OPTIMIZATION that was left in

diff --git a/Engine/Plugins/AI/HTNPlanner/Source/HTNTestSuite/Private/HTNTest.cpp b/Engine/Plugins/AI/HTNPlanner/Source/HTNTestSuite/Private/HTNTest.cpp
--- a/Engine/Plugins/AI/HTNPlanner/Source/HTNTestSuite/Private/HTNTest.cpp
+++ b/Engine/Plugins/AI/HTNPlanner/Source/HTNTestSuite/Private/HTNTest.cpp
@@ -11,8 +11,6 @@
 
 #define LOCTEXT_NAMESPACE "AITestSuite_HTNTest"
 
-PRAGMA_DISABLE_OPTIMIZATION
-
 // missing tests:
 //  testing invalid conditions
 //	testing specific scenarios with known plans 
@@ -631,6 +629,4 @@ IMPLEMENT_AI_INSTANT_TEST(FAITest_HTNOperatorsOfGeneratedPlan, "System.AI.HTN.Op
 //----------------------------------------------------------------------//
 typedef FAITest_SimpleComponentBasedTest<UMockHTNComponent> FAITest_HTNComponentTest;
 
-PRAGMA_ENABLE_OPTIMIZATION
-
 #undef LOCTEXT_NAMESPACE
diff --git a/Engine/Plugins/AI/MLAdapter/Source/MLAdapterTestSuite/Private/AgentTest.cpp b/Engine/Plugins/AI/MLAdapter/Source/MLAdapterTestSuite/Private/AgentTest.cpp
--- a/Engine/Plugins/AI/MLAdapter/Source/MLAdapterTestSuite/Private/AgentTest.cpp
+++ b/Engine/Plugins/AI/MLAdapter/Source/MLAdapterTestSuite/Private/AgentTest.cpp
@@ -8,9 +8,6 @@
 
 #define LOCTEXT_NAMESPACE "AITestSuite_MLAdapterTest"
 
-
-PRAGMA_DISABLE_OPTIMIZATION
-
 struct FMockAgentConfig_Actors : public FMLAdapterAgentConfig
 {
 	FMockAgentConfig_Actors()
@@ -144,6 +141,4 @@ IMPLEMENT_INSTANT_TEST_WITH_FIXTURE(FMLAdapterTest_AvatarSettingFixture, "System
 	return true;
 }
 
-PRAGMA_ENABLE_OPTIMIZATION 
-
 #undef LOCTEXT_NAMESPACE
diff --git a/Engine/Plugins/AI/MLAdapter/Source/MLAdapterTestSuite/Private/RPCServerTest.cpp b/Engine/Plugins/AI/MLAdapter/Source/MLAdapterTestSuite/Private/RPCServerTest.cpp
--- a/Engine/Plugins/AI/MLAdapter/Source/MLAdapterTestSuite/Private/RPCServerTest.cpp
+++ b/Engine/Plugins/AI/MLAdapter/Source/MLAdapterTestSuite/Private/RPCServerTest.cpp
@@ -7,9 +7,6 @@
 
 #define LOCTEXT_NAMESPACE "AITestSuite_MLAdapterTest"
 
-PRAGMA_DISABLE_OPTIMIZATION
-
-
 void FRPCTestBase::TearDown() 
 {
 	UMLAdapterManager::Get().GetOnAddClientFunctions().Remove(BindClientHandle);
@@ -141,8 +138,6 @@ struct FRPCTest_ServerBinds : public FRPCTest_BasicBinds
 };
 IMPLEMENT_AI_INSTANT_TEST(FRPCTest_ServerBinds, "System.AI.MLAdapter.RPC.ServerBinds")
 
-PRAGMA_ENABLE_OPTIMIZATION
-
 #undef LOCTEXT_NAMESPACE
 
 #endif // WITH_RPCLIB
diff --git a/Engine/Plugins/AI/MLAdapter/Source/MLAdapterTestSuite/Private/SessionTest.cpp b/Engine/Plugins/AI/MLAdapter/Source/MLAdapterTestSuite/Private/SessionTest.cpp
--- a/Engine/Plugins/AI/MLAdapter/Source/MLAdapterTestSuite/Private/SessionTest.cpp
+++ b/Engine/Plugins/AI/MLAdapter/Source/MLAdapterTestSuite/Private/SessionTest.cpp
@@ -12,8 +12,6 @@
 
 #define LOCTEXT_NAMESPACE "AITestSuite_MLAdapterTest"
 
-PRAGMA_DISABLE_OPTIMIZATION
-
 /**
  *	this fixture creates a game instance and can react to changes done to pawn.controller
  *	also, the session instance is created via the UMLAdapterManager so all other notifies should get through as well, most 
@@ -118,6 +116,4 @@ IMPLEMENT_INSTANT_TEST_WITH_FIXTURE(FMLAdapterTest_WithSession, "System.AI.MLAda
 	return true;
 }
 
-PRAGMA_ENABLE_OPTIMIZATION
-
 #undef LOCTEXT_NAMESPACE
diff --git a/Engine/Plugins/AI/MassAI/Source/MassAITestSuite/Private/MassAITest.cpp b/Engine/Plugins/AI/MassAI/Source/MassAITestSuite/Private/MassAITest.cpp
--- a/Engine/Plugins/AI/MassAI/Source/MassAITestSuite/Private/MassAITest.cpp
+++ b/Engine/Plugins/AI/MassAI/Source/MassAITestSuite/Private/MassAITest.cpp
@@ -5,8 +5,6 @@
 
 #define LOCTEXT_NAMESPACE "MassTest"
 
-PRAGMA_DISABLE_OPTIMIZATION
 
-PRAGMA_ENABLE_OPTIMIZATION
 
 #undef LOCTEXT_NAMESPACE
diff --git a/Engine/Plugins/Experimental/GPULightmass/Source/GPULightmass/Private/CompactingObjectPoolTest.cpp b/Engine/Plugins/Experimental/GPULightmass/Source/GPULightmass/Private/CompactingObjectPoolTest.cpp
--- a/Engine/Plugins/Experimental/GPULightmass/Source/GPULightmass/Private/CompactingObjectPoolTest.cpp
+++ b/Engine/Plugins/Experimental/GPULightmass/Source/GPULightmass/Private/CompactingObjectPoolTest.cpp
@@ -1,9 +1,6 @@
 ﻿// Copyright Epic Games, Inc. All Rights Reserved.
 
 #include "CoreMinimal.h"
-
-PRAGMA_DISABLE_OPTIMIZATION
-
 #include "CompactingObjectPool.h"
 
 struct Base
@@ -82,4 +79,3 @@ void RunCompactingObjectPoolTests()
 	ensure(!MyDynamicHandleToIntPart.PointsToValidObject());
 }
 
-PRAGMA_ENABLE_OPTIMIZATION
diff --git a/Engine/Plugins/Experimental/VirtualCamera/Source/VCamCore/Public/VCamTypes.h b/Engine/Plugins/Experimental/VirtualCamera/Source/VCamCore/Public/VCamTypes.h
--- a/Engine/Plugins/Experimental/VirtualCamera/Source/VCamCore/Public/VCamTypes.h
+++ b/Engine/Plugins/Experimental/VirtualCamera/Source/VCamCore/Public/VCamTypes.h
@@ -11,8 +11,6 @@
 
 class UVCamModifier;
 
-PRAGMA_DISABLE_OPTIMIZATION
-
 // Links a Modifier with a Name for use in a Modifier Stack
 USTRUCT()
 struct FModifierStackEntry
@@ -82,5 +80,3 @@ struct FModifierStackEntry
 		return !(*this == Other);
 	}
 };
-
-PRAGMA_ENABLE_OPTIMIZATION
\ No newline at end of file
diff --git a/Engine/Plugins/Runtime/AR/AppleAR/AppleARKit/Source/AppleARKit/Private/AppleARKitSystem.cpp b/Engine/Plugins/Runtime/AR/AppleAR/AppleARKit/Source/AppleARKit/Private/AppleARKitSystem.cpp
--- a/Engine/Plugins/Runtime/AR/AppleAR/AppleARKit/Source/AppleARKit/Private/AppleARKitSystem.cpp
+++ b/Engine/Plugins/Runtime/AR/AppleAR/AppleARKit/Source/AppleARKit/Private/AppleARKitSystem.cpp
@@ -1572,8 +1572,6 @@ void FAppleARKitSystem::OnSpawnARActor(AARActor* NewARActor, UARComponent* NewAR
 	}
 }
 
-
-PRAGMA_DISABLE_OPTIMIZATION
 bool FAppleARKitSystem::Run(UARSessionConfig* SessionConfig)
 {
 	TimecodeProvider = UAppleARKitSettings::GetTimecodeProvider();
@@ -1752,7 +1750,6 @@ bool FAppleARKitSystem::Run(UARSessionConfig* SessionConfig)
 	}
 	return true;
 }
-PRAGMA_ENABLE_OPTIMIZATION
 
 bool FAppleARKitSystem::IsRunning() const
 {
diff --git a/Engine/Plugins/Runtime/ZoneGraph/Source/ZoneGraphTestSuite/Private/ZoneGraphTest.cpp b/Engine/Plugins/Runtime/ZoneGraph/Source/ZoneGraphTestSuite/Private/ZoneGraphTest.cpp
--- a/Engine/Plugins/Runtime/ZoneGraph/Source/ZoneGraphTestSuite/Private/ZoneGraphTest.cpp
+++ b/Engine/Plugins/Runtime/ZoneGraph/Source/ZoneGraphTestSuite/Private/ZoneGraphTest.cpp
@@ -5,8 +5,5 @@
 
 #define LOCTEXT_NAMESPACE "ZoneGraphTest"
 
-PRAGMA_DISABLE_OPTIMIZATION
-
-PRAGMA_ENABLE_OPTIMIZATION
 
 #undef LOCTEXT_NAMESPACE
diff --git a/Engine/Plugins/Tests/RHITests/Source/RHITests/Public/RHIBufferTests.h b/Engine/Plugins/Tests/RHITests/Source/RHITests/Public/RHIBufferTests.h
--- a/Engine/Plugins/Tests/RHITests/Source/RHITests/Public/RHIBufferTests.h
+++ b/Engine/Plugins/Tests/RHITests/Source/RHITests/Public/RHIBufferTests.h
@@ -5,8 +5,6 @@
 #include "RHITestsCommon.h"
 #include "Async/ParallelFor.h"
 
-PRAGMA_DISABLE_OPTIMIZATION
-
 class FRHIBufferTests
 {
 	// Copies data in the specified vertex buffer back to the CPU, and passes a pointer to that data to the provided verification lambda.
@@ -442,5 +440,3 @@ public:
 		});
 	}
 };
-
-PRAGMA_ENABLE_OPTIMIZATION
\ No newline at end of file
diff --git a/Engine/Plugins/Tests/RuntimeTests/Source/RuntimeTests/Private/Slate/RichTextMarkupProcessingTest.cpp b/Engine/Plugins/Tests/RuntimeTests/Source/RuntimeTests/Private/Slate/RichTextMarkupProcessingTest.cpp
--- a/Engine/Plugins/Tests/RuntimeTests/Source/RuntimeTests/Private/Slate/RichTextMarkupProcessingTest.cpp
+++ b/Engine/Plugins/Tests/RuntimeTests/Source/RuntimeTests/Private/Slate/RichTextMarkupProcessingTest.cpp
@@ -6,9 +6,6 @@
 
 #include "Framework/Text/RichTextMarkupProcessing.h"
 
-// Disable optimization for RickTextMarkupProcessingTest as it compiles very slowly in development builds
-PRAGMA_DISABLE_OPTIMIZATION
-
 #define LOCTEXT_NAMESPACE "Slate.RichText.MarkupProcessingTest"
 
 IMPLEMENT_SIMPLE_AUTOMATION_TEST(FRTFProcessingTest, "System.Slate.RichText.MarkupProcessing", EAutomationTestFlags::ApplicationContextMask | EAutomationTestFlags::SmokeFilter)
@@ -380,5 +377,3 @@ bool FRTFProcessingTest::RunTest (const FString& Parameters)
 }
 
 #undef LOCTEXT_NAMESPACE 
-
-PRAGMA_ENABLE_OPTIMIZATION
diff --git a/Engine/Plugins/VirtualProduction/RemoteControl/Source/RemoteControl/Private/Tests/RemoteControlPropertyHandleTests.cpp b/Engine/Plugins/VirtualProduction/RemoteControl/Source/RemoteControl/Private/Tests/RemoteControlPropertyHandleTests.cpp
--- a/Engine/Plugins/VirtualProduction/RemoteControl/Source/RemoteControl/Private/Tests/RemoteControlPropertyHandleTests.cpp
+++ b/Engine/Plugins/VirtualProduction/RemoteControl/Source/RemoteControl/Private/Tests/RemoteControlPropertyHandleTests.cpp
@@ -11,8 +11,6 @@
 #include "RemoteControlPropertyHandleTestData.h"
 #include "RemoteControlPreset.h"
 
-PRAGMA_DISABLE_OPTIMIZATION
-
 #define PROP_NAME(Class, Name) GET_MEMBER_NAME_CHECKED(Class, Name)
 #define GET_TEST_PROP(PropName) URemoteControlAPITestObject::StaticClass()->FindPropertyByName(PROP_NAME(URemoteControlAPITestObject, PropName))
 
@@ -1085,5 +1083,3 @@ bool FRemoteControlAPIIntegrationTest::RunTest(const FString& Parameters)
 
 #undef GET_TEST_PROP
 #undef PROP_NAME
-
-PRAGMA_ENABLE_OPTIMIZATION
diff --git a/Engine/Plugins/VirtualProduction/RemoteControl/Source/RemoteControlLogic/Private/Tests/RemoteControlLogicTests.cpp b/Engine/Plugins/VirtualProduction/RemoteControl/Source/RemoteControlLogic/Private/Tests/RemoteControlLogicTests.cpp
--- a/Engine/Plugins/VirtualProduction/RemoteControl/Source/RemoteControlLogic/Private/Tests/RemoteControlLogicTests.cpp
+++ b/Engine/Plugins/VirtualProduction/RemoteControl/Source/RemoteControlLogic/Private/Tests/RemoteControlLogicTests.cpp
@@ -20,8 +20,6 @@
 #include "Misc/AutomationTest.h"
 #include "UObject/StrongObjectPtr.h"
 
-PRAGMA_DISABLE_OPTIMIZATION
-
 #define PROP_NAME(Class, Name) GET_MEMBER_NAME_CHECKED(Class, Name)
 #define GET_TEST_PROP(PropName) URemoteControlLogicTestData::StaticClass()->FindPropertyByName(PROP_NAME(URemoteControlLogicTestData, PropName))
  
@@ -166,7 +164,4 @@ bool FRemoteControlLogicTest::RunTest(const FString& Parameters)
 	TestEqual(TEXT("After empty controllers could should be equal 0"), Preset->GetNumControllers(), 0);
  
 	return true;
-}
- 
-PRAGMA_ENABLE_OPTIMIZATION
- 
\ No newline at end of file
+} 
\ No newline at end of file
diff --git a/Engine/Plugins/VirtualProduction/Rivermax/RivermaxCore/Source/RivermaxCore/Private/Tests/PTPTest.cpp b/Engine/Plugins/VirtualProduction/Rivermax/RivermaxCore/Source/RivermaxCore/Private/Tests/PTPTest.cpp
--- a/Engine/Plugins/VirtualProduction/Rivermax/RivermaxCore/Source/RivermaxCore/Private/Tests/PTPTest.cpp
+++ b/Engine/Plugins/VirtualProduction/Rivermax/RivermaxCore/Source/RivermaxCore/Private/Tests/PTPTest.cpp
@@ -12,7 +12,6 @@
 #include "Misc/AutomationTest.h"
 #include "RivermaxPTPUtils.h"
 
-PRAGMA_DISABLE_OPTIMIZATION
 #if WITH_DEV_AUTOMATION_TESTS
 
 IMPLEMENT_SIMPLE_AUTOMATION_TEST(FPTPTest, "Plugin.Rivermax.PTP", EAutomationTestFlags::ApplicationContextMask | EAutomationTestFlags::EngineFilter)
@@ -105,4 +104,3 @@ bool FPTPTest::RunTest(const FString& Parameters)
 
 
 #endif //WITH_DEV_AUTOMATION_TESTS
-PRAGMA_ENABLE_OPTIMIZATION
\ No newline at end of file
diff --git a/Engine/Source/Runtime/Experimental/ChaosVehicles/ChaosVehiclesCore/Public/SimModule/EngineModule.h b/Engine/Source/Runtime/Experimental/ChaosVehicles/ChaosVehiclesCore/Public/SimModule/EngineModule.h
--- a/Engine/Source/Runtime/Experimental/ChaosVehicles/ChaosVehiclesCore/Public/SimModule/EngineModule.h
+++ b/Engine/Source/Runtime/Experimental/ChaosVehicles/ChaosVehiclesCore/Public/SimModule/EngineModule.h
@@ -4,8 +4,6 @@
 
 #include "SimModule/TorqueSimModule.h"
 
-PRAGMA_DISABLE_OPTIMIZATION
-
 namespace Chaos
 {
 	struct FAllInputs;
@@ -77,5 +75,3 @@ namespace Chaos
 	};
 
 } // namespace Chaos
-
-PRAGMA_ENABLE_OPTIMIZATION
