Remote Control
- Embedded presets are no longer assigned to the "active" / "live" panel.

diff --git a/Engine/Plugins/VirtualProduction/RemoteControl/Source/RemoteControlUI/Private/RemoteControlUIModule.cpp b/Engine/Plugins/VirtualProduction/RemoteControl/Source/RemoteControlUI/Private/RemoteControlUIModule.cpp
--- a/Engine/Plugins/VirtualProduction/RemoteControl/Source/RemoteControlUI/Private/RemoteControlUIModule.cpp
+++ b/Engine/Plugins/VirtualProduction/RemoteControl/Source/RemoteControlUI/Private/RemoteControlUIModule.cpp
@@ -308,24 +308,34 @@ TSharedRef<SRemoteControlPanel> FRemoteControlUIModule::CreateRemoteControlPanel
 		Panel->SetLiveMode(bIsInLiveMode);
 	}
 
-	TSharedRef<SRemoteControlPanel> PanelRef = SAssignNew(WeakActivePanel, SRemoteControlPanel, Preset, ToolkitHost)
+	TSharedRef<SRemoteControlPanel> PanelRef = SNew(SRemoteControlPanel, Preset, ToolkitHost)
 		.OnLiveModeChange_Lambda(
-			[this](TSharedPtr<SRemoteControlPanel> Panel, bool bLiveMode) 
+			[this](TSharedPtr<SRemoteControlPanel> Panel, bool bLiveMode)
 			{
-				// Activating the live mode on a panel sets it as the active panel 
-				if (bLiveMode)
+				URemoteControlPreset* Preset = Panel->GetPreset();
+
+				if (!IsValid(Preset) || Preset->IsEmbeddedPreset() == false)
 				{
-					if (TSharedPtr<SRemoteControlPanel> ActivePanel = WeakActivePanel.Pin())
+					// Activating the live mode on a panel sets it as the active panel 
+					if (bLiveMode)
 					{
-						if (ActivePanel != Panel)
+						if (TSharedPtr<SRemoteControlPanel> ActivePanel = WeakActivePanel.Pin())
 						{
-							ActivePanel->SetLiveMode(true);
+							if (ActivePanel != Panel)
+							{
+								ActivePanel->SetLiveMode(true);
+							}
 						}
+						WeakActivePanel = MoveTemp(Panel);
 					}
-					WeakActivePanel = MoveTemp(Panel);
 				}
 			});
 
+	if (Preset->IsEmbeddedPreset() == false)
+	{
+		WeakActivePanel = PanelRef;
+	}
+
 	RegisteredRemoteControlPanels.Add(TWeakPtr<SRemoteControlPanel>(PanelRef));
 
 	// NOTE : Reregister the module with the detail panel when this panel is created.
