Fix allowing changes to protected proeprties without the AllowPrivateAccess Metadata

- Adds two warning widgets to the Panel which warn the users in case its activated
- Adds the ability to ignore protected and Getter/Setter Properties using certian functions.

diff --git a/Engine/Plugins/VirtualProduction/RemoteControl/Source/RemoteControl/Private/RemoteControlModule.cpp b/Engine/Plugins/VirtualProduction/RemoteControl/Source/RemoteControl/Private/RemoteControlModule.cpp
--- a/Engine/Plugins/VirtualProduction/RemoteControl/Source/RemoteControl/Private/RemoteControlModule.cpp
+++ b/Engine/Plugins/VirtualProduction/RemoteControl/Source/RemoteControl/Private/RemoteControlModule.cpp
@@ -21,6 +21,7 @@
 #include "RemoteControlInterceptionProcessor.h"
 #include "RemoteControlInstanceMaterial.h"
 #include "RemoteControlPreset.h"
+#include "RemoteControlSettings.h"
 #include "SceneInterface.h"
 #include "Serialization/PropertyMapStructDeserializerBackendWrapper.h"
 #include "StructDeserializer.h"
@@ -99,13 +100,17 @@ namespace RemoteControlUtil
 
 	bool IsPropertyAllowed(const FProperty* InProperty, ERCAccess InAccessType, bool bObjectInGamePackage)
 	{
+#if WITH_EDITOR
+		URemoteControlSettings* RCSettings = GetMutableDefault<URemoteControlSettings>();
+#endif
 		// The property is allowed to be accessed if it exists...
 		return InProperty &&
 				// it doesn't have exposed getter/setter that should be used instead
 #if WITH_EDITOR
-				(!InProperty->HasMetaData(RemoteControlUtil::NAME_BlueprintGetter) || !InProperty->HasMetaData(RemoteControlUtil::NAME_BlueprintSetter)) &&
-				// it isn't private or protected, except if AllowPrivateAccess is true
-				(!InProperty->HasAnyPropertyFlags(CPF_NativeAccessSpecifierProtected | CPF_NativeAccessSpecifierPrivate) || InProperty->GetBoolMetaData(RemoteControlUtil::NAME_AllowPrivateAccess)) &&
+				(!InProperty->HasMetaData(RemoteControlUtil::NAME_BlueprintGetter) || !InProperty->HasMetaData(RemoteControlUtil::NAME_BlueprintSetter) || RCSettings->bIgnoreGetterSetterCheck) &&
+				// it isn't private or protected, except if AllowPrivateAccess is true and if only for Protected whenever it should be ignored
+				(!InProperty->HasAnyPropertyFlags(CPF_NativeAccessSpecifierPrivate) || InProperty->GetBoolMetaData(RemoteControlUtil::NAME_AllowPrivateAccess)) &&
+				(!InProperty->HasAnyPropertyFlags(CPF_NativeAccessSpecifierProtected) || RCSettings->bIgnoreProtectedCheck || InProperty->GetBoolMetaData(RemoteControlUtil::NAME_AllowPrivateAccess)) &&	
 #endif
 				// it isn't blueprint private
 				!InProperty->HasAnyPropertyFlags(CPF_DisableEditOnInstance) &&
diff --git a/Engine/Plugins/VirtualProduction/RemoteControl/Source/RemoteControlCommon/Public/RemoteControlSettings.h b/Engine/Plugins/VirtualProduction/RemoteControl/Source/RemoteControlCommon/Public/RemoteControlSettings.h
--- a/Engine/Plugins/VirtualProduction/RemoteControl/Source/RemoteControlCommon/Public/RemoteControlSettings.h
+++ b/Engine/Plugins/VirtualProduction/RemoteControl/Source/RemoteControlCommon/Public/RemoteControlSettings.h
@@ -123,6 +123,15 @@ public:
 	UPROPERTY(config)
 	bool bUseRebindingContext = true;
 
+	UPROPERTY(config, EditAnywhere, Category = "Remote Control Preset", DisplayName = "Ignore Remote Control Protected Check")
+	bool bIgnoreProtectedCheck = false;
+
+	UPROPERTY(Config, EditAnywhere, Category = "Remote Control Preset", DisplayName = "Ignore Remote Control Getter/Setter Check")
+	bool bIgnoreGetterSetterCheck = false;
+
+	UPROPERTY(Config, EditAnywhere, Category = "Remote Control Preset", DisplayName = "Ignore Warnings")
+	bool bIgnoreWarnings = false;
+
 	UPROPERTY(config, EditAnywhere, Category = "Remote Control Web Interface", DisplayName = "Remote Control Web Interface Passphrase")
 	TArray<FRCPassphrase> Passphrases = {};
 };
diff --git a/Engine/Plugins/VirtualProduction/RemoteControl/Source/RemoteControlUI/Private/UI/SRemoteControlPanel.cpp b/Engine/Plugins/VirtualProduction/RemoteControl/Source/RemoteControlUI/Private/UI/SRemoteControlPanel.cpp
--- a/Engine/Plugins/VirtualProduction/RemoteControl/Source/RemoteControlUI/Private/UI/SRemoteControlPanel.cpp
+++ b/Engine/Plugins/VirtualProduction/RemoteControl/Source/RemoteControlUI/Private/UI/SRemoteControlPanel.cpp
@@ -514,6 +514,18 @@ void SRemoteControlPanel::Construct(const FArguments& InArgs, URemoteControlPres
 		.Padding(FMargin(5.f, 8.f, 5.f, 5.f))
 		[
 			CreateCPUThrottleWarning()
+		]
+		+SVerticalBox::Slot()
+		.AutoHeight()
+		.Padding(FMargin(5.f, 8.f, 5.f, 5.f))
+		[
+			CreateProtectedIgnoredWarning()
+		]
+		+SVerticalBox::Slot()
+		.AutoHeight()
+		.Padding(FMargin(5.f, 8.f, 5.f, 5.f))
+		[
+			CreateGetterSetterIgnoredWarning()
 		];
 
 	// Output Log Dock panel
@@ -967,6 +979,15 @@ FReply SRemoteControlPanel::OnClickDisableUseLessCPU() const
 	return FReply::Handled();
 }
 
+FReply SRemoteControlPanel::OnClickIgnoreWarnings() const
+{
+	URemoteControlSettings* Settings = GetMutableDefault<URemoteControlSettings>();
+	Settings->bIgnoreWarnings = true;
+	Settings->PostEditChange();
+	Settings->SaveConfig();
+	return FReply::Handled();
+}
+
 TSharedRef<SWidget> SRemoteControlPanel::CreateCPUThrottleWarning() const
 {
 	FProperty* PerformanceThrottlingProperty = FindFieldChecked<FProperty>(UEditorPerformanceSettings::StaticClass(), GET_MEMBER_NAME_CHECKED(UEditorPerformanceSettings, bThrottleCPUWhenNotForeground));
@@ -975,14 +996,63 @@ TSharedRef<SWidget> SRemoteControlPanel::CreateCPUThrottleWarning() const
 	FText PerformanceWarningText = FText::Format(LOCTEXT("RemoteControlPerformanceWarning", "Warning: The editor setting '{PropertyName}' is currently enabled\nThis will stop editor windows from updating in realtime while the editor is not in focus"), Arguments);
 
 	return SNew(SWarningOrErrorBox)
-		.Visibility_Lambda([]() { return GetDefault<UEditorPerformanceSettings>()->bThrottleCPUWhenNotForeground ? EVisibility::Visible : EVisibility::Collapsed; })
+		.Visibility_Lambda([]() { return GetDefault<UEditorPerformanceSettings>()->bThrottleCPUWhenNotForeground && !GetDefault<URemoteControlSettings>()->bIgnoreWarnings ? EVisibility::Visible : EVisibility::Collapsed; })
 		.MessageStyle(EMessageStyle::Warning)
 		.Message(PerformanceWarningText)
+		[
+			SNew(SHorizontalBox)
+			+ SHorizontalBox::Slot()
+			[
+				SNew(SButton)
+				.OnClicked(this, &SRemoteControlPanel::OnClickDisableUseLessCPU)
+				.TextStyle(FAppStyle::Get(), "DialogButtonText")
+				.Text(LOCTEXT("RemoteControlPerformanceWarningDisable", "Disable"))
+			]
+			+ SHorizontalBox::Slot()
+			[
+				SNew(SButton)
+				.OnClicked(this, &SRemoteControlPanel::OnClickIgnoreWarnings)
+				.TextStyle(FAppStyle::Get(), "DialogButtonText")
+				.Text(LOCTEXT("RemoteControlIgnoreWarningEnable", "Ignore"))
+			]
+		];
+}
+
+TSharedRef<SWidget> SRemoteControlPanel::CreateProtectedIgnoredWarning() const
+{
+	FProperty* IgnoreProtectedCheckProperty = FindFieldChecked<FProperty>(URemoteControlSettings::StaticClass(), GET_MEMBER_NAME_CHECKED(URemoteControlSettings, bIgnoreProtectedCheck));
+	FFormatNamedArguments Arguments;
+	Arguments.Add(TEXT("PropertyName"), IgnoreProtectedCheckProperty->GetDisplayNameText());
+	FText ProtectedIgnoredWarningText = FText::Format(LOCTEXT("RemoteControlProtectedIgnoredWarning", "Warning: The editor setting '{PropertyName}' is currently enabled\nThis will let properties with the protected flag be let through for certain checks"), Arguments);
+
+	return SNew(SWarningOrErrorBox)
+		.Visibility_Lambda([](){ return GetDefault<URemoteControlSettings>()->bIgnoreProtectedCheck && !GetDefault<URemoteControlSettings>()->bIgnoreWarnings ? EVisibility::Visible : EVisibility::Collapsed; })
+		.MessageStyle(EMessageStyle::Warning)
+		.Message(ProtectedIgnoredWarningText)
+		[
+			SNew(SButton)
+			.OnClicked(this, &SRemoteControlPanel::OnClickIgnoreWarnings)
+			.TextStyle(FAppStyle::Get(), "DialogButtonText")
+			.Text(LOCTEXT("RemoteControlIgnoreWarningEnable", "Ignore"))
+		];
+}
+
+TSharedRef<SWidget> SRemoteControlPanel::CreateGetterSetterIgnoredWarning() const
+{
+	FProperty* IgnoreGetterSetterCheckProperty = FindFieldChecked<FProperty>(URemoteControlSettings::StaticClass(), GET_MEMBER_NAME_CHECKED(URemoteControlSettings, bIgnoreGetterSetterCheck));
+	FFormatNamedArguments Arguments;
+	Arguments.Add(TEXT("PropertyName"), IgnoreGetterSetterCheckProperty->GetDisplayNameText());
+	FText ProtectedIgnoredWarningText = FText::Format(LOCTEXT("RemoteControlGetterSetterIgnoredWarning", "Warning: The editor setting '{PropertyName}' is currently enabled\nThis will let properties which have Getter/Setter let through for certain checks"), Arguments);
+
+	return SNew(SWarningOrErrorBox)
+		.Visibility_Lambda([](){ return GetDefault<URemoteControlSettings>()->bIgnoreGetterSetterCheck && !GetDefault<URemoteControlSettings>()->bIgnoreWarnings ? EVisibility::Visible : EVisibility::Collapsed; })
+		.MessageStyle(EMessageStyle::Warning)
+		.Message(ProtectedIgnoredWarningText)
 		[
 			SNew(SButton)
-			.OnClicked(this, &SRemoteControlPanel::OnClickDisableUseLessCPU)
+			.OnClicked(this, &SRemoteControlPanel::OnClickIgnoreWarnings)
 			.TextStyle(FAppStyle::Get(), "DialogButtonText")
-			.Text(LOCTEXT("RemoteControlPerformanceWarningDisable", "Disable"))
+			.Text(LOCTEXT("RemoteControlIgnoreWarningEnable", "Ignore"))
 		];
 }
 
diff --git a/Engine/Plugins/VirtualProduction/RemoteControl/Source/RemoteControlUI/Private/UI/SRemoteControlPanel.h b/Engine/Plugins/VirtualProduction/RemoteControl/Source/RemoteControlUI/Private/UI/SRemoteControlPanel.h
--- a/Engine/Plugins/VirtualProduction/RemoteControl/Source/RemoteControlUI/Private/UI/SRemoteControlPanel.h
+++ b/Engine/Plugins/VirtualProduction/RemoteControl/Source/RemoteControlUI/Private/UI/SRemoteControlPanel.h
@@ -197,9 +197,18 @@ private:
 	/** Handles disabling CPU throttling. */
 	FReply OnClickDisableUseLessCPU() const;
 
+	/** Handles ignoring Warnings. */
+	FReply OnClickIgnoreWarnings() const;
+
 	/** Creates a widget that warns the user when CPU throttling is enabled.  */
 	TSharedRef<SWidget> CreateCPUThrottleWarning() const;
 
+	/** Creates a widget that warns the user when the settings are currently ignoring the Protected flag is enabled. */
+	TSharedRef<SWidget> CreateProtectedIgnoredWarning() const;
+	
+	/** Creates a widget that warns the user when the settings are currently ignoring the Getter/Setter is enabled. */
+	TSharedRef<SWidget> CreateGetterSetterIgnoredWarning() const;
+
 	/** Create expose button, allowing to expose blueprints and actor functions. */
 	TSharedRef<SWidget> CreateExposeFunctionsButton();
 	
