Remote Control Logic - Path Behaviour - External Textures

- Add HTTP Support

diff --git a/Engine/Plugins/VirtualProduction/RemoteControl/Source/RemoteControlLogic/Private/Behaviour/Builtin/Path/RCSetAssetByPathBehaviour.cpp b/Engine/Plugins/VirtualProduction/RemoteControl/Source/RemoteControlLogic/Private/Behaviour/Builtin/Path/RCSetAssetByPathBehaviour.cpp
--- a/Engine/Plugins/VirtualProduction/RemoteControl/Source/RemoteControlLogic/Private/Behaviour/Builtin/Path/RCSetAssetByPathBehaviour.cpp
+++ b/Engine/Plugins/VirtualProduction/RemoteControl/Source/RemoteControlLogic/Private/Behaviour/Builtin/Path/RCSetAssetByPathBehaviour.cpp
@@ -12,11 +12,15 @@
 #include "Engine/Texture2D.h"
 #include "Engine/Texture2DDynamic.h"
 #include "Engine/World.h"
+#include "HttpManager.h"
+#include "HttpModule.h"
 #include "IRemoteControlModule.h"
 #include "IRemoteControlPropertyHandle.h"
+#include "Interfaces/IHttpResponse.h"
 #include "Kismet/KismetRenderingLibrary.h"
 #include "Materials/Material.h"
 #include "Materials/MaterialInstance.h"
+#include "Misc/FileHelper.h"
 #include "RCVirtualProperty.h"
 #include "RCVirtualPropertyContainer.h"
 #include "RemoteControlLogger.h"
@@ -378,8 +382,8 @@ bool URCSetAssetByPathBehaviour::SetExternalAsset(FString InExternalPath)
 		
 		return false;
 	});
-	const TSharedPtr<FRemoteControlProperty> RemoteControlProperty = ExposedPropertyPtr->Pin();
 
+	const TSharedPtr<FRemoteControlProperty> RemoteControlProperty = ExposedPropertyPtr->Pin();
 	if (!RemoteControlProperty)
 	{
 		FRemoteControlLogger::Get().Log(SetAssetByPathBehaviourHelpers::SetAssetByPathBehaviour, []
@@ -389,18 +393,7 @@ bool URCSetAssetByPathBehaviour::SetExternalAsset(FString InExternalPath)
 		return false;
 	}
 
-	UTexture2D* ImportedTexture = UKismetRenderingLibrary::ImportFileAsTexture2D(this, InExternalPath);
-	if (!ImportedTexture)
-	{
-		FRemoteControlLogger::Get().Log(SetAssetByPathBehaviourHelpers::SetAssetByPathBehaviour, [InExternalPath]
-		{
-			return FText::FromString(FString("Path Behaviour File not found: ") + *InExternalPath);
-		});
-		return false;
-	}
-
-	// Apply Texture
-	if (SetTextureAsset(RemoteControlProperty, ImportedTexture))
+	if (!SetTextureFromPath(RemoteControlProperty, InExternalPath))
 	{
 		FRemoteControlLogger::Get().Log(SetAssetByPathBehaviourHelpers::SetAssetByPathBehaviour, [InExternalPath]
 		{
@@ -432,9 +425,7 @@ bool URCSetAssetByPathBehaviour::SetTextureAsset(TSharedPtr<FRemoteControlProper
 	{
 		if (IRemoteControlModule::Get().ResolveObjectProperty(ObjectRef.Access, Object, ObjectRef.PropertyPathInfo, ObjectRef))
 		{
-			/**
-			 * Setting with Pointer and using Serialization of itself afterwards as a workaround with the Texture Asset not updating in the world.
-			 */
+			/** Setting with Pointer and using Serialization of itself afterwards as a workaround with the Texture Asset not updating in the world. */
 			FObjectPropertyBase* ObjectProperty = CastField<FObjectPropertyBase>(InRemoteControlPropertyPtr->GetProperty());
 			if (!ObjectProperty)
 			{
@@ -487,3 +478,43 @@ void URCSetAssetByPathBehaviour::RefreshPathArray()
 	}
 }
 
+bool URCSetAssetByPathBehaviour::SetTextureFromPath(TSharedPtr<FRemoteControlProperty> RCPropertyToSet, FString& FileName)
+{
+	// Local and Network Drives
+	if (FPaths::FileExists(FileName))
+	{
+		if (UTexture2D* LocalTexturePtr = UKismetRenderingLibrary::ImportFileAsTexture2D(this, FileName))
+		{
+			SetTextureAsset(RCPropertyToSet, LocalTexturePtr);
+			return true;
+		}
+	}
+	
+	// Http Link
+	TSharedRef<class IHttpRequest> HttpRequest = FHttpModule::Get().CreateRequest();
+	HttpRequest->OnProcessRequestComplete().BindUObject(this, &URCSetAssetByPathBehaviour::ReadFileHttpHandler, RCPropertyToSet);
+	HttpRequest->SetURL(FileName);
+	HttpRequest->SetVerb(TEXT("GET"));
+	return HttpRequest->ProcessRequest();
+}
+
+void URCSetAssetByPathBehaviour::ReadFileHttpHandler(FHttpRequestPtr HttpRequest, FHttpResponsePtr HttpResponse, bool bSucceeded, TSharedPtr<FRemoteControlProperty> InRCPropertyToSet)
+{
+	if (!HttpRequest.IsValid())
+	{
+		return;
+	}
+
+	if (bSucceeded && HttpResponse.IsValid())
+	{
+		UE_LOG(LogTemp, Verbose, TEXT("ReadFile request complete, Url=%s Code=%d"), *HttpRequest->GetURL(), HttpResponse->GetResponseCode());
+
+		const TArray<uint8> ResponseBuffer = HttpResponse->GetContent();
+		if (UTexture2D* HttpTexture = UKismetRenderingLibrary::ImportBufferAsTexture2D(this, ResponseBuffer))
+		{
+			SetTextureAsset(InRCPropertyToSet, HttpTexture);
+		}
+	}
+}
+
+
diff --git a/Engine/Plugins/VirtualProduction/RemoteControl/Source/RemoteControlLogic/Public/Behaviour/Builtin/Path/RCSetAssetByPathBehaviour.h b/Engine/Plugins/VirtualProduction/RemoteControl/Source/RemoteControlLogic/Public/Behaviour/Builtin/Path/RCSetAssetByPathBehaviour.h
--- a/Engine/Plugins/VirtualProduction/RemoteControl/Source/RemoteControlLogic/Public/Behaviour/Builtin/Path/RCSetAssetByPathBehaviour.h
+++ b/Engine/Plugins/VirtualProduction/RemoteControl/Source/RemoteControlLogic/Public/Behaviour/Builtin/Path/RCSetAssetByPathBehaviour.h
@@ -3,6 +3,7 @@
 #pragma once
 
 #include "Behaviour/RCBehaviour.h"
+#include "HttpModule.h"
 #include "RCVirtualProperty.h"
 #include "RemoteControlEntity.h"
 
@@ -21,6 +22,29 @@ namespace SetAssetByPathBehaviourHelpers
 	const FName SetAssetByPathBehaviour = FName(TEXT("Set Asset By Path"));
 }
 
+/** Info used to send request for a file */
+struct FPendingFileRequest
+{
+	/**
+	 * Constructor
+	 */
+	FPendingFileRequest(const FString& InFileName=FString(TEXT("")))
+		:  FileName(InFileName)
+	{
+	}
+
+	/**
+	 * Equality op
+	 */
+	inline bool operator==(const FPendingFileRequest& Other) const
+	{
+		return FileName == Other.FileName;
+	}
+
+	/** File being operated on by the pending request */
+	FString FileName;
+};
+
 /** Struct to help generate Widgts for the DetailsPanel of the Bahviour */
 USTRUCT()
 struct FRCSetAssetPath
@@ -108,4 +132,9 @@ private:
 	/** Auxiliary Function to apply a Texture onto a given Property */
 	bool SetTextureAsset(TSharedPtr<FRemoteControlProperty> InRemoteControlPropertyPtr, UTexture* InObject);
 
+	/** Auxiliary Function to apply a Texture onto a given Property */
+	bool SetTextureFromPath(TSharedPtr<FRemoteControlProperty> TexturePtr, FString& FileName);
+
+	/** Http Request Process Handler */
+	void ReadFileHttpHandler(FHttpRequestPtr HttpRequest, FHttpResponsePtr HttpResponse, bool bSucceeded, TSharedPtr<FRemoteControlProperty> InRCPropertyToSet);
 };
diff --git a/Engine/Plugins/VirtualProduction/RemoteControl/Source/RemoteControlLogic/RemoteControlLogic.Build.cs b/Engine/Plugins/VirtualProduction/RemoteControl/Source/RemoteControlLogic/RemoteControlLogic.Build.cs
--- a/Engine/Plugins/VirtualProduction/RemoteControl/Source/RemoteControlLogic/RemoteControlLogic.Build.cs
+++ b/Engine/Plugins/VirtualProduction/RemoteControl/Source/RemoteControlLogic/RemoteControlLogic.Build.cs
@@ -17,6 +17,7 @@ public class RemoteControlLogic : ModuleRules
 			new string[] {
 				"Cbor",
 				"Engine",
+				"HTTP",
 				"RemoteControl",
 				"Serialization",
 				"StructUtils"
