1. ID and Channel Fields Should to Be Pre-Populated With Text Field

2. Remote Control Preset - No existing icon for RCP actor

diff --git a/Engine/Plugins/VirtualProduction/RemoteControl/Source/RemoteControl/Private/RemoteControlPresetActor.cpp b/Engine/Plugins/VirtualProduction/RemoteControl/Source/RemoteControl/Private/RemoteControlPresetActor.cpp
new file mode 100644
--- /dev/null
+++ b/Engine/Plugins/VirtualProduction/RemoteControl/Source/RemoteControl/Private/RemoteControlPresetActor.cpp
@@ -0,0 +1,18 @@
+// Copyright Epic Games, Inc. All Rights Reserved.
+
+#include "RemoteControlPresetActor.h"
+
+ARemoteControlPresetActor::ARemoteControlPresetActor()
+{
+	RootComponent = CreateDefaultSubobject<USceneComponent>(TEXT("RootComponent0"));
+	check(RootComponent);
+
+#if WITH_EDITOR
+	RootComponent->bVisualizeComponent = true;
+#endif // WITH_EDITOR
+
+	RootComponent->SetMobility(EComponentMobility::Static);
+
+	SetHidden(true);
+	SetCanBeDamaged(false);
+}
diff --git a/Engine/Plugins/VirtualProduction/RemoteControl/Source/RemoteControl/Public/RemoteControlPresetActor.h b/Engine/Plugins/VirtualProduction/RemoteControl/Source/RemoteControl/Public/RemoteControlPresetActor.h
--- a/Engine/Plugins/VirtualProduction/RemoteControl/Source/RemoteControl/Public/RemoteControlPresetActor.h
+++ b/Engine/Plugins/VirtualProduction/RemoteControl/Source/RemoteControl/Public/RemoteControlPresetActor.h
@@ -10,12 +10,14 @@
 /**
  * Actor that wraps a remote control preset, allows linking a specific preset to a level.
  */
-UCLASS(hideCategories = (Rendering, Physics, LOD, Activation, Input, Replication, Collision, Actor, Cooking))
+UCLASS(HideCategories = (Rendering, Physics, LOD, Activation, Input, Replication, Collision, Actor, Cooking))
 class REMOTECONTROL_API ARemoteControlPresetActor : public AActor
 {
 public:
 	GENERATED_BODY()
 
+	ARemoteControlPresetActor();
+
 public:
 	UPROPERTY(EditAnywhere, Category="General")
 	TObjectPtr<URemoteControlPreset> Preset;
diff --git a/Engine/Plugins/VirtualProduction/RemoteControl/Source/RemoteControlProtocolWidgets/Private/RemoteControlProtocolWidgetsModule.cpp b/Engine/Plugins/VirtualProduction/RemoteControl/Source/RemoteControlProtocolWidgets/Private/RemoteControlProtocolWidgetsModule.cpp
--- a/Engine/Plugins/VirtualProduction/RemoteControl/Source/RemoteControlProtocolWidgets/Private/RemoteControlProtocolWidgetsModule.cpp
+++ b/Engine/Plugins/VirtualProduction/RemoteControl/Source/RemoteControlProtocolWidgets/Private/RemoteControlProtocolWidgetsModule.cpp
@@ -29,6 +29,16 @@ void FRemoteControlProtocolWidgetsModule::ShutdownModule()
 	OnActiveProtocolChanged().RemoveAll(this);
 }
 
+void FRemoteControlProtocolWidgetsModule::AddProtocolBinding(const FName InProtocolName)
+{
+	if (!RCProtocolBindingList.IsValid())
+	{
+		return;
+	}
+
+	RCProtocolBindingList->AddProtocolBinding(InProtocolName);
+}
+
 TSharedRef<SWidget> FRemoteControlProtocolWidgetsModule::GenerateDetailsForEntity(URemoteControlPreset* InPreset, const FGuid& InFieldId, const EExposedFieldType& InFieldType)
 {
 	check(InPreset);
diff --git a/Engine/Plugins/VirtualProduction/RemoteControl/Source/RemoteControlProtocolWidgets/Private/RemoteControlProtocolWidgetsModule.h b/Engine/Plugins/VirtualProduction/RemoteControl/Source/RemoteControlProtocolWidgets/Private/RemoteControlProtocolWidgetsModule.h
--- a/Engine/Plugins/VirtualProduction/RemoteControl/Source/RemoteControlProtocolWidgets/Private/RemoteControlProtocolWidgetsModule.h
+++ b/Engine/Plugins/VirtualProduction/RemoteControl/Source/RemoteControlProtocolWidgets/Private/RemoteControlProtocolWidgetsModule.h
@@ -7,6 +7,7 @@
 REMOTECONTROLPROTOCOLWIDGETS_API DECLARE_LOG_CATEGORY_EXTERN(LogRemoteControlProtocolWidgets, Log, All);
 
 class FProtocolBindingViewModel;
+class FProtocolEntityViewModel;
 class IRCProtocolBindingList;
 class URemoteControlPreset;
 
@@ -20,6 +21,7 @@ public:
 	//~ End IModuleInterface
 
 	//~ Begin IRemoteControlProtocolWidgetsModule Interface
+	virtual void AddProtocolBinding(const FName InProtocolName) override;
 	virtual TSharedRef<SWidget> GenerateDetailsForEntity(URemoteControlPreset* InPreset, const FGuid& InFieldId, const EExposedFieldType& InFieldType) override;
 	virtual void ResetProtocolBindingList() override;
 	virtual TSharedPtr<IRCProtocolBindingList> GetProtocolBindingList() const override;
diff --git a/Engine/Plugins/VirtualProduction/RemoteControl/Source/RemoteControlProtocolWidgets/Private/Widgets/SRCProtocolBindingList.cpp b/Engine/Plugins/VirtualProduction/RemoteControl/Source/RemoteControlProtocolWidgets/Private/Widgets/SRCProtocolBindingList.cpp
--- a/Engine/Plugins/VirtualProduction/RemoteControl/Source/RemoteControlProtocolWidgets/Private/Widgets/SRCProtocolBindingList.cpp
+++ b/Engine/Plugins/VirtualProduction/RemoteControl/Source/RemoteControlProtocolWidgets/Private/Widgets/SRCProtocolBindingList.cpp
@@ -5,6 +5,7 @@
 #include "Delegates/Delegate.h"
 #include "Editor.h"
 #include "EditorFontGlyphs.h"
+#include "Misc/MessageDialog.h"
 #include "Styling/AppStyle.h"
 #include "IRemoteControlProtocolModule.h"
 #include "RemoteControlProtocolWidgetsModule.h"
@@ -193,6 +194,25 @@ SRCProtocolBindingList::~SRCProtocolBindingList()
 	AwaitingProtocolEntities.Empty();
 }
 
+void SRCProtocolBindingList::AddProtocolBinding(const FName InProtocolName)
+{
+	if (!ViewModel.IsValid())
+	{
+		return;
+	}
+
+	FText CannotAddBinding;
+
+	if (!ViewModel->CanAddBinding(InProtocolName, CannotAddBinding))
+	{
+		FMessageDialog::Open(EAppMsgType::Ok, CannotAddBinding);
+
+		return;
+	}
+
+	ViewModel->AddBinding(InProtocolName);
+}
+
 TSharedRef<ITableRow> SRCProtocolBindingList::ConstructBindingWidget(TSharedPtr<IRCTreeNodeViewModel> InViewModel, const TSharedRef<STableViewBase>& InOwnerTable)
 {
 	check(InViewModel.IsValid());
diff --git a/Engine/Plugins/VirtualProduction/RemoteControl/Source/RemoteControlProtocolWidgets/Private/Widgets/SRCProtocolBindingList.h b/Engine/Plugins/VirtualProduction/RemoteControl/Source/RemoteControlProtocolWidgets/Private/Widgets/SRCProtocolBindingList.h
--- a/Engine/Plugins/VirtualProduction/RemoteControl/Source/RemoteControlProtocolWidgets/Private/Widgets/SRCProtocolBindingList.h
+++ b/Engine/Plugins/VirtualProduction/RemoteControl/Source/RemoteControlProtocolWidgets/Private/Widgets/SRCProtocolBindingList.h
@@ -28,6 +28,8 @@ public:
 	virtual ~SRCProtocolBindingList() override;
 
 	//~ Begin IRCProtocolBindingList Interface
+	virtual void AddProtocolBinding(const FName InProtocolName) override;
+
 	/** Get the set of entities which is awaiting state and waiting for binding. */
 	virtual FRemoteControlProtocolEntitySet GetAwaitingProtocolEntities() const override
 	{
diff --git a/Engine/Plugins/VirtualProduction/RemoteControl/Source/RemoteControlProtocolWidgets/Public/IRCProtocolBindingList.h b/Engine/Plugins/VirtualProduction/RemoteControl/Source/RemoteControlProtocolWidgets/Public/IRCProtocolBindingList.h
--- a/Engine/Plugins/VirtualProduction/RemoteControl/Source/RemoteControlProtocolWidgets/Public/IRCProtocolBindingList.h
+++ b/Engine/Plugins/VirtualProduction/RemoteControl/Source/RemoteControlProtocolWidgets/Public/IRCProtocolBindingList.h
@@ -18,6 +18,9 @@ public:
 	/** Virtual destructor */
 	virtual ~IRCProtocolBindingList() = default;
 
+	/** Adds a new protocol binding of specified type to the current view model. */
+	virtual void AddProtocolBinding(const FName InProtocolName) = 0;
+
 	/** Get the set of entities which is awaiting state and waiting for binding */
 	virtual FRemoteControlProtocolEntitySet GetAwaitingProtocolEntities() const = 0;
 
diff --git a/Engine/Plugins/VirtualProduction/RemoteControl/Source/RemoteControlProtocolWidgets/Public/IRemoteControlProtocolWidgetsModule.h b/Engine/Plugins/VirtualProduction/RemoteControl/Source/RemoteControlProtocolWidgets/Public/IRemoteControlProtocolWidgetsModule.h
--- a/Engine/Plugins/VirtualProduction/RemoteControl/Source/RemoteControlProtocolWidgets/Public/IRemoteControlProtocolWidgetsModule.h
+++ b/Engine/Plugins/VirtualProduction/RemoteControl/Source/RemoteControlProtocolWidgets/Public/IRemoteControlProtocolWidgetsModule.h
@@ -31,6 +31,9 @@ public:
 		return FModuleManager::LoadModuleChecked<IRemoteControlProtocolWidgetsModule>(ModuleName);
 	}
 
+	/** Adds a new protocol binding of specified type to the current view model. */
+	virtual void AddProtocolBinding(const FName InProtocolName) = 0;
+
 	/** Creates a widget for the given Preset Field and FieldType */
 	virtual TSharedRef<SWidget> GenerateDetailsForEntity(URemoteControlPreset* InPreset, const FGuid& InFieldId, const EExposedFieldType& InFieldType = EExposedFieldType::Invalid) = 0;
 
diff --git a/Engine/Plugins/VirtualProduction/RemoteControl/Source/RemoteControlUI/Private/UI/SRCPanelTreeNode.cpp b/Engine/Plugins/VirtualProduction/RemoteControl/Source/RemoteControlUI/Private/UI/SRCPanelTreeNode.cpp
--- a/Engine/Plugins/VirtualProduction/RemoteControl/Source/RemoteControlUI/Private/UI/SRCPanelTreeNode.cpp
+++ b/Engine/Plugins/VirtualProduction/RemoteControl/Source/RemoteControlUI/Private/UI/SRCPanelTreeNode.cpp
@@ -2,10 +2,13 @@
 
 #include "SRCPanelTreeNode.h"
 
+#include "IRemoteControlProtocolWidgetsModule.h"
+
 #include "Widgets/Layout/SBox.h"
 #include "Widgets/Layout/SBorder.h"
 #include "Widgets/SBoxPanel.h"
 #include "Widgets/SOverlay.h"
+#include "Widgets/Input/SEditableTextBox.h"
 #include "Widgets/Text/STextBlock.h"
 
 #define LOCTEXT_NAMESPACE "RemoteControlPanelNode"
@@ -21,7 +24,13 @@ TSet<FName> SRCPanelTreeNode::DefaultColumns = {
 
 TSharedRef<SWidget> SRCPanelTreeNode::GetProtocolWidget(const FName ForColumnName, const FName InProtocolName)
 {
-	return SNullWidget::NullWidget;
+	if (ForColumnName == RemoteControlPresetColumns::BindingStatus)
+	{
+		return SNullWidget::NullWidget;
+	}
+
+	return SNew(SEditableTextBox)
+		.OnTextChanged(this, &SRCPanelTreeNode::OnProtocolTextChanged, InProtocolName);
 }
 
 const bool SRCPanelTreeNode::HasProtocolExtension() const
@@ -177,6 +186,13 @@ void SRCPanelTreeNode::OnLeftColumnResized(float) const
 	// We do nothing here because it is handled by the column size data
 }
 
+void SRCPanelTreeNode::OnProtocolTextChanged(const FText& InText, const FName InProtocolName)
+{
+	IRemoteControlProtocolWidgetsModule& RCProtocolsWidgets = IRemoteControlProtocolWidgetsModule::Get();
+
+	RCProtocolsWidgets.AddProtocolBinding(InProtocolName);
+}
+
 float SRCPanelTreeNode::GetLeftColumnWidth() const
 {
 	const float Offset = GetRCType() == ENodeType::Group ? SplitterOffset : 0;
diff --git a/Engine/Plugins/VirtualProduction/RemoteControl/Source/RemoteControlUI/Private/UI/SRCPanelTreeNode.h b/Engine/Plugins/VirtualProduction/RemoteControl/Source/RemoteControlUI/Private/UI/SRCPanelTreeNode.h
--- a/Engine/Plugins/VirtualProduction/RemoteControl/Source/RemoteControlUI/Private/UI/SRCPanelTreeNode.h
+++ b/Engine/Plugins/VirtualProduction/RemoteControl/Source/RemoteControlUI/Private/UI/SRCPanelTreeNode.h
@@ -60,6 +60,7 @@ struct SRCPanelTreeNode : public SCompoundWidget, public IHasProtocolExtensibili
 	virtual TSharedRef<SWidget> GetProtocolWidget(const FName ForColumnName, const FName InProtocolName = NAME_None) override;
 	virtual const bool HasProtocolExtension() const override;
 	virtual const bool GetProtocolBindingsNum() const override;
+	virtual void OnProtocolTextChanged(const FText& InText, const FName InProtocolName);
 	virtual const bool SupportsProtocol(const FName& InProtocolName) const override;
 	//~ END : IHasProtocolExtensibility Interface
 
@@ -105,7 +106,7 @@ protected:
 private:
 	/** Stub handler for column resize callback to prevent the splitter from handling it internally.  */
 	void OnLeftColumnResized(float) const;
-
+	
 	//~ Wrappers around ColumnSizeData's delegate needed in order to offset the splitter for RC Groups. 
 	float GetLeftColumnWidth() const;
 	float GetRightColumnWidth() const;
diff --git a/Engine/Plugins/VirtualProduction/RemoteControl/Source/RemoteControlUI/Public/IHasProtocolExtensibility.h b/Engine/Plugins/VirtualProduction/RemoteControl/Source/RemoteControlUI/Public/IHasProtocolExtensibility.h
--- a/Engine/Plugins/VirtualProduction/RemoteControl/Source/RemoteControlUI/Public/IHasProtocolExtensibility.h
+++ b/Engine/Plugins/VirtualProduction/RemoteControl/Source/RemoteControlUI/Public/IHasProtocolExtensibility.h
@@ -25,6 +25,11 @@ public:
 	 * Returns true when the derived classes has two or more protocol bindings.
 	 */
 	virtual const bool GetProtocolBindingsNum() const = 0;
+	
+	/**
+	 * Callback for the editable text's OnTextChanged event.
+	 */
+	virtual void OnProtocolTextChanged(const FText& InText, const FName InProtocolName) = 0;
 
 	/**
 	 * Determines whether the derived class supports the given protocol.
