Remote Control Property Handle - Minor Change in Conditional

Change in the conditional to fix potential problems with different named structs.

diff --git a/Engine/Plugins/VirtualProduction/RemoteControl/Source/RemoteControl/Private/RemoteControlPropertyHandle.cpp b/Engine/Plugins/VirtualProduction/RemoteControl/Source/RemoteControl/Private/RemoteControlPropertyHandle.cpp
--- a/Engine/Plugins/VirtualProduction/RemoteControl/Source/RemoteControl/Private/RemoteControlPropertyHandle.cpp
+++ b/Engine/Plugins/VirtualProduction/RemoteControl/Source/RemoteControl/Private/RemoteControlPropertyHandle.cpp
@@ -448,7 +448,7 @@ FRemoteControlPropertyHandle::FRemoteControlPropertyHandle(const TSharedPtr<FRem
 	{
 		if (const UScriptStruct* Struct = Cast<UScriptStruct>(InProperty->GetOwnerUObject()))
 		{
-			FieldPath = FString::Printf(TEXT("%s.%s"), *Struct->GetName(), *ThisFieldPath);
+			FieldPath = InRCProperty->FieldPathInfo.ToPathPropertyString();
 		}
 		else
 		{
