Bind Behavior Not Working With Integers in Avalanche

diff --git a/Engine/Plugins/VirtualProduction/RemoteControl/Source/RemoteControlLogic/Private/Action/Bind/RCPropertyBindAction.cpp b/Engine/Plugins/VirtualProduction/RemoteControl/Source/RemoteControlLogic/Private/Action/Bind/RCPropertyBindAction.cpp
--- a/Engine/Plugins/VirtualProduction/RemoteControl/Source/RemoteControlLogic/Private/Action/Bind/RCPropertyBindAction.cpp
+++ b/Engine/Plugins/VirtualProduction/RemoteControl/Source/RemoteControlLogic/Private/Action/Bind/RCPropertyBindAction.cpp
@@ -151,7 +151,14 @@ void URCPropertyBindAction::Execute() const
 			// Numeric To Numeric
 			if (RemoteControlProperty->IsA(FNumericProperty::StaticClass()))
 			{
-				Handle->SetValue(NumericValue);
+				if (RemoteControlProperty->IsA(FFloatProperty::StaticClass()))
+				{
+					Handle->SetValue(NumericValue);
+				}
+				else if (RemoteControlProperty->IsA(FIntProperty::StaticClass()))
+				{
+					Handle->SetValue((int32)NumericValue);
+				}
 			}
 			// Numeric To Boolean
 			else if (RemoteControlProperty->IsA(FBoolProperty::StaticClass()))
diff --git a/Engine/Plugins/VirtualProduction/RemoteControl/Source/RemoteControlUI/Private/UI/Controller/SRCControllerPanelList.cpp b/Engine/Plugins/VirtualProduction/RemoteControl/Source/RemoteControlUI/Private/UI/Controller/SRCControllerPanelList.cpp
--- a/Engine/Plugins/VirtualProduction/RemoteControl/Source/RemoteControlUI/Private/UI/Controller/SRCControllerPanelList.cpp
+++ b/Engine/Plugins/VirtualProduction/RemoteControl/Source/RemoteControlUI/Private/UI/Controller/SRCControllerPanelList.cpp
@@ -686,7 +686,8 @@ void SRCControllerPanelList::CreateAutoBindForProperty(TSharedPtr<const FRemoteC
 			Reset();
 
 			// Step 3. Create Bind Behaviour and Bind to the property
-			CreateBindBehaviourAndAssignTo(NewController, RemoteControlProperty.ToSharedRef(), false);
+			constexpr bool bExecuteBind = true;
+			CreateBindBehaviourAndAssignTo(NewController, RemoteControlProperty.ToSharedRef(), bExecuteBind);
 		}
 	}
 }
