Fix holding reference to temporary TCHAR*
#rnx

diff --git a/Engine/Plugins/VirtualProduction/RemoteControl/Source/RemoteControl/Private/RemoteControlPreset.cpp b/Engine/Plugins/VirtualProduction/RemoteControl/Source/RemoteControl/Private/RemoteControlPreset.cpp
--- a/Engine/Plugins/VirtualProduction/RemoteControl/Source/RemoteControl/Private/RemoteControlPreset.cpp
+++ b/Engine/Plugins/VirtualProduction/RemoteControl/Source/RemoteControl/Private/RemoteControlPreset.cpp
@@ -496,15 +496,15 @@ TWeakPtr<FRemoteControlActor> URemoteControlPreset::ExposeActor(AActor* Actor, F
 	check(Actor);
 
 #if WITH_EDITOR
-	const TCHAR* DesiredName = Args.Label.IsEmpty() ? *Actor->GetActorLabel() : *Args.Label;
+	const FName UniqueLabel = Registry->GenerateUniqueLabel(Args.Label.IsEmpty() ? *Actor->GetActorLabel() : *Args.Label);
 #else
-	const TCHAR* DesiredName = Args.Label.IsEmpty() ? *Actor->GetName() : *Args.Label;
+	const FName UniqueLabel = Registry->GenerateUniqueLabel(Args.Label.IsEmpty() ? *Actor->GetName() : *Args.Label);
 #endif
 
 	FText LogText = FText::Format(LOCTEXT("ExposedActor", "Exposed actor ({0})"), FText::FromString(Actor->GetName()));
     FRemoteControlLogger::Get().Log(TEXT("RemoteControlPreset"), [Text = MoveTemp(LogText)](){ return Text; });
 
-	FRemoteControlActor RCActor{this, Registry->GenerateUniqueLabel(DesiredName), { FindOrAddBinding(Actor)} };
+	FRemoteControlActor RCActor{this, UniqueLabel, { FindOrAddBinding(Actor)} };
 	return StaticCastSharedPtr<FRemoteControlActor>(Expose(MoveTemp(RCActor), FRemoteControlActor::StaticStruct(), Args.GroupId));
 }
 
