- Fixed issue with changing an exposed color property from color picker not calling EndTransaction

- Fix issue with Duplicating RCP Controllers causing the template id to be copied over without generating a new one

- Fix issue with duplicating RCP causing the new duplicate RCP to have an empty LabelToIdCache map, until loaded again from disk.

diff --git a/Engine/Plugins/VirtualProduction/RemoteControl/Source/RemoteControl/Private/RCVirtualPropertyContainer.cpp b/Engine/Plugins/VirtualProduction/RemoteControl/Source/RemoteControl/Private/RCVirtualPropertyContainer.cpp
--- a/Engine/Plugins/VirtualProduction/RemoteControl/Source/RemoteControl/Private/RCVirtualPropertyContainer.cpp
+++ b/Engine/Plugins/VirtualProduction/RemoteControl/Source/RemoteControl/Private/RCVirtualPropertyContainer.cpp
@@ -101,6 +101,7 @@ URCVirtualPropertyInContainer* URCVirtualPropertyContainerBase::DuplicateVirtual
 	if (URCVirtualPropertyInContainer* NewVirtualProperty = DuplicateObject<URCVirtualPropertyInContainer>(InVirtualProperty, InVirtualProperty->GetOuter()))
 	{
 		NewVirtualProperty->PropertyName = GenerateUniquePropertyName(InVirtualProperty->PropertyName, this);
+		NewVirtualProperty->Id = FGuid::NewGuid();
 
 		// Sync Property Bag
 		Bag.AddProperty(NewVirtualProperty->PropertyName, InVirtualProperty->GetProperty());
diff --git a/Engine/Plugins/VirtualProduction/RemoteControl/Source/RemoteControl/Private/RemoteControlExposeRegistry.cpp b/Engine/Plugins/VirtualProduction/RemoteControl/Source/RemoteControl/Private/RemoteControlExposeRegistry.cpp
--- a/Engine/Plugins/VirtualProduction/RemoteControl/Source/RemoteControl/Private/RemoteControlExposeRegistry.cpp
+++ b/Engine/Plugins/VirtualProduction/RemoteControl/Source/RemoteControl/Private/RemoteControlExposeRegistry.cpp
@@ -162,6 +162,12 @@ void URemoteControlExposeRegistry::PostLoad()
 	CacheLabels();
 }
 
+void URemoteControlExposeRegistry::PostDuplicate(bool bDuplicateForPIE)
+{
+	UObject::PostDuplicate(bDuplicateForPIE);
+	CacheLabels();
+}
+
 TSharedPtr<FRemoteControlEntity> URemoteControlExposeRegistry::GetEntity(const FGuid& EntityId)
 {
 	/** Get a raw pointer to an entity using its id. */
diff --git a/Engine/Plugins/VirtualProduction/RemoteControl/Source/RemoteControl/Private/RemoteControlExposeRegistry.h b/Engine/Plugins/VirtualProduction/RemoteControl/Source/RemoteControl/Private/RemoteControlExposeRegistry.h
--- a/Engine/Plugins/VirtualProduction/RemoteControl/Source/RemoteControl/Private/RemoteControlExposeRegistry.h
+++ b/Engine/Plugins/VirtualProduction/RemoteControl/Source/RemoteControl/Private/RemoteControlExposeRegistry.h
@@ -197,6 +197,8 @@ public:
 
 	virtual void PostLoad() override;
 
+	virtual void PostDuplicate(bool bDuplicateForPIE) override;
+
 private:
 	/** Get a raw pointer to an entity using its id. */
 	TSharedPtr<FRemoteControlEntity> GetEntity(const FGuid& EntityId);
diff --git a/Engine/Plugins/VirtualProduction/RemoteControl/Source/RemoteControlUI/Private/UI/RCPanelWidgetRegistry.cpp b/Engine/Plugins/VirtualProduction/RemoteControl/Source/RemoteControlUI/Private/UI/RCPanelWidgetRegistry.cpp
--- a/Engine/Plugins/VirtualProduction/RemoteControl/Source/RemoteControlUI/Private/UI/RCPanelWidgetRegistry.cpp
+++ b/Engine/Plugins/VirtualProduction/RemoteControl/Source/RemoteControlUI/Private/UI/RCPanelWidgetRegistry.cpp
@@ -181,7 +181,11 @@ void FRCPanelWidgetRegistry::Refresh(UObject* InObject)
 {
 	if (TSharedPtr<IPropertyRowGenerator>* Generator = ObjectToRowGenerator.Find({InObject}))
 	{
-		(*Generator)->SetObjects({InObject});
+		const TArray<TWeakObjectPtr<UObject>>& SelectedObjects = (*Generator)->GetSelectedObjects();
+		if (SelectedObjects.Num() != 1 || SelectedObjects[0] != InObject)
+		{
+			(*Generator)->SetObjects({InObject});
+		}
 	}
 	else if (IsNDisplayObject(InObject))
 	{
