[Engine]
* Fixed non unity/pch compile errors for various platforms

Can't run preflight, seems like horde don't recognize my changelist

diff --git a/Engine/Plugins/Editor/WaveformEditor/Source/WaveformEditor/Private/WaveformEditorInstantiator.cpp b/Engine/Plugins/Editor/WaveformEditor/Source/WaveformEditor/Private/WaveformEditorInstantiator.cpp
--- a/Engine/Plugins/Editor/WaveformEditor/Source/WaveformEditor/Private/WaveformEditorInstantiator.cpp
+++ b/Engine/Plugins/Editor/WaveformEditor/Source/WaveformEditor/Private/WaveformEditorInstantiator.cpp
@@ -3,6 +3,7 @@
 #include "WaveformEditorInstantiator.h"
 
 #include "ContentBrowserModule.h"
+#include "Framework/Application/SlateApplication.h"
 #include "Framework/MultiBox/MultiBoxBuilder.h"
 #include "Sound/SoundWave.h"
 #include "SWaveformEditorMessageDialog.h"
diff --git a/Engine/Plugins/Experimental/ChaosCloth/Source/Chaos/Private/ChaosCloth/ChaosClothVisualization.cpp b/Engine/Plugins/Experimental/ChaosCloth/Source/Chaos/Private/ChaosCloth/ChaosClothVisualization.cpp
--- a/Engine/Plugins/Experimental/ChaosCloth/Source/Chaos/Private/ChaosCloth/ChaosClothVisualization.cpp
+++ b/Engine/Plugins/Experimental/ChaosCloth/Source/Chaos/Private/ChaosCloth/ChaosClothVisualization.cpp
@@ -27,6 +27,7 @@
 #include "DynamicMeshBuilder.h"
 #include "Engine/EngineTypes.h"
 #include "SceneManagement.h"
+#include "SceneView.h"
 #if WITH_EDITOR
 #include "Materials/Material.h"
 #include "Engine/Canvas.h"  // For draw text
diff --git a/Engine/Plugins/VirtualProduction/RemoteControl/Source/RemoteControlLogic/Private/Behaviour/Builtin/Path/RCSetAssetByPathBehaviour.cpp b/Engine/Plugins/VirtualProduction/RemoteControl/Source/RemoteControlLogic/Private/Behaviour/Builtin/Path/RCSetAssetByPathBehaviour.cpp
--- a/Engine/Plugins/VirtualProduction/RemoteControl/Source/RemoteControlLogic/Private/Behaviour/Builtin/Path/RCSetAssetByPathBehaviour.cpp
+++ b/Engine/Plugins/VirtualProduction/RemoteControl/Source/RemoteControlLogic/Private/Behaviour/Builtin/Path/RCSetAssetByPathBehaviour.cpp
@@ -21,6 +21,7 @@
 #include "Materials/Material.h"
 #include "Materials/MaterialInstance.h"
 #include "Misc/FileHelper.h"
+#include "Misc/Paths.h"
 #include "RCVirtualProperty.h"
 #include "RCVirtualPropertyContainer.h"
 #include "RemoteControlLogger.h"
diff --git a/Engine/Source/Runtime/Engine/Private/GPUSkinCacheVisualizationData.cpp b/Engine/Source/Runtime/Engine/Private/GPUSkinCacheVisualizationData.cpp
--- a/Engine/Source/Runtime/Engine/Private/GPUSkinCacheVisualizationData.cpp
+++ b/Engine/Source/Runtime/Engine/Private/GPUSkinCacheVisualizationData.cpp
@@ -3,6 +3,7 @@
 #include "GPUSkinCacheVisualizationData.h"
 #include "HAL/IConsoleManager.h"
 #include "RenderUtils.h"
+#include "RHIDefinitions.h"
 
 #define LOCTEXT_NAMESPACE "FGPUSkinCacheVisualizationData"
 
diff --git a/Engine/Source/Runtime/Renderer/Private/RayTracing/RayTracingDeferredMaterials.cpp b/Engine/Source/Runtime/Renderer/Private/RayTracing/RayTracingDeferredMaterials.cpp
--- a/Engine/Source/Runtime/Renderer/Private/RayTracing/RayTracingDeferredMaterials.cpp
+++ b/Engine/Source/Runtime/Renderer/Private/RayTracing/RayTracingDeferredMaterials.cpp
@@ -6,6 +6,7 @@
 #include "RenderCore.h"
 #include "GlobalShader.h"
 #include "ShaderParameterStruct.h"
+#include "RayTracingPayloadType.h"
 #include "RenderGraphUtils.h"
 #include "DeferredShadingRenderer.h"
 #include "PipelineStateCache.h"
diff --git a/Engine/Source/Runtime/Renderer/Private/RayTracing/RaytracingSkylight.cpp b/Engine/Source/Runtime/Renderer/Private/RayTracing/RaytracingSkylight.cpp
--- a/Engine/Source/Runtime/Renderer/Private/RayTracing/RaytracingSkylight.cpp
+++ b/Engine/Source/Runtime/Renderer/Private/RayTracing/RaytracingSkylight.cpp
@@ -5,6 +5,7 @@
 =============================================================================*/
 
 #include "DeferredShadingRenderer.h"
+#include "DataDrivenShaderPlatformInfo.h"
 #include "SceneTextureParameters.h"
 
 static int32 GRayTracingSkyLight = 0;
