StructUtils: FixStructUtils constness issues in  RemoteControl.

diff --git a/Engine/Plugins/VirtualProduction/RemoteControl/Source/RemoteControl/Private/RCVirtualProperty.cpp b/Engine/Plugins/VirtualProduction/RemoteControl/Source/RemoteControl/Private/RCVirtualProperty.cpp
--- a/Engine/Plugins/VirtualProduction/RemoteControl/Source/RemoteControl/Private/RCVirtualProperty.cpp
+++ b/Engine/Plugins/VirtualProduction/RemoteControl/Source/RemoteControl/Private/RCVirtualProperty.cpp
@@ -33,7 +33,7 @@ const uint8* URCVirtualPropertyBase::GetContainerPtr() const
 {
 	if (const FInstancedPropertyBag* InstancedPropertyBag = GetPropertyBagInstance())
 	{
-		return InstancedPropertyBag->GetMutableValue().GetMemory();
+		return InstancedPropertyBag->GetValue().GetMemory();
 	}
 	
 	return GetContainerPtr();
@@ -43,7 +43,7 @@ uint8* URCVirtualPropertyBase::GetContainerPtr()
 {
 	if (const FInstancedPropertyBag* InstancedPropertyBag = GetPropertyBagInstance())
 	{
-		return InstancedPropertyBag->GetMutableValue().GetMutableMemory();
+		return const_cast<FInstancedPropertyBag*>(InstancedPropertyBag)->GetMutableValue().GetMutableMemory();
 	}
 	
 	return nullptr;
@@ -768,7 +768,7 @@ const FInstancedPropertyBag* URCVirtualPropertySelfContainer::GetPropertyBagInst
 	return &Bag;
 }
 
-TSharedPtr<FStructOnScope> URCVirtualPropertySelfContainer::CreateStructOnScope() const
+TSharedPtr<FStructOnScope> URCVirtualPropertySelfContainer::CreateStructOnScope()
 {
 	return MakeShared<FStructOnScope>(Bag.GetPropertyBagStruct(), Bag.GetMutableValue().GetMutableMemory());
 }
\ No newline at end of file
diff --git a/Engine/Plugins/VirtualProduction/RemoteControl/Source/RemoteControl/Private/RCVirtualPropertyContainer.cpp b/Engine/Plugins/VirtualProduction/RemoteControl/Source/RemoteControl/Private/RCVirtualPropertyContainer.cpp
--- a/Engine/Plugins/VirtualProduction/RemoteControl/Source/RemoteControl/Private/RCVirtualPropertyContainer.cpp
+++ b/Engine/Plugins/VirtualProduction/RemoteControl/Source/RemoteControl/Private/RCVirtualPropertyContainer.cpp
@@ -202,7 +202,7 @@ int32 URCVirtualPropertyContainerBase::GetNumVirtualProperties() const
 	return NumPropertiesInBag;
 }
 
-TSharedPtr<FStructOnScope> URCVirtualPropertyContainerBase::CreateStructOnScope() const
+TSharedPtr<FStructOnScope> URCVirtualPropertyContainerBase::CreateStructOnScope()
 {
 	return MakeShared<FStructOnScope>(Bag.GetPropertyBagStruct(), Bag.GetMutableValue().GetMutableMemory());
 }
diff --git a/Engine/Plugins/VirtualProduction/RemoteControl/Source/RemoteControl/Public/RCVirtualProperty.h b/Engine/Plugins/VirtualProduction/RemoteControl/Source/RemoteControl/Public/RCVirtualProperty.h
--- a/Engine/Plugins/VirtualProduction/RemoteControl/Source/RemoteControl/Public/RCVirtualProperty.h
+++ b/Engine/Plugins/VirtualProduction/RemoteControl/Source/RemoteControl/Public/RCVirtualProperty.h
@@ -392,7 +392,7 @@ public:
 	void Reset();
 
 	/** Creates new Struct on Scope for this Property Bag UStruct and Memory */
-	TSharedPtr<FStructOnScope> CreateStructOnScope() const;
+	TSharedPtr<FStructOnScope> CreateStructOnScope();
 
 protected:
 	//~ Begin URCVirtualPropertyBase interface
diff --git a/Engine/Plugins/VirtualProduction/RemoteControl/Source/RemoteControl/Public/RCVirtualPropertyContainer.h b/Engine/Plugins/VirtualProduction/RemoteControl/Source/RemoteControl/Public/RCVirtualPropertyContainer.h
--- a/Engine/Plugins/VirtualProduction/RemoteControl/Source/RemoteControl/Public/RCVirtualPropertyContainer.h
+++ b/Engine/Plugins/VirtualProduction/RemoteControl/Source/RemoteControl/Public/RCVirtualPropertyContainer.h
@@ -105,7 +105,7 @@ public:
 	/**
 	 * Creates new Struct on Scope for this Property Bag UStruct and Memory
 	 */
-	TSharedPtr<FStructOnScope> CreateStructOnScope() const;
+	TSharedPtr<FStructOnScope> CreateStructOnScope();
 
 	/**
 	 * Generates unique name for the property for specified property container
diff --git a/Engine/Plugins/VirtualProduction/RemoteControl/Source/WebRemoteControl/Private/WebSocketMessageHandler.cpp b/Engine/Plugins/VirtualProduction/RemoteControl/Source/WebRemoteControl/Private/WebSocketMessageHandler.cpp
--- a/Engine/Plugins/VirtualProduction/RemoteControl/Source/WebRemoteControl/Private/WebSocketMessageHandler.cpp
+++ b/Engine/Plugins/VirtualProduction/RemoteControl/Source/WebRemoteControl/Private/WebSocketMessageHandler.cpp
@@ -1732,7 +1732,7 @@ bool FWebSocketMessageHandler::WriteControllerChangeEventPayload(URemoteControlP
 	TArray<FStructOnScope*> ControllerValueOnScope;
 	for (const FGuid& ControllerId : InModifiedControllerIds)
 	{
-		if (const URCVirtualPropertySelfContainer* Controller = Cast<URCVirtualPropertySelfContainer>(InPreset->GetController(ControllerId)))
+		if (URCVirtualPropertySelfContainer* Controller = Cast<URCVirtualPropertySelfContainer>(InPreset->GetController(ControllerId)))
 		{
 			bHasController = true;
 
