- Fixing remote control non unity compile errors

diff --git a/Engine/Plugins/VirtualProduction/RemoteControl/Source/RemoteControl/Private/RemoteControlModule.cpp b/Engine/Plugins/VirtualProduction/RemoteControl/Source/RemoteControl/Private/RemoteControlModule.cpp
--- a/Engine/Plugins/VirtualProduction/RemoteControl/Source/RemoteControl/Private/RemoteControlModule.cpp
+++ b/Engine/Plugins/VirtualProduction/RemoteControl/Source/RemoteControl/Private/RemoteControlModule.cpp
@@ -13,6 +13,7 @@
 #include "IRemoteControlModule.h"
 #include "IStructDeserializerBackend.h"
 #include "IStructSerializerBackend.h"
+#include "Misc/ScopeExit.h"
 #include "RCPropertyUtilities.h"
 #include "RCVirtualPropertyContainer.h"
 #include "RCVirtualProperty.h"
