Fixed non unity compile errors when building unrealgame win64 -allplugins. Same changes has been made in ue main already.

diff --git a/Engine/Plugins/Media/WmfMedia/Source/WmfMediaFactory/Private/WmfMediaFactoryModule.cpp b/Engine/Plugins/Media/WmfMedia/Source/WmfMediaFactory/Private/WmfMediaFactoryModule.cpp
--- a/Engine/Plugins/Media/WmfMedia/Source/WmfMediaFactory/Private/WmfMediaFactoryModule.cpp
+++ b/Engine/Plugins/Media/WmfMedia/Source/WmfMediaFactory/Private/WmfMediaFactoryModule.cpp
@@ -17,6 +17,7 @@
 	#include "WmfMediaSettings.h"
 #else
 	#include "Internationalization/Internationalization.h"
+	#include "Misc/Guid.h"
 #endif
 
 #include "../../WmfMedia/Public/IWmfMediaModule.h"
diff --git a/Engine/Plugins/Runtime/AR/Google/GoogleARCoreServices/Source/GoogleARCoreServices/Private/GoogleARCoreUtils.h b/Engine/Plugins/Runtime/AR/Google/GoogleARCoreServices/Source/GoogleARCoreServices/Private/GoogleARCoreUtils.h
--- a/Engine/Plugins/Runtime/AR/Google/GoogleARCoreServices/Source/GoogleARCoreServices/Private/GoogleARCoreUtils.h
+++ b/Engine/Plugins/Runtime/AR/Google/GoogleARCoreServices/Source/GoogleARCoreServices/Private/GoogleARCoreUtils.h
@@ -2,6 +2,8 @@
 
 #pragma once
 
+#include "HAL/Platform.h"
+
 #if PLATFORM_ANDROID || PLATFORM_IOS
 #define ARCORE_SERVICE_SUPPORTED_PLATFORM 1
 #else
diff --git a/Engine/Plugins/VirtualProduction/RemoteControl/Source/RemoteControl/Private/RemoteControlModule.cpp b/Engine/Plugins/VirtualProduction/RemoteControl/Source/RemoteControl/Private/RemoteControlModule.cpp
--- a/Engine/Plugins/VirtualProduction/RemoteControl/Source/RemoteControl/Private/RemoteControlModule.cpp
+++ b/Engine/Plugins/VirtualProduction/RemoteControl/Source/RemoteControl/Private/RemoteControlModule.cpp
@@ -39,6 +39,8 @@
 #include "HAL/IConsoleManager.h"
 #include "ScopedTransaction.h"
 #include "UnrealEdGlobals.h"
+#else
+#include "Misc/ScopeExit.h"
 #endif
 
 DEFINE_LOG_CATEGORY(LogRemoteControl);
