RC Controllers - use field name when drag/dropped as controller

diff --git a/Engine/Plugins/VirtualProduction/RemoteControl/Source/RemoteControlUI/Private/UI/Controller/SRCControllerPanelList.cpp b/Engine/Plugins/VirtualProduction/RemoteControl/Source/RemoteControlUI/Private/UI/Controller/SRCControllerPanelList.cpp
--- a/Engine/Plugins/VirtualProduction/RemoteControl/Source/RemoteControlUI/Private/UI/Controller/SRCControllerPanelList.cpp
+++ b/Engine/Plugins/VirtualProduction/RemoteControl/Source/RemoteControlUI/Private/UI/Controller/SRCControllerPanelList.cpp
@@ -675,7 +675,7 @@ void SRCControllerPanelList::CreateAutoBindForProperty(TSharedPtr<const FRemoteC
 		if(bSuccess)
 		{
 			// Step 1. Create a Controller of matching type
-			URCController* NewController = Cast<URCController>(Preset->AddController(URCController::StaticClass(), PropertyBagType, StructObject));
+			URCController* NewController = Cast<URCController>(Preset->AddController(URCController::StaticClass(), PropertyBagType, StructObject, RemoteControlProperty->FieldName));
 			NewController->DisplayIndex = Preset->GetNumControllers() - 1;
 
 			// Transfer property value from Exposed Property to the New Controller.
