Remote Control:
- Fix issue with Reordering Controller Items not correctly updating all Controller Display Indices when reordering an item downwards

diff --git a/Engine/Plugins/VirtualProduction/RemoteControl/Source/RemoteControlUI/Private/UI/Controller/SRCControllerPanelList.cpp b/Engine/Plugins/VirtualProduction/RemoteControl/Source/RemoteControlUI/Private/UI/Controller/SRCControllerPanelList.cpp
--- a/Engine/Plugins/VirtualProduction/RemoteControl/Source/RemoteControlUI/Private/UI/Controller/SRCControllerPanelList.cpp
+++ b/Engine/Plugins/VirtualProduction/RemoteControl/Source/RemoteControlUI/Private/UI/Controller/SRCControllerPanelList.cpp
@@ -585,9 +585,9 @@ void SRCControllerPanelList::ReorderControllerItem(TSharedRef<FRCControllerModel
 	ControllerItems.Insert(ItemToMove, Index);
 
 	// Update display indices
-	for (int32 i = Index; i < ControllerItems.Num(); i++)
+	for (int32 i = 0; i < ControllerItems.Num(); i++)
 	{
-		if (ensure(ControllerItems[i]))
+		if (ControllerItems[i])
 		{
 			if (URCVirtualPropertyBase* Controller = ControllerItems[i]->GetVirtualProperty())
 			{
