Remote Control:
- Fix issue with Vector Controllers not correctly updating Bound Vector Properties

diff --git a/Engine/Plugins/VirtualProduction/RemoteControl/Source/RemoteControl/Private/RemoteControlPreset.cpp b/Engine/Plugins/VirtualProduction/RemoteControl/Source/RemoteControl/Private/RemoteControlPreset.cpp
--- a/Engine/Plugins/VirtualProduction/RemoteControl/Source/RemoteControl/Private/RemoteControlPreset.cpp
+++ b/Engine/Plugins/VirtualProduction/RemoteControl/Source/RemoteControl/Private/RemoteControlPreset.cpp
@@ -663,7 +663,7 @@ void URemoteControlPreset::OnModifyController(const FPropertyChangedEvent& Prope
 {
 	if (ensure(ControllerContainer))
 	{
-		if (URCVirtualPropertyBase* ModifiedController = ControllerContainer->GetVirtualProperty(PropertyChangedEvent.Property->GetFName()))
+		if (URCVirtualPropertyBase* ModifiedController = ControllerContainer->GetVirtualProperty(PropertyChangedEvent.MemberProperty->GetFName()))
 		{
 			ControllerContainer->OnModifyPropertyValue(PropertyChangedEvent);
 			OnControllerModified().Broadcast(this, {ModifiedController->Id});
