1. RCP -Allowing String controller to be multiline when controller is binding Textbox

2. Remote Control Preset WebSocket updates are sometimes incorrect for nested bitfields

diff --git a/Engine/Plugins/VirtualProduction/RemoteControl/Source/RemoteControlUI/Private/UI/Controller/SRCControllerPanelList.cpp b/Engine/Plugins/VirtualProduction/RemoteControl/Source/RemoteControlUI/Private/UI/Controller/SRCControllerPanelList.cpp
--- a/Engine/Plugins/VirtualProduction/RemoteControl/Source/RemoteControlUI/Private/UI/Controller/SRCControllerPanelList.cpp
+++ b/Engine/Plugins/VirtualProduction/RemoteControl/Source/RemoteControlUI/Private/UI/Controller/SRCControllerPanelList.cpp
@@ -375,6 +375,11 @@ void SRCControllerPanelList::Reset()
 			FProperty* Property = Child->CreatePropertyHandle()->GetProperty();
  			check(Property);
 
+			if(Property->IsA<FStrProperty>() || Property->IsA<FTextProperty>())
+			{
+				Property->AppendMetaData({{TEXT("multiline"), TEXT("true")}});
+			}
+
 			if (URCVirtualPropertyBase* Controller = Preset->GetController(Property->GetFName()))
 			{
 				if(ensureAlways(ControllerItems.IsValidIndex(Controller->DisplayIndex)))
diff --git a/Engine/Plugins/VirtualProduction/RemoteControl/Source/WebRemoteControl/Private/WebSocketMessageHandler.cpp b/Engine/Plugins/VirtualProduction/RemoteControl/Source/WebRemoteControl/Private/WebSocketMessageHandler.cpp
--- a/Engine/Plugins/VirtualProduction/RemoteControl/Source/WebRemoteControl/Private/WebSocketMessageHandler.cpp
+++ b/Engine/Plugins/VirtualProduction/RemoteControl/Source/WebRemoteControl/Private/WebSocketMessageHandler.cpp
@@ -1002,12 +1002,16 @@ void FWebSocketMessageHandler::ProcessChangedProperties()
 			{
 				const int64 SequenceNumber = GetSequenceNumber(ClientToEventsPair.Key);
 
-				TArray<uint8> WorkingBuffer;
-				if (ClientToEventsPair.Value.Num() && WritePropertyChangeEventPayload(Preset, { ClassToEventsPair.Value }, SequenceNumber, WorkingBuffer))
+				//Check if multiple booleans properties want to be sent and send them since multiple booleans have problem with the common workflow.
+				if(!TrySendMultipleBoolProperties(Preset, ClientToEventsPair.Key, ClassToEventsPair.Value, SequenceNumber))
 				{
-					TArray<uint8> Payload;
-					WebRemoteControlUtils::ConvertToUTF8(WorkingBuffer, Payload);
-					Server->Send(ClientToEventsPair.Key, Payload);
+					TArray<uint8> WorkingBuffer;
+					if (ClientToEventsPair.Value.Num() && WritePropertyChangeEventPayload(Preset, { ClassToEventsPair.Value }, SequenceNumber, WorkingBuffer))
+					{
+						TArray<uint8> Payload;
+						WebRemoteControlUtils::ConvertToUTF8(WorkingBuffer, Payload);
+						Server->Send(ClientToEventsPair.Key, Payload);
+					}
 				}
 			}
 		}
@@ -1786,6 +1790,35 @@ bool FWebSocketMessageHandler::WritePropertyChangeEventPayload(URemoteControlPre
 	return bHasProperty;
 }
 
+bool FWebSocketMessageHandler::TrySendMultipleBoolProperties(URemoteControlPreset* InPreset,
+	const FGuid& InTargetClientId, const TSet<FGuid>& InModifiedPropertyIds, int64 InSequenceNumber)
+{
+	bool bFound = false;
+	int32 NumberSent = 0;
+	if(InModifiedPropertyIds.Num() > 1)
+	{
+		if (TSharedPtr<FRemoteControlProperty> RCProperty = InPreset->GetExposedEntity<FRemoteControlProperty>(InModifiedPropertyIds.Array()[0]).Pin())
+		{
+			if(RCProperty->GetProperty()->IsA<FBoolProperty>())
+			{
+				bFound = true;
+				for(FGuid ModifiedPropertyId : InModifiedPropertyIds)
+				{
+					TArray<uint8> BoolsWorkingBuffer;
+					if(WritePropertyChangeEventPayload(InPreset, { ModifiedPropertyId }, InSequenceNumber, BoolsWorkingBuffer))
+					{
+						TArray<uint8> Payload;
+						WebRemoteControlUtils::ConvertToUTF8(BoolsWorkingBuffer, Payload);
+						Server->Send(InTargetClientId, Payload);
+						++NumberSent;
+					}
+				}
+			}
+		}
+	}
+	return bFound && (NumberSent == InModifiedPropertyIds.Num());
+}
+
 
 bool FWebSocketMessageHandler::WriteActorPropertyChangePayload(URemoteControlPreset* InPreset, const TMap<FRemoteControlActor, TArray<FRCObjectReference>>& InModifications, FMemoryWriter& InWriter)
 {
diff --git a/Engine/Plugins/VirtualProduction/RemoteControl/Source/WebRemoteControl/Private/WebSocketMessageHandler.h b/Engine/Plugins/VirtualProduction/RemoteControl/Source/WebRemoteControl/Private/WebSocketMessageHandler.h
--- a/Engine/Plugins/VirtualProduction/RemoteControl/Source/WebRemoteControl/Private/WebSocketMessageHandler.h
+++ b/Engine/Plugins/VirtualProduction/RemoteControl/Source/WebRemoteControl/Private/WebSocketMessageHandler.h
@@ -207,6 +207,16 @@ private:
 	 */
 	bool WritePropertyChangeEventPayload(URemoteControlPreset* InPreset, const TSet<FGuid>& InModifiedPropertyIds, int64 InSequenceNumber, TArray<uint8>& OutBuffer);
 
+	/**
+	 * Check if multiple booleans properties want to be sent over to a client, and send a message for each of those properties. This is needed because multiple booleans have problems with the common workflow.
+	 * @param InPreset The current preset
+	 * @param InTargetClientId Client Id to send the message
+	 * @param InModifiedPropertyIds Set of modified properties already divided by type
+	 * @param InSequenceNumber SequenceNumber of the client
+	 * @return True if multiple booleans property are found and sent, false if they are not found or if not all of them are sent.
+	 */
+	bool TrySendMultipleBoolProperties(URemoteControlPreset* InPreset, const FGuid& InTargetClientId, const TSet<FGuid>& InModifiedPropertyIds, int64 InSequenceNumber);
+	
 	/**
 	 * Write the provided list of controller events to a buffer.
 	 */
