CIS fix.

diff --git a/Engine/Plugins/VirtualProduction/RemoteControl/Source/RemoteControlCommon/Public/RCTypeUtilities.h b/Engine/Plugins/VirtualProduction/RemoteControl/Source/RemoteControlCommon/Public/RCTypeUtilities.h
--- a/Engine/Plugins/VirtualProduction/RemoteControl/Source/RemoteControlCommon/Public/RCTypeUtilities.h
+++ b/Engine/Plugins/VirtualProduction/RemoteControl/Source/RemoteControlCommon/Public/RCTypeUtilities.h
@@ -389,7 +389,7 @@ namespace RemoteControlTypeUtilities
 	}
 
 	template <typename ValueType>
-	static typename TEnableIf<RemoteControlTypeTraits::TNumericValueConstraint_V<ValueType>::Value, ValueType>::Type
+	static typename TEnableIf<RemoteControlTypeTraits::TNumericValueConstraint_V<ValueType>, ValueType>::Type
 	GetDefaultRangeValueMax(const FProperty* InProperty)
 	{
 		check(InProperty);
