RC error output callback

diff --git a/Engine/Plugins/VirtualProduction/RemoteControl/Source/RemoteControl/Private/RemoteControlModule.cpp b/Engine/Plugins/VirtualProduction/RemoteControl/Source/RemoteControl/Private/RemoteControlModule.cpp
--- a/Engine/Plugins/VirtualProduction/RemoteControl/Source/RemoteControl/Private/RemoteControlModule.cpp
+++ b/Engine/Plugins/VirtualProduction/RemoteControl/Source/RemoteControl/Private/RemoteControlModule.cpp
@@ -15,12 +15,12 @@
 #include "IStructSerializerBackend.h"
 #include "Misc/ScopeExit.h"
 #include "RCPropertyUtilities.h"
-#include "RCVirtualPropertyContainer.h"
 #include "RCVirtualProperty.h"
+#include "RCVirtualPropertyContainer.h"
 #include "RemoteControlFieldPath.h"
+#include "RemoteControlInstanceMaterial.h"
 #include "RemoteControlInterceptionHelpers.h"
 #include "RemoteControlInterceptionProcessor.h"
-#include "RemoteControlInstanceMaterial.h"
 #include "RemoteControlPreset.h"
 #include "RemoteControlSettings.h"
 #include "SceneInterface.h"
@@ -65,6 +65,7 @@ namespace RemoteControlUtil
 	const FName NAME_BlueprintGetter(TEXT("BlueprintGetter"));
 	const FName NAME_BlueprintSetter(TEXT("BlueprintSetter"));
 	const FName NAME_AllowPrivateAccess(TEXT("AllowPrivateAccess"));
+	const FName NAME_RelativeLocation(TEXT("RelativeLocation"));
 
 	bool CompareFunctionName(const FString& FunctionName, const FString& ScriptName)
 	{
@@ -644,6 +645,11 @@ IRemoteControlModule::FOnPresetUnregistered& FRemoteControlModule::OnPresetUnreg
 	return OnPresetUnregisteredDelegate;
 }
 
+IRemoteControlModule::FOnError& FRemoteControlModule::OnError()
+{
+	return OnErrorDelegate;
+}
+
 /** Register the preset with the module, enabling using the preset remotely using its name. */
 bool FRemoteControlModule::RegisterPreset(FName Name, URemoteControlPreset* Preset)
 {
@@ -840,10 +846,16 @@ bool FRemoteControlModule::ResolveCall(const FString& ObjectPath, const FString&
 		bSuccess = false;
 	}
 
+	if (!bSuccess && !ErrorText.IsEmpty())
+	{
+		IRemoteControlModule::BroadcastError(ErrorText);
+	}
+
 	if (OutErrorText && !ErrorText.IsEmpty())
 	{
 		*OutErrorText = MoveTemp(ErrorText);
 	}
+	
 	return bSuccess;
 }
 
@@ -977,6 +989,11 @@ bool FRemoteControlModule::ResolveObject(ERCAccess AccessType, const FString& Ob
 		bSuccess = false;
 	}
 
+	if (!bSuccess && !ErrorText.IsEmpty())
+	{
+		IRemoteControlModule::BroadcastError(ErrorText);
+	}
+
 	if (OutErrorText && !ErrorText.IsEmpty())
 	{
 		*OutErrorText = MoveTemp(ErrorText);
@@ -1038,6 +1055,11 @@ bool FRemoteControlModule::ResolveObjectProperty(ERCAccess AccessType, UObject*
 		bSuccess = false;
 	}
 
+	if (!bSuccess && !ErrorText.IsEmpty())
+	{
+		IRemoteControlModule::BroadcastError(ErrorText);
+	}
+
 	if (OutErrorText && !ErrorText.IsEmpty())
 	{
 		*OutErrorText = MoveTemp(ErrorText);
@@ -1393,7 +1415,7 @@ void FRemoteControlModule::RefreshEditorPostSetObjectProperties(const FRCObjectR
 	{
 		if (FProperty* Property = ObjectAccess.Property.Get())
 		{
-			if(Property->GetName() == "RelativeLocation")
+			if(Property->GetFName() == RemoteControlUtil::NAME_RelativeLocation)
 			{
 				if (AActor* Actor = SceneComponent->GetOwner())
 				{
diff --git a/Engine/Plugins/VirtualProduction/RemoteControl/Source/RemoteControl/Private/RemoteControlModule.h b/Engine/Plugins/VirtualProduction/RemoteControl/Source/RemoteControl/Private/RemoteControlModule.h
--- a/Engine/Plugins/VirtualProduction/RemoteControl/Source/RemoteControl/Private/RemoteControlModule.h
+++ b/Engine/Plugins/VirtualProduction/RemoteControl/Source/RemoteControl/Private/RemoteControlModule.h
@@ -1,12 +1,12 @@
 // Copyright Epic Games, Inc. All Rights Reserved.
 #pragma once
 
-#include "IRemoteControlModule.h"
 #include "AssetRegistry/AssetData.h"
 #include "CoreMinimal.h"
 #include "Engine/EngineTypes.h"
 #include "Factories/IRemoteControlMaskingFactory.h"
 #include "IRemoteControlInterceptionFeature.h"
+#include "IRemoteControlModule.h"
 
 /**
  * Implementation of the RemoteControl interface
@@ -24,6 +24,7 @@ public:
 	//~ Begin IRemoteControlModule
 	virtual FOnPresetRegistered& OnPresetRegistered() override;
 	virtual FOnPresetUnregistered& OnPresetUnregistered() override;
+	virtual FOnError& OnError() override;
 	virtual bool RegisterPreset(FName Name, URemoteControlPreset* Preset) override;
 	virtual void UnregisterPreset(FName Name) override;
 	virtual bool RegisterEmbeddedPreset(URemoteControlPreset* Preset, bool bReplaceExisting) override;
@@ -150,6 +151,9 @@ private:
 	/** Delegate for preset unregistration */
 	FOnPresetUnregistered OnPresetUnregisteredDelegate;
 
+	/** Delegate for errors to allow external custom handling. */
+	FOnError OnErrorDelegate;
+
 	/** RC Processor feature instance */
 	TUniquePtr<IRemoteControlInterceptionFeatureProcessor> RCIProcessor;
 
diff --git a/Engine/Plugins/VirtualProduction/RemoteControl/Source/RemoteControl/Public/IRemoteControlModule.h b/Engine/Plugins/VirtualProduction/RemoteControl/Source/RemoteControl/Public/IRemoteControlModule.h
--- a/Engine/Plugins/VirtualProduction/RemoteControl/Source/RemoteControl/Public/IRemoteControlModule.h
+++ b/Engine/Plugins/VirtualProduction/RemoteControl/Source/RemoteControl/Public/IRemoteControlModule.h
@@ -1,4 +1,5 @@
 // Copyright Epic Games, Inc. All Rights Reserved.
+
 #pragma  once
 
 #include "CoreMinimal.h"
@@ -244,6 +245,16 @@ public:
 	UE_DEPRECATED(4.27, "OnPresetUnregistered is deprecated.")
 	virtual FOnPresetUnregistered& OnPresetUnregistered() = 0;
 
+	/** Delegate triggered when an error occurs, which aren't necessarily written to UE LOG */
+	DECLARE_MULTICAST_DELEGATE_TwoParams(FOnError, const FString& /*Message*/, ELogVerbosity::Type /*Verbosity*/);
+	virtual FOnError& OnError() = 0;
+
+	/** Broadcast an error to the OnError delegate */
+	static void BroadcastError(const FString& Message, ELogVerbosity::Type Verbosity = ELogVerbosity::Error)
+	{
+		Get().OnError().Broadcast(Message, Verbosity);
+	}
+
 	/**
 	 * Register the preset with the module, enabling using the preset remotely using its name.
 	 * @return whether registration was successful.
diff --git a/Engine/Plugins/VirtualProduction/RemoteControl/Source/RemoteControlProtocolWidgets/Private/Widgets/SPropertyView.cpp b/Engine/Plugins/VirtualProduction/RemoteControl/Source/RemoteControlProtocolWidgets/Private/Widgets/SPropertyView.cpp
--- a/Engine/Plugins/VirtualProduction/RemoteControl/Source/RemoteControlProtocolWidgets/Private/Widgets/SPropertyView.cpp
+++ b/Engine/Plugins/VirtualProduction/RemoteControl/Source/RemoteControlProtocolWidgets/Private/Widgets/SPropertyView.cpp
@@ -7,8 +7,10 @@
 #include "Editor.h"
 #include "IDetailTreeNode.h"
 #include "IPropertyRowGenerator.h"
-#include "SRCProtocolShared.h"
+#include "IRemoteControlModule.h"
 #include "Modules/ModuleManager.h"
+#include "RemoteControlProtocolWidgetsModule.h"
+#include "SRCProtocolShared.h"
 #include "Widgets/Layout/SGridPanel.h"
 #include "Widgets/Text/STextBlock.h"
 
@@ -145,6 +147,7 @@ void SPropertyView::ConstructInternal()
 	else
 	{
 		const FText ErrorText = LOCTEXT("InvalidObject", "Error: Not a valid Object");
+		IRemoteControlModule::BroadcastError(ErrorText.ToString());
 		ChildSlot
 		[
 			SNew(SVerticalBox)
diff --git a/Engine/Plugins/VirtualProduction/RemoteControl/Source/RemoteControlUI/Private/Factories/RemoteControlPresetActorFactory.cpp b/Engine/Plugins/VirtualProduction/RemoteControl/Source/RemoteControlUI/Private/Factories/RemoteControlPresetActorFactory.cpp
--- a/Engine/Plugins/VirtualProduction/RemoteControl/Source/RemoteControlUI/Private/Factories/RemoteControlPresetActorFactory.cpp
+++ b/Engine/Plugins/VirtualProduction/RemoteControl/Source/RemoteControlUI/Private/Factories/RemoteControlPresetActorFactory.cpp
@@ -1,5 +1,7 @@
 // Copyright Epic Games, Inc. All Rights Reserved.
+
 #include "RemoteControlPresetActorFactory.h"
+
 #include "AssetRegistry/AssetData.h"
 #include "Engine/Level.h"
 #include "Framework/Notifications/NotificationManager.h"
@@ -21,6 +23,7 @@ bool URemoteControlPresetActorFactory::CanCreateActorFrom( const FAssetData& Ass
 	if (AssetData.IsValid() && !AssetData.GetClass()->IsChildOf(URemoteControlPreset::StaticClass()))
 	{
 		OutErrorMsg = NSLOCTEXT("CanCreateActor", "NoRemoteControlPresetAsset", "A valid remote control preset asset must be specified.");
+		IRemoteControlModule::BroadcastError(OutErrorMsg.ToString());
 		return false;
 	}
 
diff --git a/Engine/Plugins/VirtualProduction/RemoteControl/Source/RemoteControlUI/Private/UI/SRCPanelExposedEntity.cpp b/Engine/Plugins/VirtualProduction/RemoteControl/Source/RemoteControlUI/Private/UI/SRCPanelExposedEntity.cpp
--- a/Engine/Plugins/VirtualProduction/RemoteControl/Source/RemoteControlUI/Private/UI/SRCPanelExposedEntity.cpp
+++ b/Engine/Plugins/VirtualProduction/RemoteControl/Source/RemoteControlUI/Private/UI/SRCPanelExposedEntity.cpp
@@ -13,6 +13,7 @@
 #include "Framework/Application/SlateApplication.h"
 #include "Framework/MultiBox/MultiBoxBuilder.h"
 #include "GameFramework/Actor.h"
+#include "IRemoteControlModule.h"
 #include "Interfaces/IMainFrameModule.h"
 #include "Modules/ModuleManager.h"
 #include "RemoteControlBinding.h"
@@ -36,9 +37,8 @@
 #include "Widgets/Layout/SBox.h"
 #include "Widgets/Notifications/SPopUpErrorText.h"
 #include "Widgets/SBoxPanel.h"
-#include "Widgets/Text/STextBlock.h"
 #include "Widgets/Text/SInlineEditableTextBlock.h"
-
+#include "Widgets/Text/STextBlock.h"
 
 #define LOCTEXT_NAMESPACE "RemoteControlPanel"
 
@@ -391,6 +391,7 @@ bool SRCPanelExposedEntity::OnVerifyItemLabelChanged(const FText& InLabel, FText
 		if (InLabel.ToString() != CachedLabel.ToString() && RCPreset->GetExposedEntityId(*InLabel.ToString()).IsValid())
 		{
 			OutErrorMessage = LOCTEXT("NameAlreadyExists", "This name already exists.");
+			IRemoteControlModule::BroadcastError(OutErrorMessage.ToString());
 			return false;
 		}
 	}
diff --git a/Engine/Plugins/VirtualProduction/RemoteControl/Source/RemoteControlUI/Private/UI/SRCPanelExposedField.cpp b/Engine/Plugins/VirtualProduction/RemoteControl/Source/RemoteControlUI/Private/UI/SRCPanelExposedField.cpp
--- a/Engine/Plugins/VirtualProduction/RemoteControl/Source/RemoteControlUI/Private/UI/SRCPanelExposedField.cpp
+++ b/Engine/Plugins/VirtualProduction/RemoteControl/Source/RemoteControlUI/Private/UI/SRCPanelExposedField.cpp
@@ -6,6 +6,7 @@
 #include "EditorFontGlyphs.h"
 #include "IDetailTreeNode.h"
 #include "IRCProtocolBindingList.h"
+#include "IRemoteControlModule.h"
 #include "IRemoteControlProtocolModule.h"
 #include "IRemoteControlProtocolWidgetsModule.h"
 #include "IRemoteControlUIModule.h"
@@ -20,17 +21,17 @@
 #include "SResetToDefaultPropertyEditor.h"
 #include "ScopedTransaction.h"
 #include "Styling/AppStyle.h"
-#include "Styling/SlateBrush.h"
 #include "Styling/RemoteControlStyles.h"
+#include "Styling/SlateBrush.h"
 #include "UObject/Object.h"
 #include "Widgets/Input/SButton.h"
 #include "Widgets/Input/SCheckBox.h"
+#include "Widgets/Layout/SBorder.h"
 #include "Widgets/Layout/SBox.h"
 #include "Widgets/Masks/SRCProtocolMask.h"
-#include "Widgets/Layout/SBorder.h"
+#include "Widgets/SNullWidget.h"
 #include "Widgets/Text/SInlineEditableTextBlock.h"
 #include "Widgets/Text/STextBlock.h"
-#include "Widgets/SNullWidget.h"
 
 #define LOCTEXT_NAMESPACE "RemoteControlPanel"
 
@@ -455,6 +456,7 @@ TSharedRef<SWidget> SRCPanelExposedField::ConstructWidget()
 					}
 
 					FText ErrorText = FText::Format(LOCTEXT("ExposedPropertyInvalidErrorMessage", "Could not find property {0} on object {1}"), FText::FromString(PropertyName), FText::FromString(TargetName));
+					IRemoteControlModule::BroadcastError(ErrorText.ToString());
 					
 					return MakeFieldWidget(CreateInvalidWidget(ErrorText));
 				}
diff --git a/Engine/Plugins/VirtualProduction/RemoteControl/Source/RemoteControlUI/Private/UI/SRCPanelFieldGroup.cpp b/Engine/Plugins/VirtualProduction/RemoteControl/Source/RemoteControlUI/Private/UI/SRCPanelFieldGroup.cpp
--- a/Engine/Plugins/VirtualProduction/RemoteControl/Source/RemoteControlUI/Private/UI/SRCPanelFieldGroup.cpp
+++ b/Engine/Plugins/VirtualProduction/RemoteControl/Source/RemoteControlUI/Private/UI/SRCPanelFieldGroup.cpp
@@ -4,21 +4,22 @@
 
 #include "Commands/RemoteControlCommands.h"
 #include "Framework/MultiBox/MultiBoxBuilder.h"
+#include "IRemoteControlModule.h"
 #include "Interfaces/IMainFrameModule.h"
 #include "RemoteControlPanelStyle.h"
 #include "RemoteControlPreset.h"
-#include "ScopedTransaction.h"
 #include "SDropTarget.h"
-#include "SRCPanelTreeNode.h"
-#include "SRCPanelExposedField.h"
 #include "SRCPanelDragHandle.h"
+#include "SRCPanelExposedField.h"
+#include "SRCPanelTreeNode.h"
 #include "SRemoteControlPanel.h"
-#include "Styling/SlateColor.h"
+#include "ScopedTransaction.h"
 #include "Styling/SlateBrush.h"
+#include "Styling/SlateColor.h"
 #include "Widgets/Input/SButton.h"
+#include "Widgets/SBoxPanel.h"
 #include "Widgets/Text/SInlineEditableTextBlock.h"
 #include "Widgets/Text/STextBlock.h"
-#include "Widgets/SBoxPanel.h"
 
 #define LOCTEXT_NAMESPACE "RemoteControlPanel"
 
@@ -255,6 +256,7 @@ bool SRCPanelGroup::OnVerifyItemLabelChanged(const FText& InLabel, FText& OutErr
 	if (TentativeName != Name && !!Preset->Layout.GetGroupByName(TentativeName))
 	{
 		OutErrorMessage = LOCTEXT("NameAlreadyExists", "This name already exists.");
+		IRemoteControlModule::BroadcastError(OutErrorMessage.ToString());
 		return false;
 	}
 
diff --git a/Engine/Plugins/VirtualProduction/RemoteControl/Source/WebRemoteControl/Private/RemoteControlWebSocketServer.cpp b/Engine/Plugins/VirtualProduction/RemoteControl/Source/WebRemoteControl/Private/RemoteControlWebSocketServer.cpp
--- a/Engine/Plugins/VirtualProduction/RemoteControl/Source/WebRemoteControl/Private/RemoteControlWebSocketServer.cpp
+++ b/Engine/Plugins/VirtualProduction/RemoteControl/Source/WebRemoteControl/Private/RemoteControlWebSocketServer.cpp
@@ -1,14 +1,15 @@
 // Copyright Epic Games, Inc. All Rights Reserved.
 
 #include "RemoteControlWebSocketServer.h"
+
 #include "Containers/Ticker.h"
-#include "WebRemoteControlInternalUtils.h"
 #include "IPAddress.h"
 #include "IRemoteControlModule.h"
 #include "IWebSocketNetworkingModule.h"
 #include "RemoteControlRequest.h"
-#include "Sockets.h"
 #include "SocketSubsystem.h"
+#include "Sockets.h"
+#include "WebRemoteControlInternalUtils.h"
 #include "WebSocketNetworkingDelegates.h"
 
 #define LOCTEXT_NAMESPACE "RCWebSocketServer"
@@ -56,6 +57,11 @@ namespace RemoteControlWebSocketServer
 			ParsedMessage = MoveTemp(Message);
 		}
 
+		if (!ErrorText.IsEmpty())
+		{
+			IRemoteControlModule::BroadcastError(ErrorText);
+		}
+
 		return ParsedMessage;
 	}
 }
diff --git a/Engine/Plugins/VirtualProduction/RemoteControl/Source/WebRemoteControl/Private/WebRemoteControlInternalUtils.cpp b/Engine/Plugins/VirtualProduction/RemoteControl/Source/WebRemoteControl/Private/WebRemoteControlInternalUtils.cpp
--- a/Engine/Plugins/VirtualProduction/RemoteControl/Source/WebRemoteControl/Private/WebRemoteControlInternalUtils.cpp
+++ b/Engine/Plugins/VirtualProduction/RemoteControl/Source/WebRemoteControl/Private/WebRemoteControlInternalUtils.cpp
@@ -1,12 +1,12 @@
 // Copyright Epic Games, Inc. All Rights Reserved.
 
 #include "WebRemoteControlInternalUtils.h"
+
 #include "HttpServerRequest.h"
 #include "Misc/Base64.h"
 #include "Serialization/JsonReader.h"
 #include "UObject/StructOnScope.h"
 
-
 namespace RemotePayloadSerializer
 {
 FName NAME_Get = "GET";
@@ -164,7 +164,10 @@ bool DeserializeCall(const FHttpServerRequest& InRequest, FRCCall& OutCall, cons
 
 	if (!bSuccess)
 	{
-		UE_LOG(LogRemoteControl, Error, TEXT("Web Remote Call deserialization error: %s"), *ErrorText);
+		const FString FormattedMsg = FString::Printf(TEXT("Web Remote Call deserialization error: %s"), *ErrorText);
+		UE_LOG(LogRemoteControl, Error, TEXT("%s"), *FormattedMsg);
+		IRemoteControlModule::BroadcastError(FormattedMsg);
+
 		TUniquePtr<FHttpServerResponse> Response = WebRemoteControlInternalUtils::CreateHttpResponse();
 		WebRemoteControlInternalUtils::CreateUTF8ErrorMessage(ErrorText, Response->Body);
 		InCompleteCallback(MoveTemp(Response));
@@ -226,7 +229,10 @@ bool DeserializeObjectRef(const FHttpServerRequest& InRequest, FRCObjectReferenc
 
 	if (!ErrorText.IsEmpty())
 	{
-		UE_LOG(LogRemoteControl, Error, TEXT("Web Remote Object Access error: %s"), *ErrorText);
+		const FString FormattedMsg = FString::Printf(TEXT("Web Remote Object Access error: %s"), *ErrorText);
+		UE_LOG(LogRemoteControl, Error, TEXT("%s"), *FormattedMsg);
+		IRemoteControlModule::BroadcastError(FormattedMsg);
+
 		TUniquePtr<FHttpServerResponse> Response = WebRemoteControlInternalUtils::CreateHttpResponse();
 		WebRemoteControlInternalUtils::CreateUTF8ErrorMessage(ErrorText, Response->Body);
 		InCompleteCallback(MoveTemp(Response));
@@ -301,7 +307,10 @@ bool WebRemoteControlInternalUtils::GetStructParametersDelimiters(TConstArrayVie
 
 	if (!ErrorText.IsEmpty())
 	{
-		UE_LOG(LogRemoteControl, Error, TEXT("Web Remote Control deserialization error: %s"), *ErrorText);
+		const FString FormattedMsg = FString::Printf(TEXT("Web Remote Control deserialization error: %s"), *ErrorText);
+		UE_LOG(LogRemoteControl, Error, TEXT("%s"), *FormattedMsg);
+		IRemoteControlModule::BroadcastError(FormattedMsg);
+
 		if (OutErrorText)
 		{
 			*OutErrorText = MoveTemp(ErrorText);
@@ -370,7 +379,10 @@ bool WebRemoteControlInternalUtils::GetBatchRequestStructDelimiters(TConstArrayV
 
 	if (!ErrorText.IsEmpty())
 	{
-		UE_LOG(LogRemoteControl, Error, TEXT("Web Remote Control deserialization error: %s"), *ErrorText);
+		const FString FormattedMsg = FString::Printf(TEXT("Web Remote Control deserialization error: %s"), *ErrorText);
+		UE_LOG(LogRemoteControl, Error, TEXT("%s"), *FormattedMsg);
+		IRemoteControlModule::BroadcastError(FormattedMsg);
+
 		if (OutErrorText)
 		{
 			*OutErrorText = MoveTemp(ErrorText);
@@ -450,7 +462,10 @@ bool WebRemoteControlInternalUtils::GetBatchWebSocketRequestStructDelimiters(TCo
 
 	if (!ErrorText.IsEmpty())
 	{
-		UE_LOG(LogRemoteControl, Error, TEXT("Web Remote Control deserialization error: %s"), *ErrorText);
+		const FString FormattedMsg = FString::Printf(TEXT("Web Remote Control deserialization error: %s"), *ErrorText);
+		UE_LOG(LogRemoteControl, Error, TEXT("%s"), *FormattedMsg);
+		IRemoteControlModule::BroadcastError(FormattedMsg);
+		
 		if (OutErrorText)
 		{
 			*OutErrorText = MoveTemp(ErrorText);
@@ -508,9 +523,13 @@ bool WebRemoteControlInternalUtils::IsRequestContentType(const FHttpServerReques
 		}
 	}
 
+	const FString ErrorText = FString::Printf(TEXT("Request content type must be %s"), *InContentType);
+	IRemoteControlModule::Get().BroadcastError(ErrorText);
+
 	if (OutErrorText)
 	{
-		*OutErrorText = FString::Printf(TEXT("Request content type must be %s"), *InContentType);
+		*OutErrorText = ErrorText;
 	}
+	
 	return false;
 }
