Add remote control security settings

diff --git a/Engine/Plugins/VirtualProduction/RemoteControl/Source/RemoteControl/Private/RemoteControlModule.cpp b/Engine/Plugins/VirtualProduction/RemoteControl/Source/RemoteControl/Private/RemoteControlModule.cpp
--- a/Engine/Plugins/VirtualProduction/RemoteControl/Source/RemoteControl/Private/RemoteControlModule.cpp
+++ b/Engine/Plugins/VirtualProduction/RemoteControl/Source/RemoteControl/Private/RemoteControlModule.cpp
@@ -810,28 +810,36 @@ bool FRemoteControlModule::ResolveCall(const FString& ObjectPath, const FString&
 		UObject* Object = StaticFindObject(UObject::StaticClass(), nullptr, *ObjectPath);
 		if (Object)
 		{
-			// Find the function to call
-			UFunction* Function = RemoteControlUtil::FindFunctionByNameOrMetaDataName(Object, FunctionName);
-
-			if (!Function)
-			{
-				ErrorText = FString::Printf(TEXT("Function: %s does not exist on object: %s"), *FunctionName, *ObjectPath);
-				bSuccess = false;
-			}
-			else if ((!Function->HasAllFunctionFlags(FUNC_BlueprintCallable | FUNC_Public) && !Function->HasAllFunctionFlags(FUNC_BlueprintEvent))
-#if WITH_EDITOR
-					|| Function->HasMetaData(RemoteControlUtil::NAME_DeprecatedFunction)
-					|| Function->HasMetaData(RemoteControlUtil::NAME_ScriptNoExport)
-#endif
-			)
+			if (CanBeAccessedRemotely(Object))
 			{
-				ErrorText = FString::Printf(TEXT("Function: %s is deprecated or unavailable remotely on object: %s"), *FunctionName, *ObjectPath);
-				bSuccess = false;
+				// Find the function to call
+				UFunction* Function = RemoteControlUtil::FindFunctionByNameOrMetaDataName(Object, FunctionName);
+
+				if (!Function)
+				{
+					ErrorText = FString::Printf(TEXT("Function: %s does not exist on object: %s"), *FunctionName, *ObjectPath);
+					bSuccess = false;
+				}
+				else if ((!Function->HasAllFunctionFlags(FUNC_BlueprintCallable | FUNC_Public) && !Function->HasAllFunctionFlags(FUNC_BlueprintEvent))
+	#if WITH_EDITOR
+						|| Function->HasMetaData(RemoteControlUtil::NAME_DeprecatedFunction)
+						|| Function->HasMetaData(RemoteControlUtil::NAME_ScriptNoExport)
+	#endif
+				)
+				{
+					ErrorText = FString::Printf(TEXT("Function: %s is deprecated or unavailable remotely on object: %s"), *FunctionName, *ObjectPath);
+					bSuccess = false;
+				}
+				else
+				{
+					OutCallRef.Object = Object;
+					OutCallRef.Function = Function;
+				}
 			}
 			else
 			{
-				OutCallRef.Object = Object;
-				OutCallRef.Function = Function;
+				ErrorText = FString::Printf(TEXT("Object %s cannot be accessed remotely, check remote control project settings."), *Object->GetName());
+				bSuccess = false;
 			}
 		}
 		else
@@ -920,7 +928,7 @@ bool FRemoteControlModule::InvokeCall(FRCCall& InCall, ERCPayloadType InPayloadT
 		FEditorScriptExecutionGuard ScriptGuard;
 		if (ensureAlways(InCall.CallRef.Object.IsValid()))
 		{
-			if(InCall.CallRef.PropertyWithSetter.IsValid())
+			if (InCall.CallRef.PropertyWithSetter.IsValid())
 			{
 				InCall.CallRef.PropertyWithSetter->CallSetter(InCall.CallRef.Object.Get(), InCall.ParamData.GetData());
 			}
@@ -975,7 +983,15 @@ bool FRemoteControlModule::ResolveObject(ERCAccess AccessType, const FString& Ob
 		UObject* Object = StaticFindObject(UObject::StaticClass(), nullptr, *ObjectPath);
 		if (Object)
 		{
-			bSuccess = ResolveObjectProperty(AccessType, Object, PropertyName, OutObjectRef, OutErrorText);
+			if (CanBeAccessedRemotely(Object))
+			{
+				bSuccess = ResolveObjectProperty(AccessType, Object, PropertyName, OutObjectRef, OutErrorText);
+			}
+			else
+			{
+				ErrorText = FString::Printf(TEXT("Object %s cannot be accessed remotely, check remote control project settings."), *Object->GetName());
+				bSuccess = false;
+			}
 		}
 		else
 		{
@@ -2283,7 +2299,28 @@ FRemoteControlModule::FOngoingChange::FOngoingChange(FRCCallReference InReferenc
 {
 	Reference.Set<FRCCallReference>(MoveTemp(InReference));
 }
-	
+
+bool FRemoteControlModule::CanBeAccessedRemotely(UObject* Object) const
+{
+	check(Object);
+		
+	if (!GetDefault<URemoteControlSettings>()->bEnableRemotePythonExecution)
+	{
+		static const FName PythonScriptName = "PythonScriptLibrary"; 
+		if (Object->GetClass()->GetFName() == PythonScriptName)
+		{
+			return false;
+		}
+	}
+
+	if (Object == GetDefault<URemoteControlSettings>())
+	{
+		return false;
+	}
+
+	return true;
+}
+
 #endif
 
 PRAGMA_ENABLE_DEPRECATION_WARNINGS
diff --git a/Engine/Plugins/VirtualProduction/RemoteControl/Source/RemoteControl/Private/RemoteControlModule.h b/Engine/Plugins/VirtualProduction/RemoteControl/Source/RemoteControl/Private/RemoteControlModule.h
--- a/Engine/Plugins/VirtualProduction/RemoteControl/Source/RemoteControl/Private/RemoteControlModule.h
+++ b/Engine/Plugins/VirtualProduction/RemoteControl/Source/RemoteControl/Private/RemoteControlModule.h
@@ -64,6 +64,7 @@ public:
 	virtual FGuid BeginManualEditorTransaction(const FText& InDescription, uint32 TypeHash) override;
 	virtual int32 EndManualEditorTransaction(const FGuid& TransactionId) override;
 	virtual const TMap<FName, TSharedPtr<IRemoteControlPropertyFactory>>& GetEntityFactories() const override { return EntityFactories; };
+	virtual bool CanBeAccessedRemotely(UObject* Object) const override;
 	//~ End IRemoteControlModule
 
 private:
diff --git a/Engine/Plugins/VirtualProduction/RemoteControl/Source/RemoteControl/Private/RemoteControlPreset.cpp b/Engine/Plugins/VirtualProduction/RemoteControl/Source/RemoteControl/Private/RemoteControlPreset.cpp
--- a/Engine/Plugins/VirtualProduction/RemoteControl/Source/RemoteControl/Private/RemoteControlPreset.cpp
+++ b/Engine/Plugins/VirtualProduction/RemoteControl/Source/RemoteControl/Private/RemoteControlPreset.cpp
@@ -685,6 +685,12 @@ TWeakPtr<FRemoteControlProperty> URemoteControlPreset::ExposeProperty(UObject* O
 	{
 		return nullptr;
 	}
+	
+	if (!IRemoteControlModule::Get().CanBeAccessedRemotely(Object))
+    {
+    	UE_LOG(LogRemoteControl, Warning, TEXT("Object %s cannot be accessed remotely. Check project settings."), *Object->GetName());
+		return nullptr;
+    }
 
 	TSharedPtr<FRemoteControlProperty> RCPropertyPtr;
 	const TMap<FName, TSharedPtr<IRemoteControlPropertyFactory>>& PropertyFactories = IRemoteControlModule::Get().GetEntityFactories();
@@ -740,6 +746,12 @@ TWeakPtr<FRemoteControlFunction> URemoteControlPreset::ExposeFunction(UObject* O
 		return nullptr;
 	}
 
+	if (!IRemoteControlModule::Get().CanBeAccessedRemotely(Object))
+	{
+		UE_LOG(LogRemoteControl, Warning, TEXT("Object %s cannot be accessed remotely. Check project settings."), *Object->GetName());
+		return nullptr;
+	}
+
 	FName DesiredName = *Args.Label;
 
 	if (DesiredName == NAME_None)
diff --git a/Engine/Plugins/VirtualProduction/RemoteControl/Source/RemoteControl/Public/IRemoteControlModule.h b/Engine/Plugins/VirtualProduction/RemoteControl/Source/RemoteControl/Public/IRemoteControlModule.h
--- a/Engine/Plugins/VirtualProduction/RemoteControl/Source/RemoteControl/Public/IRemoteControlModule.h
+++ b/Engine/Plugins/VirtualProduction/RemoteControl/Source/RemoteControl/Public/IRemoteControlModule.h
@@ -533,5 +533,11 @@ public:
 
 	/** Get map of the factories which is responsible for the Remote Control property creation */
 	virtual const TMap<FName, TSharedPtr<IRemoteControlPropertyFactory>>& GetEntityFactories() const = 0;
+
+	/**
+	 *  @return whether an object can be accessed remotely.
+	 *  @note Check Remote Control project settings to configure.
+	 */
+	virtual bool CanBeAccessedRemotely(UObject* InObject) const = 0;
 };
 
diff --git a/Engine/Plugins/VirtualProduction/RemoteControl/Source/RemoteControlCommon/Public/RemoteControlSettings.h b/Engine/Plugins/VirtualProduction/RemoteControl/Source/RemoteControlCommon/Public/RemoteControlSettings.h
--- a/Engine/Plugins/VirtualProduction/RemoteControl/Source/RemoteControlCommon/Public/RemoteControlSettings.h
+++ b/Engine/Plugins/VirtualProduction/RemoteControl/Source/RemoteControlCommon/Public/RemoteControlSettings.h
@@ -73,7 +73,7 @@ public:
 	 * Disabling transactions improves performance but will prevent events from being transacted to Multi-User
 	 * unless using the Remote Control Interception feature.
 	 */
-	UPROPERTY(config, EditAnywhere, Category = RemoteControl)
+	UPROPERTY(config, EditAnywhere, Category = "Remote Control")
 	bool bProtocolsGenerateTransactions = true;
 
 	/** The remote control web app http port. */
@@ -88,14 +88,6 @@ public:
 	UPROPERTY(config, EditAnywhere, Category = "Remote Control Web Interface", DisplayName = "Log WebApp requests handle duration")
 	bool bWebAppLogRequestDuration = false;
 
-	/** Whether communication with the Web Interface should only be allowed with an Passphrase */
-	UPROPERTY(Config, EditAnywhere, Category = "Remote Control Web Interface", DisplayName = "Use Passphrase to block Access")
-	bool bUseRemoteControlPassphrase = false;
-
-	/** Whether the User should be warned that Passphrase usage is disabled or now. Initially activated */
-	UPROPERTY(Config, EditAnywhere, Category = "Remote Control Web Interface", DisplayName = "Warn that Passphrase might be disabled ")
-	bool bShowPassphraseDisabledWarning = true;
-
 	/** Whether web server is started automatically. */
 	UPROPERTY(config, EditAnywhere, Category = "Remote Control Web Server")
 	bool bAutoStartWebServer = true;
@@ -129,9 +121,21 @@ public:
 	UPROPERTY(Config, EditAnywhere, Category = "Remote Control Preset", DisplayName = "Ignore Remote Control Getter/Setter Check")
 	bool bIgnoreGetterSetterCheck = false;
 
-	UPROPERTY(Config, EditAnywhere, Category = "Remote Control Preset", DisplayName = "Ignore Warnings")
+	UPROPERTY(Config, EditAnywhere, Category = "Remote Control Preset")
 	bool bIgnoreWarnings = false;
+	
+	/** Whether communication with the Web Interface should only be allowed with an Passphrase */
+	UPROPERTY(Config, EditAnywhere, Category = "Remote Control | Security", DisplayName = "Use Passphrase to block Access")
+	bool bUseRemoteControlPassphrase = false;
+
+	/** Whether the User should be warned that Passphrase usage is disabled or now. Initially activated */
+	UPROPERTY(Config, EditAnywhere, Category = "Remote Control | Security", DisplayName = "Warn that Passphrase might be disabled ")
+	bool bShowPassphraseDisabledWarning = true;
+
+	/** Enable remote python execution, enabling this could open you open to vulnerabilities if an outside actor has access to your server. */
+	UPROPERTY(Config, EditAnywhere, Category = "Remote Control | Security")
+	bool bEnableRemotePythonExecution = false;
 
-	UPROPERTY(config, EditAnywhere, Category = "Remote Control Web Interface", DisplayName = "Remote Control Web Interface Passphrase")
+	UPROPERTY(config, EditAnywhere, Category = "Remote Control | Security", DisplayName = "Remote Control Passphrase")
 	TArray<FRCPassphrase> Passphrases = {};
 };
diff --git a/Engine/Plugins/VirtualProduction/RemoteControl/Source/WebRemoteControl/Private/WebRemoteControlInternalUtils.cpp b/Engine/Plugins/VirtualProduction/RemoteControl/Source/WebRemoteControl/Private/WebRemoteControlInternalUtils.cpp
--- a/Engine/Plugins/VirtualProduction/RemoteControl/Source/WebRemoteControl/Private/WebRemoteControlInternalUtils.cpp
+++ b/Engine/Plugins/VirtualProduction/RemoteControl/Source/WebRemoteControl/Private/WebRemoteControlInternalUtils.cpp
@@ -3,7 +3,9 @@
 #include "WebRemoteControlInternalUtils.h"
 
 #include "HttpServerRequest.h"
+#include "Kismet/KismetSystemLibrary.h"
 #include "Misc/Base64.h"
+#include "RemoteControlSettings.h"
 #include "Serialization/JsonReader.h"
 #include "UObject/StructOnScope.h"
 
@@ -150,7 +152,14 @@ bool DeserializeCall(const FHttpServerRequest& InRequest, FRCCall& OutCall, cons
 			Reader.SetLimitSize(ParametersDelimiters.BlockEnd);
 
 			FRCJsonStructDeserializerBackend Backend(Reader);
-			if (!FStructDeserializer::Deserialize((void*)OutCall.ParamStruct.GetStructMemory(), *const_cast<UStruct*>(OutCall.ParamStruct.GetStruct()), Backend, FStructDeserializerPolicies()))
+			if (FStructDeserializer::Deserialize((void*)OutCall.ParamStruct.GetStructMemory(), *const_cast<UStruct*>(OutCall.ParamStruct.GetStruct()), Backend, FStructDeserializerPolicies()))
+			{
+				if (!WebRemoteControlInternalUtils::ValidateFunctionCall(OutCall, &ErrorText))
+				{
+					bSuccess = false;
+				}
+			}
+			else
 			{
 				ErrorText = TEXT("Parameters object improperly formatted.");
 				bSuccess = false;
@@ -489,6 +498,49 @@ bool WebRemoteControlInternalUtils::ValidateContentType(const FHttpServerRequest
 	return true;
 }
 
+bool WebRemoteControlInternalUtils::ValidateFunctionCall(const FRCCall& InRCCall, FString* OutErrorText)
+{
+	if (!InRCCall.IsValid())
+	{
+		return false;
+	}
+	
+	if (!GetDefault<URemoteControlSettings>()->bEnableRemotePythonExecution
+		&& InRCCall.CallRef.Object->IsA(UKismetSystemLibrary::StaticClass())
+		&& InRCCall.CallRef.Function == UKismetSystemLibrary::StaticClass()->FindFunctionByName(GET_FUNCTION_NAME_CHECKED(UKismetSystemLibrary, ExecuteConsoleCommand)))
+	{
+		// Make sure python is not called through KismetSystemLibrary.
+		if (const FProperty* Property = InRCCall.CallRef.Function->FindPropertyByName("Command"))
+		{
+			FString FullCommand;
+			Property->GetValue_InContainer(InRCCall.ParamStruct.GetStructMemory(), &FullCommand);
+
+			TArray<FString> SplitCommands;
+			FullCommand.ParseIntoArray(SplitCommands, TEXT("|"));
+
+			for (FString& Command : SplitCommands)
+			{
+				Command.TrimStartAndEndInline();
+				if (Command.StartsWith(TEXT("py ")) || Command.StartsWith(TEXT("python ")))
+				{
+					if (OutErrorText)
+					{
+						*OutErrorText = TEXT("Executing Python remotely is not enabled in the remote control settings.");
+					}
+					return false;
+				}
+			}
+		}
+		else
+		{
+			ensureMsgf(false, TEXT("Could not find the Command parameter on ExecuteConsoleCommand"));
+			return false;
+		}
+	}
+
+	return true;
+}
+
 void WebRemoteControlInternalUtils::AddContentTypeHeaders(FHttpServerResponse* InOutResponse, FString InContentType)
 {
 	InOutResponse->Headers.Add(TEXT("content-type"), { MoveTemp(InContentType) });
diff --git a/Engine/Plugins/VirtualProduction/RemoteControl/Source/WebRemoteControl/Private/WebRemoteControlInternalUtils.h b/Engine/Plugins/VirtualProduction/RemoteControl/Source/WebRemoteControl/Private/WebRemoteControlInternalUtils.h
--- a/Engine/Plugins/VirtualProduction/RemoteControl/Source/WebRemoteControl/Private/WebRemoteControlInternalUtils.h
+++ b/Engine/Plugins/VirtualProduction/RemoteControl/Source/WebRemoteControl/Private/WebRemoteControlInternalUtils.h
@@ -319,6 +319,13 @@ namespace WebRemoteControlInternalUtils
 	 */
 	UE_NODISCARD bool ValidateContentType(const FHttpServerRequest& InRequest, FString InContentType, const FHttpResultCallback& InCompleteCallback);
 
+	/**
+	 * Check if a function call is valid since some objects//functions are disabled remotely for security reasons.
+	 * @param InRCCall The RC call to validate.
+	 * @param OutErrorText Optional error text.
+	 **/
+	UE_NODISCARD bool ValidateFunctionCall(const FRCCall& InRCCall, FString* OutErrorText);
+
 	/**
 	 * Add the desired content type to the http response headers.
 	 * @param InResponse The response to add the content type to.
