1. RCP Bind Logic is a step behind in transactions

2. 'Reset to Default' setting point light intensity at 5,000 in Web App

diff --git a/Engine/Plugins/VirtualProduction/RemoteControl/Source/RemoteControl/Private/Factories/RCDefaultValueFactories.cpp b/Engine/Plugins/VirtualProduction/RemoteControl/Source/RemoteControl/Private/Factories/RCDefaultValueFactories.cpp
new file mode 100644
--- /dev/null
+++ b/Engine/Plugins/VirtualProduction/RemoteControl/Source/RemoteControl/Private/Factories/RCDefaultValueFactories.cpp
@@ -0,0 +1,63 @@
+// Copyright Epic Games, Inc. All Rights Reserved.
+
+#include "Factories/RCDefaultValueFactories.h"
+#include "Components/LightComponent.h"
+
+TSharedRef<IRCDefaultValueFactory> FLightIntensityDefaultValueFactory::MakeInstance()
+{
+	return MakeShared<FLightIntensityDefaultValueFactory>();
+}
+
+bool FLightIntensityDefaultValueFactory::CanResetToDefaultValue(UObject* InObject) const
+{
+	if (const ULightComponent* LightComponent = Cast<ULightComponent>(InObject))
+	{
+		if (const ULightComponent* ArchetypeComponent = Cast<ULightComponent>(LightComponent->GetArchetype()))
+		{
+			return !FMath::IsNearlyEqual(LightComponent->ComputeLightBrightness(), ArchetypeComponent->ComputeLightBrightness());
+		}
+	}
+
+	return false;
+}
+
+void FLightIntensityDefaultValueFactory::ResetToDefaultValue(UObject* InObject, FProperty* InProperty)
+{
+#if WITH_EDITOR
+	
+	InObject->PreEditChange(InProperty);
+	InObject->Modify();
+
+	if (ULightComponent* LightComponent = Cast<ULightComponent>(InObject))
+	{
+		if (const ULightComponent* ArchetypeComponent = Cast<ULightComponent>(LightComponent->GetArchetype()))
+		{
+			LightComponent->SetLightBrightness(ArchetypeComponent->ComputeLightBrightness());
+		}
+	}
+
+	FPropertyChangedEvent ChangeEvent(InProperty, EPropertyChangeType::ValueSet);
+	InObject->PostEditChangeProperty(ChangeEvent);
+
+#endif // WITH_EDITOR
+}
+
+bool FLightIntensityDefaultValueFactory::SupportsClass(const UClass* InObjectClass) const
+{
+	if (!InObjectClass)
+	{
+		return false;
+	}
+
+	return InObjectClass->IsChildOf<ULightComponentBase>();
+}
+
+bool FLightIntensityDefaultValueFactory::SupportsProperty(const FProperty* InProperty) const
+{
+	if (!InProperty)
+	{
+		return false;
+	}
+
+	return InProperty->GetFName() == GET_MEMBER_NAME_CHECKED(ULightComponentBase, Intensity);
+}
diff --git a/Engine/Plugins/VirtualProduction/RemoteControl/Source/RemoteControl/Private/Factories/RCDefaultValueFactories.h b/Engine/Plugins/VirtualProduction/RemoteControl/Source/RemoteControl/Private/Factories/RCDefaultValueFactories.h
new file mode 100644
--- /dev/null
+++ b/Engine/Plugins/VirtualProduction/RemoteControl/Source/RemoteControl/Private/Factories/RCDefaultValueFactories.h
@@ -0,0 +1,22 @@
+// Copyright Epic Games, Inc. All Rights Reserved.
+
+#pragma once
+
+#include "CoreMinimal.h"
+#include "Factories/IRCDefaultValueFactory.h"
+
+/**
+ * A specialized default value factory for handling light components.
+ */
+class FLightIntensityDefaultValueFactory : public IRCDefaultValueFactory
+{
+public:
+	static TSharedRef<IRCDefaultValueFactory> MakeInstance();
+
+	//~ Begin IRCDefaultValueFactory interface
+	virtual bool CanResetToDefaultValue(UObject* InObject) const override;
+	virtual void ResetToDefaultValue(UObject* InObject, FProperty* InProperty) override;
+	virtual bool SupportsClass(const UClass* InObjectClass) const override;
+	virtual bool SupportsProperty(const FProperty* InProperty) const override;
+	//~ End IRCDefaultValueFactory interface
+};
diff --git a/Engine/Plugins/VirtualProduction/RemoteControl/Source/RemoteControl/Private/RemoteControlModule.cpp b/Engine/Plugins/VirtualProduction/RemoteControl/Source/RemoteControl/Private/RemoteControlModule.cpp
--- a/Engine/Plugins/VirtualProduction/RemoteControl/Source/RemoteControl/Private/RemoteControlModule.cpp
+++ b/Engine/Plugins/VirtualProduction/RemoteControl/Source/RemoteControl/Private/RemoteControlModule.cpp
@@ -7,6 +7,7 @@
 #include "Backends/CborStructSerializerBackend.h"
 #include "Components/ActorComponent.h"
 #include "Components/LightComponent.h"
+#include "Factories/RCDefaultValueFactories.h"
 #include "Factories/RemoteControlMaskingFactories.h"
 #include "Features/IModularFeatures.h"
 #include "IRemoteControlInterceptionFeature.h"
@@ -602,6 +603,9 @@ void FRemoteControlModule::StartupModule()
 	FCoreUObjectDelegates::PreLoadMap.AddRaw(this, &FRemoteControlModule::HandleMapPreLoad);
 #endif
 	
+	// Register Default Value Factories
+	RegisterDefaultValueFactories();
+
 	// Register Property Factories
 	RegisterEntityFactory(FRemoteControlInstanceMaterial::StaticStruct()->GetFName(), FRemoteControlInstanceMaterialFactory::MakeInstance());
 
@@ -632,6 +636,10 @@ void FRemoteControlModule::ShutdownModule()
 	{
 		// Unregister Property Factories
 		UnregisterEntityFactory(FRemoteControlInstanceMaterial::StaticStruct()->GetFName());
+		
+		// Unregister Default Value & Masking factories.
+		DefaultValueFactories.Empty();
+		MaskingFactories.Empty();
 	}
 }
 
@@ -740,6 +748,67 @@ void FRemoteControlModule::UnregisterEmbeddedPreset(URemoteControlPreset* Preset
 	EmbeddedPresets.Remove(PresetName);
 }
 
+bool FRemoteControlModule::CanResetToDefaultValue(UObject* InObject, const FProperty* InProperty) const
+{
+	if (!InObject || !InProperty)
+	{
+		return false;
+	}
+
+	// NOTE : Get the owning class to which the property actually belongs.
+	if (UClass* OwningClass = InProperty->GetOwnerClass())
+	{
+		const FString DefaultValueKey = OwningClass->GetName() + TEXT(".") + InProperty->GetName();
+		
+		if (const TSharedPtr<IRCDefaultValueFactory>* DefaultValueFactory = DefaultValueFactories.Find(*DefaultValueKey))
+		{
+			return (*DefaultValueFactory)->CanResetToDefaultValue(InObject);
+		}
+	}
+
+	return false;
+}
+
+bool FRemoteControlModule::HasDefaultValueCustomization(const UObject* InObject, const FProperty* InProperty) const
+{
+	if (!InObject || !InProperty)
+	{
+		return false;
+	}
+
+	// NOTE : Get the owning class to which the property actually belongs.
+	if (UClass* OwningClass = InProperty->GetOwnerClass())
+	{
+		const FString DefaultValueKey = OwningClass->GetName() + TEXT(".") + InProperty->GetName();
+
+		if (const TSharedPtr<IRCDefaultValueFactory>* DefaultValueFactory = DefaultValueFactories.Find(*DefaultValueKey))
+		{
+			return (*DefaultValueFactory)->SupportsClass(InObject->GetClass()) && (*DefaultValueFactory)->SupportsProperty(InProperty);
+		}
+	}
+	
+	return false;
+}
+
+void FRemoteControlModule::ResetToDefaultValue(UObject* InObject, FProperty* InProperty)
+{
+	if (!InObject || !InProperty)
+	{
+		return;
+	}
+
+	// NOTE : Get the owning class to which the property actually belongs.
+	if (UClass* OwningClass = InProperty->GetOwnerClass())
+	{
+		const FString DefaultValueKey = OwningClass->GetName() + TEXT(".") + InProperty->GetName();
+
+		if (const TSharedPtr<IRCDefaultValueFactory>* DefaultValueFactory = DefaultValueFactories.Find(*DefaultValueKey))
+		{
+			(*DefaultValueFactory)->ResetToDefaultValue(InObject, InProperty);
+		}
+	}
+}
+
 void FRemoteControlModule::PerformMasking(const TSharedRef<FRCMaskingOperation>& InMaskingOperation)
 {
 	TRACE_CPUPROFILER_EVENT_SCOPE(FRemoteControlModule::PerformMasking);
@@ -1565,15 +1634,22 @@ bool FRemoteControlModule::ResetObjectProperties(const FRCObjectReference& Objec
 			ObjectAccess.PropertyPathInfo.ToEditPropertyChain(PreEditChain);
 			Object->PreEditChange(PreEditChain);
 		}
-#endif
 
-		// Copy the value from the field on the CDO.
-		FRCFieldPathInfo FieldPathInfo = ObjectAccess.PropertyPathInfo;
-		void* TargetAddress = FieldPathInfo.GetResolvedData().ContainerAddress;
-		UObject* DefaultObject = Object->GetClass()->GetDefaultObject();
-		FieldPathInfo.Resolve(DefaultObject);
-		FRCFieldResolvedData DefaultObjectResolvedData = FieldPathInfo.GetResolvedData();
-		ObjectAccess.Property->CopyCompleteValue_InContainer(TargetAddress, DefaultObjectResolvedData.ContainerAddress);
+		if(HasDefaultValueCustomization(Object, ObjectAccess.Property.Get()))
+		{
+			ResetToDefaultValue(Object, ObjectAccess.Property.Get());
+		}
+		else
+#endif
+		{
+			// Copy the value from the field on the CDO.
+			FRCFieldPathInfo FieldPathInfo = ObjectAccess.PropertyPathInfo;
+			void* TargetAddress = FieldPathInfo.GetResolvedData().ContainerAddress;
+			UObject* DefaultObject = Object->GetClass()->GetDefaultObject();
+			FieldPathInfo.Resolve(DefaultObject);
+			FRCFieldResolvedData DefaultObjectResolvedData = FieldPathInfo.GetResolvedData();
+			ObjectAccess.Property->CopyCompleteValue_InContainer(TargetAddress, DefaultObjectResolvedData.ContainerAddress);
+		}
 
 		// if we are generating a transaction, also generate post edit property event, event if the change ended up unsuccessful
 		// this is to match the pre edit change call that can unregister components for example
@@ -2219,6 +2295,29 @@ bool FRemoteControlModule::DeserializeDeltaModificationData(const FRCObjectRefer
 	return bSuccess;
 }
 
+void FRemoteControlModule::RegisterDefaultValueFactoryForType(UClass* RemoteControlPropertyType, const FName PropertyName, const TSharedPtr<IRCDefaultValueFactory>& InDefaultValueFactory)
+{
+	const FString DefaultValueKey = RemoteControlPropertyType->GetName() + TEXT(".") + PropertyName.ToString();
+
+	if (!DefaultValueFactories.Contains(*DefaultValueKey))
+	{
+		DefaultValueFactories.Add(*DefaultValueKey, InDefaultValueFactory);
+	}
+}
+
+void FRemoteControlModule::UnregisterDefaultValueFactoryForType(UClass* RemoteControlPropertyType, const FName PropertyName)
+{
+	const FString DefaultValueKey = RemoteControlPropertyType->GetName() + TEXT(".") + PropertyName.ToString();
+
+	DefaultValueFactories.Remove(*DefaultValueKey);
+}
+
+void FRemoteControlModule::RegisterDefaultValueFactories()
+{
+	// NOTE : Use the class to which the property actually belongs.
+	RegisterDefaultValueFactoryForType(ULightComponentBase::StaticClass(), GET_MEMBER_NAME_CHECKED(ULightComponentBase, Intensity), FLightIntensityDefaultValueFactory::MakeInstance());
+}
+
 void FRemoteControlModule::RegisterMaskingFactories()
 {
 	RegisterMaskingFactoryForType(TBaseStructure<FVector>::Get(), FVectorMaskingFactory::MakeInstance());
diff --git a/Engine/Plugins/VirtualProduction/RemoteControl/Source/RemoteControl/Private/RemoteControlModule.h b/Engine/Plugins/VirtualProduction/RemoteControl/Source/RemoteControl/Private/RemoteControlModule.h
--- a/Engine/Plugins/VirtualProduction/RemoteControl/Source/RemoteControl/Private/RemoteControlModule.h
+++ b/Engine/Plugins/VirtualProduction/RemoteControl/Source/RemoteControl/Private/RemoteControlModule.h
@@ -4,6 +4,7 @@
 #include "AssetRegistry/AssetData.h"
 #include "CoreMinimal.h"
 #include "Engine/EngineTypes.h"
+#include "Factories/IRCDefaultValueFactory.h"
 #include "Factories/IRemoteControlMaskingFactory.h"
 #include "IRemoteControlInterceptionFeature.h"
 #include "IRemoteControlModule.h"
@@ -30,6 +31,9 @@ public:
 	virtual bool RegisterEmbeddedPreset(URemoteControlPreset* Preset, bool bReplaceExisting) override;
 	virtual void UnregisterEmbeddedPreset(FName Name) override;
 	virtual void UnregisterEmbeddedPreset(URemoteControlPreset* Preset) override;
+	virtual bool CanResetToDefaultValue(UObject* InObject, const FProperty* InProperty) const override;
+	virtual bool HasDefaultValueCustomization(const UObject* InObject, const FProperty* InProperty) const override;
+	virtual void ResetToDefaultValue(UObject* InObject, FProperty* InProperty) override;
 	virtual void PerformMasking(const TSharedRef<FRCMaskingOperation>& InMaskingOperation) override;
 	virtual void RegisterMaskingFactoryForType(UScriptStruct* RemoteControlPropertyType, const TSharedPtr<IRemoteControlMaskingFactory>& InMaskingFactory) override;
 	virtual void UnregisterMaskingFactoryForType(UScriptStruct* RemoteControlPropertyType) override;
@@ -96,7 +100,22 @@ private:
 	 * @return True if the data was successfully deserialized and modified.
 	 */
 	static bool DeserializeDeltaModificationData(const FRCObjectReference& ObjectAccess, IStructDeserializerBackend& Backend, ERCModifyOperation Operation, TArray<uint8>& OutData);
+	
+	/**
+	 * Register a default factory to handle that default value for the supported properties.
+	 */
+	void RegisterDefaultValueFactoryForType(UClass* RemoteControlPropertyType, const FName PropertyName, const TSharedPtr<IRCDefaultValueFactory>& InDefaultValueFactory);
+	
+	/**
+	 * Unregister a previously registered default value factory.
+	 */
+	void UnregisterDefaultValueFactoryForType(UClass* RemoteControlPropertyType, const FName PropertyName);
 
+	/**
+	 * Register(s) default value factories of supported types.
+	 */
+	void RegisterDefaultValueFactories();
+	
 	/**
 	 * Register(s) masking factories of supported types.
 	 */
@@ -204,6 +223,9 @@ private:
 
 	/** Holds the set of active masking operations. */
 	TSet<TSharedPtr<FRCMaskingOperation>> ActiveMaskingOperations;
+
+	/** Map of the factories which is responsible for resetting the Remote Control property to its default value. */
+	TMap<FName, TSharedPtr<IRCDefaultValueFactory>> DefaultValueFactories;
 };
 
 PRAGMA_ENABLE_DEPRECATION_WARNINGS
diff --git a/Engine/Plugins/VirtualProduction/RemoteControl/Source/RemoteControl/Private/RemoteControlPreset.cpp b/Engine/Plugins/VirtualProduction/RemoteControl/Source/RemoteControl/Private/RemoteControlPreset.cpp
--- a/Engine/Plugins/VirtualProduction/RemoteControl/Source/RemoteControl/Private/RemoteControlPreset.cpp
+++ b/Engine/Plugins/VirtualProduction/RemoteControl/Source/RemoteControl/Private/RemoteControlPreset.cpp
@@ -663,7 +663,14 @@ void URemoteControlPreset::OnModifyController(const FPropertyChangedEvent& Prope
 {
 	if (ensure(ControllerContainer))
 	{
-		if (URCVirtualPropertyBase* ModifiedController = ControllerContainer->GetVirtualProperty(PropertyChangedEvent.MemberProperty->GetFName()))
+		URCVirtualPropertyBase* ModifiedController = ControllerContainer->GetVirtualProperty(PropertyChangedEvent.MemberProperty->GetFName());
+
+		if (!ModifiedController)
+		{
+			ModifiedController = ControllerContainer->GetVirtualProperty(PropertyChangedEvent.MemberProperty->GetFName());
+		}
+
+		if (ModifiedController)
 		{
 			ControllerContainer->OnModifyPropertyValue(PropertyChangedEvent);
 			OnControllerModified().Broadcast(this, {ModifiedController->Id});
diff --git a/Engine/Plugins/VirtualProduction/RemoteControl/Source/RemoteControl/Public/Factories/IRCDefaultValueFactory.h b/Engine/Plugins/VirtualProduction/RemoteControl/Source/RemoteControl/Public/Factories/IRCDefaultValueFactory.h
new file mode 100644
--- /dev/null
+++ b/Engine/Plugins/VirtualProduction/RemoteControl/Source/RemoteControl/Public/Factories/IRCDefaultValueFactory.h
@@ -0,0 +1,43 @@
+// Copyright Epic Games, Inc. All Rights Reserved.
+
+#pragma once
+
+#include "CoreMinimal.h"
+
+/**
+ * Factory which resets the FRemoteControlProperty to its default values.
+ */
+class IRCDefaultValueFactory : public TSharedFromThis<IRCDefaultValueFactory>
+{
+public:
+
+	/** Virtual destructor */
+	virtual ~IRCDefaultValueFactory() {}
+
+	/**
+	 * Returns true when the given object can be reset to its default value, false otherwise.
+	 * @param InObject Reference to the exposed object.
+	 */
+	virtual bool CanResetToDefaultValue(UObject* InObject) const = 0;
+
+	/**
+	 * Performs actual data reset on the given remote object.
+	 * @param InObject Reference to the exposed object.
+	 * @param InProperty Reference to the exposed property.
+	 */
+	virtual void ResetToDefaultValue(UObject* InObject, FProperty* InProperty) = 0;
+	
+	/**
+	 * Whether the factory support exposed entity.
+	 * @param InObjectClass Reference to the exposed object class.
+	 * @return true if the exposed object is supported by given factory
+	 */
+	virtual bool SupportsClass(const UClass* InObjectClass) const = 0;
+
+	/**
+	 * Whether the factory support property.
+	 * @param InProperty Reference to the exposed property.
+	 * @return true if the exposed property is supported by given factory
+	 */
+	virtual bool SupportsProperty(const FProperty* InProperty) const = 0;
+};
diff --git a/Engine/Plugins/VirtualProduction/RemoteControl/Source/RemoteControl/Public/IRemoteControlModule.h b/Engine/Plugins/VirtualProduction/RemoteControl/Source/RemoteControl/Public/IRemoteControlModule.h
--- a/Engine/Plugins/VirtualProduction/RemoteControl/Source/RemoteControl/Public/IRemoteControlModule.h
+++ b/Engine/Plugins/VirtualProduction/RemoteControl/Source/RemoteControl/Public/IRemoteControlModule.h
@@ -279,6 +279,27 @@ public:
 	virtual void UnregisterEmbeddedPreset(FName Name) = 0;
 	virtual void UnregisterEmbeddedPreset(URemoteControlPreset* Preset) = 0;
 
+	/**
+	 * Returns true when the given object can be reset to its default value, false otherwise.
+	 * @param InObject Reference to the exposed object.
+	 * @param InProperty Reference to the exposed property.
+	 */
+	virtual bool CanResetToDefaultValue(UObject* InObject, const FProperty* InProperty) const = 0;
+
+	/**
+	 * Returns true when the given object and property has custom default value, false otherwise.
+	 * @param InObject Reference to the exposed object.
+	 * @param InProperty Reference to the exposed property.
+	 */
+	virtual bool HasDefaultValueCustomization(const UObject* InObject, const FProperty* InProperty) const = 0;
+
+	/**
+	 * Performs actual data reset on the given object.
+	 * @param InObject Reference to the exposed object.
+	 * @param InProperty Reference to the exposed property.
+	 */
+	virtual void ResetToDefaultValue(UObject* InObject, FProperty* InProperty) = 0;
+
 	/**
 	 * Performs the given masking operation.
 	 * @param InMaskingOperation Masking operation to be performed.
diff --git a/Engine/Plugins/VirtualProduction/RemoteControl/Source/RemoteControlUI/Private/UI/SRCPanelExposedField.cpp b/Engine/Plugins/VirtualProduction/RemoteControl/Source/RemoteControlUI/Private/UI/SRCPanelExposedField.cpp
--- a/Engine/Plugins/VirtualProduction/RemoteControl/Source/RemoteControlUI/Private/UI/SRCPanelExposedField.cpp
+++ b/Engine/Plugins/VirtualProduction/RemoteControl/Source/RemoteControlUI/Private/UI/SRCPanelExposedField.cpp
@@ -3,6 +3,7 @@
 #include "SRCPanelExposedField.h"
 
 #include "Algo/Transform.h"
+#include "Components/LocalLightComponent.h"
 #include "EditorFontGlyphs.h"
 #include "IDetailTreeNode.h"
 #include "IRCProtocolBindingList.h"
@@ -419,7 +420,7 @@ TSharedRef<SWidget> SRCPanelExposedField::ConstructWidget()
 							.VAlign(VAlign_Center)
 							.Padding(2.f, 4.f)
 							[
-								SNew(SResetToDefaultPropertyEditor, PropertyHandle)
+								ConstructResetToDefaultWidget(Object, PropertyHandle)
 							];
 					}
 
@@ -569,6 +570,48 @@ void SRCPanelExposedField::ConstructFunctionWidget()
 	ChildSlot.AttachWidget(ConstructWidget());
 }
 
+TSharedRef<SWidget> SRCPanelExposedField::ConstructResetToDefaultWidget(UObject* InObject, TSharedPtr<IPropertyHandle> InPropertyHandle)
+{
+	TSharedPtr<SResetToDefaultPropertyEditor> ResetDefaultProperty = SNew(SResetToDefaultPropertyEditor, InPropertyHandle);
+	
+	const TAttribute<bool> CanResetToDefaultValue = TAttribute<bool>::Create(TAttribute<bool>::FGetter::CreateLambda([this, InObject, InPropertyHandle]()
+		{
+			IRemoteControlModule& RemoteControlModule = IRemoteControlModule::Get();
+		
+			if (RemoteControlModule.HasDefaultValueCustomization(InObject, InPropertyHandle->GetProperty()))
+			{
+				return RemoteControlModule.CanResetToDefaultValue(InObject, InPropertyHandle->GetProperty());
+			}
+
+			return InPropertyHandle->DiffersFromDefault();
+		})
+	);
+
+	const FSimpleDelegate PerformValueReset = FSimpleDelegate::CreateLambda([this, InObject, InPropertyHandle]()
+		{
+			IRemoteControlModule& RemoteControlModule = IRemoteControlModule::Get();
+
+			if (RemoteControlModule.HasDefaultValueCustomization(InObject, InPropertyHandle->GetProperty()))
+			{
+				RemoteControlModule.ResetToDefaultValue(InObject, InPropertyHandle->GetProperty());
+			
+				return;
+			}
+
+			InPropertyHandle->ResetToDefault();
+		}
+	);
+	
+	const FResetToDefaultOverride ResetValueOverride = FResetToDefaultOverride::Create(CanResetToDefaultValue, PerformValueReset);
+	
+	ResetDefaultProperty.Reset();
+	
+	ResetDefaultProperty = SNew(SResetToDefaultPropertyEditor, InPropertyHandle)
+		.CustomResetToDefault(ResetValueOverride);
+
+	return ResetDefaultProperty.ToSharedRef();
+}
+
 TSharedRef<SWidget> SRCPanelExposedField::ConstructCallFunctionButton(bool bIsEnabled)
 {
 	return SNew(SVerticalBox)
diff --git a/Engine/Plugins/VirtualProduction/RemoteControl/Source/RemoteControlUI/Private/UI/SRCPanelExposedField.h b/Engine/Plugins/VirtualProduction/RemoteControl/Source/RemoteControlUI/Private/UI/SRCPanelExposedField.h
--- a/Engine/Plugins/VirtualProduction/RemoteControl/Source/RemoteControlUI/Private/UI/SRCPanelExposedField.h
+++ b/Engine/Plugins/VirtualProduction/RemoteControl/Source/RemoteControlUI/Private/UI/SRCPanelExposedField.h
@@ -82,6 +82,8 @@ private:
 	void ConstructPropertyWidget();
 	/** Construct this field widget as a function widget. */
 	void ConstructFunctionWidget();
+	/** Construct the ResetToDefault Widget that can have CustomReset for certain Property (Right now only Light Intensity) */
+	TSharedRef<SWidget> ConstructResetToDefaultWidget(UObject* InObject, TSharedPtr<IPropertyHandle> InPropertyHandle);
 	/**
 	 * Construct a call function button
 	 * @param bIsEnabled Whether the button should be clickable or not.
