Remote Control
- Add setting to control allowed origins
- Add request preprocessor to validate IP of peer that initiated connection

Http Listener
- Pass IP to request callback

diff --git a/Engine/Plugins/Experimental/WebSocketNetworking/Source/WebSocketNetworking/Private/WebSocketServer.cpp b/Engine/Plugins/Experimental/WebSocketNetworking/Source/WebSocketNetworking/Private/WebSocketServer.cpp
--- a/Engine/Plugins/Experimental/WebSocketNetworking/Source/WebSocketNetworking/Private/WebSocketServer.cpp
+++ b/Engine/Plugins/Experimental/WebSocketNetworking/Source/WebSocketNetworking/Private/WebSocketServer.cpp
@@ -164,6 +164,11 @@ bool FWebSocketServer::Init(uint32 Port, FWebSocketClientConnectedCallBack CallB
 	return true;
 }
 
+void FWebSocketServer::SetFilterConnectionCallback(FWebSocketFilterConnectionCallback InFilterConnectionCallback)
+{
+	FilterConnectionCallback = MoveTemp(InFilterConnectionCallback);
+}
+
 void FWebSocketServer::Tick()
 {
 #if USE_LIBWEBSOCKET
@@ -214,6 +219,7 @@ static int unreal_networking_server
 	struct lws_context* Context = lws_get_context(Wsi);
 	PerSessionDataServer* BufferInfo = (PerSessionDataServer*)User;
 	FWebSocketServer* Server = (FWebSocketServer*)lws_context_user(Context);
+	bool bRejectConnection = false;
 
 	switch (Reason)
 	{
@@ -296,6 +302,45 @@ static int unreal_networking_server
 			break;
 		case LWS_CALLBACK_PROTOCOL_DESTROY:
 			break;
+		case LWS_CALLBACK_FILTER_PROTOCOL_CONNECTION:
+			{
+				if (Server && Server->FilterConnectionCallback.IsBound())
+            	{
+					const int32 OriginHeaderLength = lws_hdr_total_length(Wsi, WSI_TOKEN_ORIGIN);
+
+					FString OriginHeader;
+					FString ClientIP;
+
+					if (OriginHeaderLength != 0)
+					{
+						constexpr uint32 MaximumHeaderLength = 1024;
+						if (OriginHeaderLength < MaximumHeaderLength)
+						{
+							char Origin[MaximumHeaderLength + 1];
+							memset(Origin, 0, MaximumHeaderLength + 1);
+
+							lws_hdr_copy(Wsi, Origin, OriginHeaderLength + 1, WSI_TOKEN_ORIGIN);
+							OriginHeader = UTF8_TO_TCHAR(Origin);
+						}
+					}
+
+					char IPAddress[16];
+					int32 IPAddressLength = 0;
+					lws_get_peer_simple(Wsi, IPAddress, IPAddressLength);
+
+					if (IPAddressLength != 0)
+					{
+						ClientIP = UTF8_TO_TCHAR(IPAddress);
+					}
+				
+					if (Server->FilterConnectionCallback.Execute(OriginHeader, ClientIP) == EWebsocketConnectionFilterResult::ConnectionRefused)
+					{
+						bRejectConnection = true;
+					}
+				}
+
+				break;
+			}
 	}
 
 	// Check if http should be enabled or not, if so, use the in-built `lws_callback_http_dummy` which handles basic http requests
@@ -303,6 +348,7 @@ static int unreal_networking_server
 	{
 		return lws_callback_http_dummy(Wsi, Reason, User, In, Len);
 	}
-	return 0;
+
+	return bRejectConnection ? 1 : 0;
 }
 #endif
\ No newline at end of file
diff --git a/Engine/Plugins/Experimental/WebSocketNetworking/Source/WebSocketNetworking/Private/WebSocketServer.h b/Engine/Plugins/Experimental/WebSocketNetworking/Source/WebSocketNetworking/Private/WebSocketServer.h
--- a/Engine/Plugins/Experimental/WebSocketNetworking/Source/WebSocketNetworking/Private/WebSocketServer.h
+++ b/Engine/Plugins/Experimental/WebSocketNetworking/Source/WebSocketNetworking/Private/WebSocketServer.h
@@ -16,6 +16,7 @@ public:
 	virtual ~FWebSocketServer() override;
 	virtual void EnableHTTPServer(TArray<FWebSocketHttpMount> DirectoriesToServe) override;
 	virtual bool Init(uint32 Port, FWebSocketClientConnectedCallBack) override;
+	virtual void SetFilterConnectionCallback(FWebSocketFilterConnectionCallback InFilterConnectionCallback) override;
 	virtual void Tick() override;
 	virtual FString Info() override;
 	//~ End IWebSocketServer interface
@@ -27,6 +28,9 @@ public:
 
 	/** Callback for a new websocket connection to the server */
 	FWebSocketClientConnectedCallBack  ConnectedCallBack;
+	
+	/** Handler called to filter incoming websocket connections. */
+    FWebSocketFilterConnectionCallback FilterConnectionCallback;
 
 	/** Internal libwebsocket context */
 	WebSocketInternalContext* Context;
diff --git a/Engine/Plugins/Experimental/WebSocketNetworking/Source/WebSocketNetworking/Public/IWebSocketServer.h b/Engine/Plugins/Experimental/WebSocketNetworking/Source/WebSocketNetworking/Public/IWebSocketServer.h
--- a/Engine/Plugins/Experimental/WebSocketNetworking/Source/WebSocketNetworking/Public/IWebSocketServer.h
+++ b/Engine/Plugins/Experimental/WebSocketNetworking/Source/WebSocketNetworking/Public/IWebSocketServer.h
@@ -73,6 +73,12 @@ public:
 	 */
 	virtual bool Init(uint32 Port, FWebSocketClientConnectedCallBack ClientConnectedCallback) = 0;
 
+	/**
+	 * Set a filter callback, used to reject connections based on the origin header.
+	 * @param FilterConnectionCallback handler called to filter the connection.
+	 */
+	virtual void SetFilterConnectionCallback(FWebSocketFilterConnectionCallback InFilterConnectionCallback) = 0;
+
 	/** Tick the server. */
 	virtual void Tick() = 0;
 
diff --git a/Engine/Plugins/Experimental/WebSocketNetworking/Source/WebSocketNetworking/Public/WebSocketNetworkingDelegates.h b/Engine/Plugins/Experimental/WebSocketNetworking/Source/WebSocketNetworking/Public/WebSocketNetworkingDelegates.h
--- a/Engine/Plugins/Experimental/WebSocketNetworking/Source/WebSocketNetworking/Public/WebSocketNetworkingDelegates.h
+++ b/Engine/Plugins/Experimental/WebSocketNetworking/Source/WebSocketNetworking/Public/WebSocketNetworkingDelegates.h
@@ -6,7 +6,14 @@
 
 class INetworkingWebSocket;
 
+enum class EWebsocketConnectionFilterResult : uint8
+{
+	ConnectionAccepted,
+	ConnectionRefused
+};
+
 DECLARE_DELEGATE(FWebSocketInfoCallBack);
 DECLARE_DELEGATE_TwoParams(FWebSocketPacketReceivedCallBack, void* /*Data*/, int32 /*Data Size*/);
 DECLARE_DELEGATE_OneParam(FWebSocketClientConnectedCallBack, INetworkingWebSocket* /*Socket*/);
+DECLARE_DELEGATE_RetVal_TwoParams(EWebsocketConnectionFilterResult, FWebSocketFilterConnectionCallback, FString /*Origin*/, FString /*ClientIP*/);
 						  
diff --git a/Engine/Plugins/VirtualProduction/RemoteControl/Source/RemoteControlCommon/Public/RemoteControlSettings.h b/Engine/Plugins/VirtualProduction/RemoteControl/Source/RemoteControlCommon/Public/RemoteControlSettings.h
--- a/Engine/Plugins/VirtualProduction/RemoteControl/Source/RemoteControlCommon/Public/RemoteControlSettings.h
+++ b/Engine/Plugins/VirtualProduction/RemoteControl/Source/RemoteControlCommon/Public/RemoteControlSettings.h
@@ -136,6 +136,31 @@ public:
 	UPROPERTY(Config, EditAnywhere, Category = "Remote Control | Security")
 	bool bEnableRemotePythonExecution = false;
 
+	/** Whether to restrict access to a list of hostname/IPs in the AllowedOrigins setting. */
+	UPROPERTY(Config, EditAnywhere, Category = "Remote Control | Security")
+	bool bRestrictServerAccess = false;
+	
+	/** 
+	 * Origin that can make requests to the remote control server. Should contain the hostname or IP of the server making requests to remote control. ie. "http://yourdomain.com", or "*" to allow all origins. 
+	 * @Note: This is used to block requests coming from a browser (ie. Coming from a script running on a website), ideally you should use both this setting and AllowedIPs, as a request coming from a websocket client can have an empty Origin.
+	 * @Note Supports wildcards (ie. *.yourdomain.com)
+	 */
+	UPROPERTY(Config, EditAnywhere, Category = "Remote Control | Security", meta=(EditCondition = bRestrictServerAccess))
+	FString AllowedOrigin = TEXT("*");
+	
+	/** 
+	 * What IP is allowed to make request to the RemoteControl and RemoteControl Websocket servers.
+	 * @Note If empty or *.*.*.*,  all IP addresses will be able to make requests to your servers if they are on your network, so consider limiting this to a range of IPs that you own.
+	 * @Note Using this setting without AllowedOrigin can potentially open you up to malicous requests coming from a website, as the request would come from localhost.
+	 * @Note Supports wildcards (ie. 202.120.*.*)
+	 */
+	UPROPERTY(Config, EditAnywhere, Category = "Remote Control | Security", meta = (EditCondition = bRestrictServerAccess))
+	FString AllowedIP = TEXT("127.0.0.1");
+	
 	UPROPERTY(config, EditAnywhere, Category = "Remote Control | Security", DisplayName = "Remote Control Passphrase")
 	TArray<FRCPassphrase> Passphrases = {};
+
+private:
+	UPROPERTY(config)
+    bool bSecuritySettingsReviewed = false;
 };
diff --git a/Engine/Plugins/VirtualProduction/RemoteControl/Source/WebRemoteControl/Private/RemoteControlWebSocketServer.cpp b/Engine/Plugins/VirtualProduction/RemoteControl/Source/WebRemoteControl/Private/RemoteControlWebSocketServer.cpp
--- a/Engine/Plugins/VirtualProduction/RemoteControl/Source/WebRemoteControl/Private/RemoteControlWebSocketServer.cpp
+++ b/Engine/Plugins/VirtualProduction/RemoteControl/Source/WebRemoteControl/Private/RemoteControlWebSocketServer.cpp
@@ -6,7 +6,9 @@
 #include "IPAddress.h"
 #include "IRemoteControlModule.h"
 #include "IWebSocketNetworkingModule.h"
+#include "Misc/WildcardString.h"
 #include "RemoteControlRequest.h"
+#include "RemoteControlSettings.h"
 #include "SocketSubsystem.h"
 #include "Sockets.h"
 #include "WebRemoteControlInternalUtils.h"
@@ -129,6 +131,10 @@ bool FRCWebSocketServer::Start(uint32 Port, TSharedPtr<FWebsocketMessageRouter>
 		return false;
 	}
 
+	FWebSocketFilterConnectionCallback FilterCallback;
+	FilterCallback.BindRaw(this, &FRCWebSocketServer::FilterConnection);
+	Server->SetFilterConnectionCallback(MoveTemp(FilterCallback));
+	
 	Router = MoveTemp(InRouter);
 	TickerHandle = FTSTicker::GetCoreTicker().AddTicker(FTickerDelegate::CreateRaw(this, &FRCWebSocketServer::Tick));
 
@@ -250,6 +256,44 @@ void FRCWebSocketServer::OnSocketClose(INetworkingWebSocket* Socket)
 	}
 }
 
+EWebsocketConnectionFilterResult FRCWebSocketServer::FilterConnection(FString OriginHeader, FString ClientIP) const
+{
+	OriginHeader.RemoveSpacesInline();
+	OriginHeader.TrimStartAndEndInline();
+	
+	auto SimplifyAddress = [] (FString Address)
+	{
+		Address.RemoveFromStart(TEXT("https://www."));
+		Address.RemoveFromStart(TEXT("http://www."));
+		Address.RemoveFromStart(TEXT("https://"));
+		Address.RemoveFromStart(TEXT("http://"));
+		Address.RemoveFromEnd(TEXT("/"));
+		return Address;
+	};
+
+	const FString SimplifiedOrigin = SimplifyAddress(OriginHeader);
+	const FWildcardString SimplifiedAllowedOrigin = SimplifyAddress(GetDefault<URemoteControlSettings>()->AllowedOrigin);
+	if (GetDefault<URemoteControlSettings>()->AllowedOrigin != TEXT("*"))
+	{
+		if (!SimplifiedAllowedOrigin.IsMatch(SimplifiedOrigin))
+		{
+			return EWebsocketConnectionFilterResult::ConnectionRefused;
+		}
+	}
+
+	const FWildcardString WildcardAllowedIP = SimplifyAddress(GetDefault<URemoteControlSettings>()->AllowedIP);
+
+	// Allow requests from localhost
+	if (ClientIP != TEXT("localhost") && ClientIP != TEXT("127.0.0.1"))
+	{
+		if (!WildcardAllowedIP.IsEmpty() && WildcardAllowedIP.IsMatch(ClientIP))
+		{
+			return EWebsocketConnectionFilterResult::ConnectionRefused;
+		}
+	}
+
+	return EWebsocketConnectionFilterResult::ConnectionAccepted;
+}
 
 #undef LOCTEXT_NAMESPACE /* FRCWebSocketServer */
 
diff --git a/Engine/Plugins/VirtualProduction/RemoteControl/Source/WebRemoteControl/Private/RemoteControlWebSocketServer.h b/Engine/Plugins/VirtualProduction/RemoteControl/Source/WebRemoteControl/Private/RemoteControlWebSocketServer.h
--- a/Engine/Plugins/VirtualProduction/RemoteControl/Source/WebRemoteControl/Private/RemoteControlWebSocketServer.h
+++ b/Engine/Plugins/VirtualProduction/RemoteControl/Source/WebRemoteControl/Private/RemoteControlWebSocketServer.h
@@ -121,6 +121,9 @@ private:
 	void ReceivedRawPacket(void* Data, int32 Size, FGuid ClientId);
 
 	void OnSocketClose(INetworkingWebSocket* Socket);
+	
+	/** Handle rejecting the websocket connection if it doesn't respect the user's CORS policy. */
+	EWebsocketConnectionFilterResult FilterConnection(FString OriginHeader, FString ClientIP) const;
 
 	/** Returns whether the port is available. */
 	bool IsPortAvailable(uint32 Port) const;
diff --git a/Engine/Plugins/VirtualProduction/RemoteControl/Source/WebRemoteControl/Private/WebRemoteControl.cpp b/Engine/Plugins/VirtualProduction/RemoteControl/Source/WebRemoteControl/Private/WebRemoteControl.cpp
--- a/Engine/Plugins/VirtualProduction/RemoteControl/Source/WebRemoteControl/Private/WebRemoteControl.cpp
+++ b/Engine/Plugins/VirtualProduction/RemoteControl/Source/WebRemoteControl/Private/WebRemoteControl.cpp
@@ -58,6 +58,7 @@
 // Miscelleanous
 #include "Blueprint/BlueprintSupport.h"
 #include "Misc/App.h"
+#include "Misc/WildcardString.h"
 #include "UObject/UnrealType.h"
 #include "Templates/UnrealTemplate.h"
 
@@ -452,45 +453,9 @@ void FWebRemoteControlModule::StartHttpServer()
 			StartRoute(Route);
 		}
 
-		const FHttpRequestHandler ValidationRequestHandler = FHttpRequestHandler([this](const FHttpServerRequest& Request, const FHttpResultCallback& OnComplete)
-		{
-			TUniquePtr<FHttpServerResponse> Response = WebRemoteControlInternalUtils::CreateHttpResponse();
-
-			TArray<FString> ValueArray = {};
-			if (Request.Headers.Find(WebRemoteControlInternalUtils::PassphraseHeader))
-			{
-				ValueArray = Request.Headers[WebRemoteControlInternalUtils::PassphraseHeader];
-			}
-			
-			const FString Passphrase = !ValueArray.IsEmpty() ? ValueArray.Last() : "";
-
-			if (!CheckPassphrase(Passphrase))
-			{
-				WebRemoteControlInternalUtils::CreateUTF8ErrorMessage(FString::Printf(TEXT("Given Passphrase is not correct!")), Response->Body);
-				Response->Code = EHttpServerResponseCodes::Denied;
-				OnComplete(MoveTemp(Response));
-				return true;
-			}
-
-			return false;
-		});
-
-		HttpRouter->RegisterRequestPreprocessor(ValidationRequestHandler);
+		RegisterDefaultPreprocessors();
+		RegisterExternalPreprocesors();
 		
-		// Go through externally registered request pre-processors and register them with the http router.
-		for (const TPair<FDelegateHandle, FHttpRequestHandler>& Handler : PreprocessorsToRegister)
-		{
-			// Find the pre-processors HTTP-handle from the one we generated.
-			FDelegateHandle& Handle = PreprocessorsHandleMappings.FindChecked(Handler.Key);
-			if (Handle.IsValid())
-			{
-				HttpRouter->UnregisterRequestPreprocessor(Handle);
-			}
-
-			// Update the preprocessor handle mapping.
-			Handle = HttpRouter->RegisterRequestPreprocessor(Handler.Value);
-		}
-
 		FHttpServerModule::Get().StartAllListeners();
 
 		bIsHttpServerRunning = true;
@@ -518,6 +483,8 @@ void FWebRemoteControlModule::StopHttpServer()
 		ActiveRouteHandles.Reset();
 	}
 
+	UnregisterAllPreprocessors();
+
 	HttpRouter.Reset();
 	bIsHttpServerRunning = false;
 
@@ -2370,6 +2337,125 @@ void FWebRemoteControlModule::InvokeWrappedRequest(const FRCRequestWrapper& Wrap
 	}
 }
 
+void FWebRemoteControlModule::RegisterDefaultPreprocessors()
+{
+	const FHttpRequestHandler PassphraseRequestHandler = FHttpRequestHandler([this](const FHttpServerRequest& Request, const FHttpResultCallback& OnComplete)
+		{
+			TUniquePtr<FHttpServerResponse> Response = WebRemoteControlInternalUtils::CreateHttpResponse();
+
+			TArray<FString> ValueArray = {};
+			if (Request.Headers.Find(WebRemoteControlInternalUtils::PassphraseHeader))
+			{
+				ValueArray = Request.Headers[WebRemoteControlInternalUtils::PassphraseHeader];
+			}
+			
+			const FString Passphrase = !ValueArray.IsEmpty() ? ValueArray.Last() : "";
+
+			if (!CheckPassphrase(Passphrase))
+			{
+				WebRemoteControlInternalUtils::CreateUTF8ErrorMessage(FString::Printf(TEXT("Given Passphrase is not correct!")), Response->Body);
+				Response->Code = EHttpServerResponseCodes::Denied;
+				OnComplete(MoveTemp(Response));
+				return true;
+			}
+
+			return false;
+		});
+
+	const FHttpRequestHandler IPValidationRequestHandler = FHttpRequestHandler([this](const FHttpServerRequest& Request, const FHttpResultCallback& OnComplete)
+	{
+		TUniquePtr<FHttpServerResponse> Response = WebRemoteControlInternalUtils::CreateHttpResponse();
+
+		FString OriginHeader;
+		if (const TArray<FString>* OriginHeaders = Request.Headers.Find(WebRemoteControlInternalUtils::OriginHeader))
+		{
+			if (OriginHeaders->Num())
+			{
+				OriginHeader = (*OriginHeaders)[0];
+			}
+		}
+
+		OriginHeader.RemoveSpacesInline();
+		OriginHeader.TrimStartAndEndInline();
+		
+		auto SimplifyAddress = [] (FString Address)
+		{
+			Address.RemoveFromStart(TEXT("https://www."));
+			Address.RemoveFromStart(TEXT("http://www."));
+			Address.RemoveFromStart(TEXT("https://"));
+			Address.RemoveFromStart(TEXT("http://"));
+			Address.RemoveFromEnd(TEXT("/"));
+			return Address;
+		};
+
+		const FString SimplifiedOrigin = SimplifyAddress(OriginHeader);
+		const FWildcardString SimplifiedAllowedOrigin = SimplifyAddress(GetDefault<URemoteControlSettings>()->AllowedOrigin);
+		if (GetDefault<URemoteControlSettings>()->AllowedOrigin != TEXT("*"))
+		{
+			if (!SimplifiedAllowedOrigin.IsMatch(SimplifiedOrigin))
+			{
+				WebRemoteControlInternalUtils::CreateUTF8ErrorMessage(FString::Printf(TEXT("Client origin %s does not respect the allowed origin set in Remote Control Settings."), *OriginHeader), Response->Body);
+				Response->Code = EHttpServerResponseCodes::Denied;
+				OnComplete(MoveTemp(Response));
+				return true;
+			}
+		}
+
+		if (Request.PeerAddress)
+		{
+			constexpr bool bAppendPort = false;
+			FString ClientIP = Request.PeerAddress->ToString(bAppendPort);
+			const FWildcardString WildcardAllowedIP = SimplifyAddress(GetDefault<URemoteControlSettings>()->AllowedIP);
+				
+			// Allow requests from localhost
+			if (ClientIP != TEXT("localhost") && ClientIP != TEXT("127.0.0.1"))
+			{
+				if (!WildcardAllowedIP.IsEmpty() && WildcardAllowedIP.IsMatch(ClientIP))
+				{
+					WebRemoteControlInternalUtils::CreateUTF8ErrorMessage(FString::Printf(TEXT("Client IP %s does not respect the allowed IP set in Remote Control Settings."), *ClientIP), Response->Body);
+					Response->Code = EHttpServerResponseCodes::Denied;
+					OnComplete(MoveTemp(Response));
+					return true;
+				}
+			}
+		}
+
+		return false;
+	});
+
+	AllRegisteredPreprocessorHandlers.Add(HttpRouter->RegisterRequestPreprocessor(PassphraseRequestHandler));
+	AllRegisteredPreprocessorHandlers.Add(HttpRouter->RegisterRequestPreprocessor(IPValidationRequestHandler));
+}
+
+void FWebRemoteControlModule::UnregisterAllPreprocessors()
+{
+	if (HttpRouter)
+	{
+		for (const FDelegateHandle& Handle : AllRegisteredPreprocessorHandlers)
+		{
+			HttpRouter->UnregisterRequestPreprocessor(Handle);
+		}
+	}
+}
+
+void FWebRemoteControlModule::RegisterExternalPreprocesors()
+{
+	for (const TPair<FDelegateHandle, FHttpRequestHandler>& Handler : PreprocessorsToRegister)
+	{
+		// Find the pre-processors HTTP-handle from the one we generated.
+		FDelegateHandle& Handle = PreprocessorsHandleMappings.FindChecked(Handler.Key);
+		if (Handle.IsValid())
+		{
+			HttpRouter->UnregisterRequestPreprocessor(Handle);
+			AllRegisteredPreprocessorHandlers.RemoveAtSwap(AllRegisteredPreprocessorHandlers.IndexOfByKey(Handle));
+		}
+
+		// Update the preprocessor handle mapping.
+		Handle = HttpRouter->RegisterRequestPreprocessor(Handler.Value);
+		AllRegisteredPreprocessorHandlers.Add(Handle);
+	}
+}
+
 #if WITH_EDITOR
 void FWebRemoteControlModule::RegisterSettings()
 {
diff --git a/Engine/Plugins/VirtualProduction/RemoteControl/Source/WebRemoteControl/Private/WebRemoteControl.h b/Engine/Plugins/VirtualProduction/RemoteControl/Source/WebRemoteControl/Private/WebRemoteControl.h
--- a/Engine/Plugins/VirtualProduction/RemoteControl/Source/WebRemoteControl/Private/WebRemoteControl.h
+++ b/Engine/Plugins/VirtualProduction/RemoteControl/Source/WebRemoteControl/Private/WebRemoteControl.h
@@ -68,7 +68,7 @@ public:
 	 * @param Route The route to unregister.
 	 */
 	void UnregisterRoute(const FRemoteControlRoute& Route);
-
+	
 	/**
 	 * Start the web control server
 	 */
@@ -141,6 +141,13 @@ private:
 
 	void InvokeWrappedRequest(const struct FRCRequestWrapper& Wrapper, FMemoryWriter& OutUTF8PayloadWriter, const FHttpServerRequest* TemplateRequest = nullptr);
 
+	/** Register the default request preprocessors used by Web RC in the HttpRouter. */
+	void RegisterDefaultPreprocessors();
+	void UnregisterAllPreprocessors();
+
+	/** Register the external request preprocessors that were added using RegisterRequestPreprocessor in the HttpRouter.   */
+	void RegisterExternalPreprocesors();
+
 #if WITH_EDITOR
 	//~ Settings
 	void RegisterSettings();
@@ -198,6 +205,11 @@ private:
 	 */
 	TMap<FDelegateHandle, FDelegateHandle> PreprocessorsHandleMappings;
 
+	/**
+	 * List of all registered preprocessor handlers. Used to unregister them on shutdown.
+	 */
+	TArray<FDelegateHandle> AllRegisteredPreprocessorHandlers;
+
 	/** An external web socket logger*/
 	TUniquePtr<class FWebRemoteControlExternalLogger> ExternalLogger;
 
diff --git a/Engine/Plugins/VirtualProduction/RemoteControl/Source/WebRemoteControl/Private/WebRemoteControlInternalUtils.cpp b/Engine/Plugins/VirtualProduction/RemoteControl/Source/WebRemoteControl/Private/WebRemoteControlInternalUtils.cpp
--- a/Engine/Plugins/VirtualProduction/RemoteControl/Source/WebRemoteControl/Private/WebRemoteControlInternalUtils.cpp
+++ b/Engine/Plugins/VirtualProduction/RemoteControl/Source/WebRemoteControl/Private/WebRemoteControlInternalUtils.cpp
@@ -559,7 +559,15 @@ bool WebRemoteControlInternalUtils::IsWrappedRequest(const FHttpServerRequest& R
 void WebRemoteControlInternalUtils::AddCORSHeaders(FHttpServerResponse* InOutResponse)
 {
 	check(InOutResponse != nullptr);
-	InOutResponse->Headers.Add(TEXT("Access-Control-Allow-Origin"), { TEXT("*") });
+	if (GetDefault<URemoteControlSettings>()->bRestrictServerAccess)
+	{
+		InOutResponse->Headers.Add(TEXT("Access-Control-Allow-Origin"), { GetDefault<URemoteControlSettings>()->AllowedOrigin });
+	}
+	else
+	{
+		InOutResponse->Headers.Add(TEXT("Access-Control-Allow-Origin"), { TEXT("*") });
+	}
+	
 	InOutResponse->Headers.Add(TEXT("Access-Control-Allow-Methods"), { TEXT("PUT, POST, GET, OPTIONS") });
 	InOutResponse->Headers.Add(TEXT("Access-Control-Allow-Headers"), { TEXT("Origin, X-Requested-With, Content-Type, Accept") });
 	InOutResponse->Headers.Add(TEXT("Access-Control-Max-Age"), { TEXT("600") });
diff --git a/Engine/Plugins/VirtualProduction/RemoteControl/Source/WebRemoteControl/Private/WebRemoteControlInternalUtils.h b/Engine/Plugins/VirtualProduction/RemoteControl/Source/WebRemoteControl/Private/WebRemoteControlInternalUtils.h
--- a/Engine/Plugins/VirtualProduction/RemoteControl/Source/WebRemoteControl/Private/WebRemoteControlInternalUtils.h
+++ b/Engine/Plugins/VirtualProduction/RemoteControl/Source/WebRemoteControl/Private/WebRemoteControlInternalUtils.h
@@ -50,6 +50,7 @@ namespace WebRemoteControlInternalUtils
 {
 	static const TCHAR* WrappedRequestHeader = TEXT("UE-Wrapped-Request");
 	static const FString PassphraseHeader = TEXT("Passphrase");
+	static const TCHAR* OriginHeader = TEXT("Origin");
 
 	/**
 	 * Construct a default http response with CORS headers.
diff --git a/Engine/Plugins/VirtualProduction/RemoteControl/Source/WebRemoteControl/Public/IWebRemoteControlModule.h b/Engine/Plugins/VirtualProduction/RemoteControl/Source/WebRemoteControl/Public/IWebRemoteControlModule.h
--- a/Engine/Plugins/VirtualProduction/RemoteControl/Source/WebRemoteControl/Public/IWebRemoteControlModule.h
+++ b/Engine/Plugins/VirtualProduction/RemoteControl/Source/WebRemoteControl/Public/IWebRemoteControlModule.h
@@ -24,6 +24,8 @@ public:
 	 * Useful for cases where you want to drop or handle incoming requests before they are handled the the web remote control module.
 	 * @param RequestPreprocessor The function called to process the incoming request.
 	 * @return FDelegateHandle The handle to the delegate, used for unregistering preprocessors. 
+	 * @note The request preprocessor should return false if the request should pass through, and true if the request has been handled by the preprocessor.
+	 *		 If the request is handled by the preprocessor, you must call the OnComplete callback.
 	 */
 	virtual FDelegateHandle RegisterRequestPreprocessor(FHttpRequestHandler RequestPreprocessor) = 0;
 
diff --git a/Engine/Source/Runtime/Online/HTTPServer/Private/HttpConnectionRequestReadContext.cpp b/Engine/Source/Runtime/Online/HTTPServer/Private/HttpConnectionRequestReadContext.cpp
--- a/Engine/Source/Runtime/Online/HTTPServer/Private/HttpConnectionRequestReadContext.cpp
+++ b/Engine/Source/Runtime/Online/HTTPServer/Private/HttpConnectionRequestReadContext.cpp
@@ -5,7 +5,9 @@
 #include "HttpServerConstantsPrivate.h"
 #include "HttpServerRequest.h"
 #include "HttpConnectionContext.h"
+#include "IPAddress.h"
 #include "Sockets.h"
+#include "SocketSubsystem.h"
 
 FHttpConnectionRequestReadContext::FHttpConnectionRequestReadContext(FSocket* InSocket)
 	: Socket(InSocket)
@@ -197,6 +199,11 @@ bool FHttpConnectionRequestReadContext::ParseContentLength(const FHttpServerRequ
 
 TSharedPtr<FHttpServerRequest> FHttpConnectionRequestReadContext::BuildRequest(const FString& RequestHeader)
 {
+	if (IsEngineExitRequested())
+	{
+		return nullptr;
+	}
+	
 	TArray<FString> ParsedHeader;
 	RequestHeader.ParseIntoArrayLines(ParsedHeader);
 
@@ -218,6 +225,18 @@ TSharedPtr<FHttpServerRequest> FHttpConnectionRequestReadContext::BuildRequest(c
 
 	Request = MakeShared<FHttpServerRequest>();
 
+	if (Socket)
+	{
+		if (ISocketSubsystem* SocketSubsystem = ISocketSubsystem::Get(PLATFORM_SOCKETSUBSYSTEM))
+		{
+			TSharedRef<FInternetAddr> RemoteAddress = SocketSubsystem->CreateInternetAddr();
+			if (Socket->GetPeerAddress(*RemoteAddress))
+			{
+				Request->PeerAddress = MoveTemp(RemoteAddress);
+			}
+		}
+	}
+
 	auto RequestVerb = HttpMethodTokens[0];
 	if (0 == RequestVerb.Compare(TEXT("GET"),
 		ESearchCase::IgnoreCase))
diff --git a/Engine/Source/Runtime/Online/HTTPServer/Public/HttpServerRequest.h b/Engine/Source/Runtime/Online/HTTPServer/Public/HttpServerRequest.h
--- a/Engine/Source/Runtime/Online/HTTPServer/Public/HttpServerRequest.h
+++ b/Engine/Source/Runtime/Online/HTTPServer/Public/HttpServerRequest.h
@@ -5,6 +5,8 @@
 #include "HttpPath.h"
 #include "HttpServerHttpVersion.h"
 
+class FInternetAddr;
+
 enum class EHttpServerRequestVerbs : uint16
 {
 	VERB_NONE   = 0,
@@ -25,14 +27,17 @@ public:
 	/** Constructor */
 	FHttpServerRequest() { };
 
+	/** The IP address of the peer that initiated the request. */
+	TSharedPtr<FInternetAddr> PeerAddress;
+
 	/** The handler-route-relative HTTP path */
 	FHttpPath RelativePath;
 
 	/** The HTTP-compliant verb  */
-	EHttpServerRequestVerbs Verb;
+	EHttpServerRequestVerbs Verb = EHttpServerRequestVerbs::VERB_NONE;
 
 	/** The request HTTP protocol version */
-	HttpVersion::EHttpServerHttpVersion HttpVersion;
+	HttpVersion::EHttpServerHttpVersion HttpVersion = HttpVersion::EHttpServerHttpVersion::HTTP_VERSION_1_1;
 
 	/** The HTTP headers */
 	TMap<FString, TArray<FString>> Headers;
