Undo/Redo in remote control doesn't fire events

diff --git a/Engine/Plugins/VirtualProduction/RemoteControl/Source/RemoteControl/Private/RemoteControlPreset.cpp b/Engine/Plugins/VirtualProduction/RemoteControl/Source/RemoteControl/Private/RemoteControlPreset.cpp
--- a/Engine/Plugins/VirtualProduction/RemoteControl/Source/RemoteControl/Private/RemoteControlPreset.cpp
+++ b/Engine/Plugins/VirtualProduction/RemoteControl/Source/RemoteControl/Private/RemoteControlPreset.cpp
@@ -26,6 +26,7 @@
 #include "RemoteControlEntityFactory.h"
 #include "RemoteControlObjectVersion.h"
 #include "RemoteControlPresetRebindingManager.h"
+#include "RemoteControlTransactionListenerHelper.h"
 
 #include "UObject/Object.h"
 #include "UObject/ObjectMacros.h"
@@ -190,7 +191,15 @@ FRemoteControlPresetGroup& FRemoteControlPresetLayout::CreateGroup(FName GroupNa
 	FRemoteControlPresetGroup& Group = Groups.Emplace_GetRef(GroupName, GroupId);
 	Owner->CacheLayoutData();
 	OnGroupAddedDelegate.Broadcast(Group);
+
+	FRCTransactionListenerHelper<FRemoteControlPresetGroup>(ERCTransaction::Undo, OnGroupDeleted(), Group);
+	FRCTransactionListenerHelper<const FRemoteControlPresetGroup&>(ERCTransaction::Redo, OnGroupAdded(), Group);
+
 	Owner->OnPresetLayoutModified().Broadcast(Owner.Get());
+
+	FRCTransactionListenerHelper<URemoteControlPreset*>(ERCTransaction::Undo, Owner->OnPresetLayoutModified(), Owner.Get());
+	FRCTransactionListenerHelper<URemoteControlPreset*>(ERCTransaction::Redo, Owner->OnPresetLayoutModified(), Owner.Get());
+
 	return Group;
 }
 
@@ -249,6 +258,11 @@ void FRemoteControlPresetLayout::SwapGroups(FGuid OriginGroupId, FGuid TargetGro
 			OriginGroupIndex += 1;
 		}
 
+		TArray<FGuid> GroupIdsBeforeTransaction;
+		GroupIdsBeforeTransaction.Reserve(Groups.Num());
+		Algo::Transform(Groups, GroupIdsBeforeTransaction, &FRemoteControlPresetGroup::Id);
+		FRCTransactionListenerHelper<const TArray<FGuid>&>(ERCTransaction::Undo, OnGroupOrderChanged(), GroupIdsBeforeTransaction);
+
 		Groups.Insert(FRemoteControlPresetGroup{ *OriginGroup }, TargetGroupIndex);
 		Groups.Swap(TargetGroupIndex, OriginGroupIndex);
 		Groups.RemoveAt(OriginGroupIndex);
@@ -257,7 +271,13 @@ void FRemoteControlPresetLayout::SwapGroups(FGuid OriginGroupId, FGuid TargetGro
 		GroupIds.Reserve(Groups.Num());
 		Algo::Transform(Groups, GroupIds, &FRemoteControlPresetGroup::Id);
 		OnGroupOrderChangedDelegate.Broadcast(GroupIds);
+
+		FRCTransactionListenerHelper<const TArray<FGuid>&>(ERCTransaction::Redo, OnGroupOrderChanged(), GroupIds);
+
 		Owner->OnPresetLayoutModified().Broadcast(Owner.Get());
+
+		FRCTransactionListenerHelper<URemoteControlPreset*>(ERCTransaction::Undo, Owner->OnPresetLayoutModified(), Owner.Get());
+		FRCTransactionListenerHelper<URemoteControlPreset*>(ERCTransaction::Redo, Owner->OnPresetLayoutModified(), Owner.Get());
 	}
 }
 
@@ -292,6 +312,9 @@ void FRemoteControlPresetLayout::SwapFields(const FFieldSwapArgs& FieldSwapArgs)
 
 		// Here we don't want to trigger add/delete delegates since the fields just get moved around.
 		TArray<FGuid>& Fields = DragTargetGroup->AccessFields();
+
+		FRCTransactionListenerHelper<const FGuid&, const TArray<FGuid>&>(ERCTransaction::Undo, OnFieldOrderChanged(), FieldSwapArgs.TargetGroupId, Fields);
+
 		Fields.Insert(FieldSwapArgs.DraggedFieldId, DragTargetFieldIndex);
 		Fields.Swap(DragTargetFieldIndex, DragOriginFieldIndex);
 		Fields.RemoveAt(DragOriginFieldIndex);
@@ -299,6 +322,11 @@ void FRemoteControlPresetLayout::SwapFields(const FFieldSwapArgs& FieldSwapArgs)
 		Owner->CacheLayoutData();
 		OnFieldOrderChangedDelegate.Broadcast(FieldSwapArgs.TargetGroupId, Fields);
 		Owner->OnPresetLayoutModified().Broadcast(Owner.Get());
+
+		FRCTransactionListenerHelper<const FGuid&, const TArray<FGuid>&>(ERCTransaction::Redo, OnFieldOrderChanged(), FieldSwapArgs.TargetGroupId, Fields);
+
+		FRCTransactionListenerHelper<URemoteControlPreset*>(ERCTransaction::Undo, Owner->OnPresetLayoutModified(), Owner.Get());
+		FRCTransactionListenerHelper<URemoteControlPreset*>(ERCTransaction::Redo, Owner->OnPresetLayoutModified(), Owner.Get());
 	}
 	else
 	{
@@ -307,10 +335,19 @@ void FRemoteControlPresetLayout::SwapFields(const FFieldSwapArgs& FieldSwapArgs)
 		DragTargetFieldIndex = DragTargetFieldIndex == INDEX_NONE ? 0 : DragTargetFieldIndex;
 		DragTargetGroup->AccessFields().Insert(FieldSwapArgs.DraggedFieldId, DragTargetFieldIndex);
 
+		FRCTransactionListenerHelper<const FGuid&, const FGuid&, int32>(ERCTransaction::Undo, OnFieldDeleted(), FieldSwapArgs.TargetGroupId, FieldSwapArgs.DraggedFieldId, DragTargetFieldIndex);
+		FRCTransactionListenerHelper<const FGuid&, const FGuid&, int32>(ERCTransaction::Undo, OnFieldAdded(), FieldSwapArgs.OriginGroupId, FieldSwapArgs.DraggedFieldId, DragOriginFieldIndex);
+
+		FRCTransactionListenerHelper<const FGuid&, const FGuid&, int32>(ERCTransaction::Redo, OnFieldDeleted(), FieldSwapArgs.OriginGroupId, FieldSwapArgs.DraggedFieldId, DragOriginFieldIndex);
+		FRCTransactionListenerHelper<const FGuid&, const FGuid&, int32>(ERCTransaction::Redo, OnFieldAdded(), FieldSwapArgs.TargetGroupId, FieldSwapArgs.DraggedFieldId, DragTargetFieldIndex);
+
 		Owner->CacheLayoutData();
 		OnFieldDeletedDelegate.Broadcast(FieldSwapArgs.OriginGroupId, FieldSwapArgs.DraggedFieldId, DragOriginFieldIndex);
 		OnFieldAddedDelegate.Broadcast(FieldSwapArgs.TargetGroupId, FieldSwapArgs.DraggedFieldId, DragTargetFieldIndex);
 		Owner->OnPresetLayoutModified().Broadcast(Owner.Get());
+
+		FRCTransactionListenerHelper<URemoteControlPreset*>(ERCTransaction::Undo, Owner->OnPresetLayoutModified(), Owner.Get());
+		FRCTransactionListenerHelper<URemoteControlPreset*>(ERCTransaction::Redo, Owner->OnPresetLayoutModified(), Owner.Get());
 	}
 }
 
@@ -330,8 +367,14 @@ void FRemoteControlPresetLayout::DeleteGroup(FGuid GroupId)
 		
 		Owner->CacheLayoutData();
 
+		FRCTransactionListenerHelper<const FRemoteControlPresetGroup&>(ERCTransaction::Undo, OnGroupAdded(), CopyTemp(DeletedGroup));
+		FRCTransactionListenerHelper<FRemoteControlPresetGroup>(ERCTransaction::Redo, OnGroupDeleted(), CopyTemp(DeletedGroup));
+
 		OnGroupDeletedDelegate.Broadcast(MoveTemp(DeletedGroup));
 		Owner->OnPresetLayoutModified().Broadcast(Owner.Get());
+
+		FRCTransactionListenerHelper<URemoteControlPreset*>(ERCTransaction::Undo, Owner->OnPresetLayoutModified(), Owner.Get());
+		FRCTransactionListenerHelper<URemoteControlPreset*>(ERCTransaction::Redo, Owner->OnPresetLayoutModified(), Owner.Get());
 	}
 }
 
@@ -339,9 +382,15 @@ void FRemoteControlPresetLayout::RenameGroup(FGuid GroupId, FName NewGroupName)
 {
 	if (FRemoteControlPresetGroup* Group = GetGroup(GroupId))
 	{
+		FRCTransactionListenerHelper<const FGuid&, FName>(ERCTransaction::Undo, OnGroupRenamed(), GroupId, Group->Name);
+		FRCTransactionListenerHelper<const FGuid&, FName>(ERCTransaction::Redo, OnGroupRenamed(), GroupId, NewGroupName);
+
 		Group->Name = NewGroupName;
 		OnGroupRenamedDelegate.Broadcast(GroupId, NewGroupName);
 		Owner->OnPresetLayoutModified().Broadcast(Owner.Get());
+
+		FRCTransactionListenerHelper<URemoteControlPreset*>(ERCTransaction::Undo, Owner->OnPresetLayoutModified(), Owner.Get());
+		FRCTransactionListenerHelper<URemoteControlPreset*>(ERCTransaction::Redo, Owner->OnPresetLayoutModified(), Owner.Get());
 	}
 }
 
@@ -361,7 +410,14 @@ void FRemoteControlPresetLayout::AddField(FGuid GroupId, FGuid FieldId)
 	{
 		Group->AccessFields().Add(FieldId);
 		OnFieldAddedDelegate.Broadcast(GroupId, FieldId, Group->AccessFields().Num() - 1);
+
+		FRCTransactionListenerHelper<const FGuid&, const FGuid&, int32>(ERCTransaction::Undo, OnFieldDeleted(), GroupId, FieldId, Group->AccessFields().Num() - 1);
+		FRCTransactionListenerHelper<const FGuid&, const FGuid&, int32>(ERCTransaction::Redo, OnFieldAdded(), GroupId, FieldId, Group->AccessFields().Num() - 1);
+
 		Owner->OnPresetLayoutModified().Broadcast(Owner.Get());
+
+		FRCTransactionListenerHelper<URemoteControlPreset*>(ERCTransaction::Undo, Owner->OnPresetLayoutModified(), Owner.Get());
+		FRCTransactionListenerHelper<URemoteControlPreset*>(ERCTransaction::Redo, Owner->OnPresetLayoutModified(), Owner.Get());
 	}
 }
 
@@ -371,7 +427,14 @@ void FRemoteControlPresetLayout::InsertFieldAt(FGuid GroupId, FGuid FieldId, int
 	{
 		Group->AccessFields().Insert(FieldId, Index);
 		OnFieldAddedDelegate.Broadcast(GroupId, FieldId, Index);
+
+		FRCTransactionListenerHelper<const FGuid&, const FGuid&, int32>(ERCTransaction::Undo, OnFieldDeleted(), GroupId, FieldId, Index);
+		FRCTransactionListenerHelper<const FGuid&, const FGuid&, int32>(ERCTransaction::Redo, OnFieldAdded(), GroupId, FieldId, Index);
+
 		Owner->OnPresetLayoutModified().Broadcast(Owner.Get());
+
+		FRCTransactionListenerHelper<URemoteControlPreset*>(ERCTransaction::Undo, Owner->OnPresetLayoutModified(), Owner.Get());
+		FRCTransactionListenerHelper<URemoteControlPreset*>(ERCTransaction::Redo, Owner->OnPresetLayoutModified(), Owner.Get());
 	}
 }
 
@@ -399,7 +462,14 @@ void FRemoteControlPresetLayout::RemoveFieldAt(FGuid GroupId, int32 Index)
 		Fields.RemoveAt(Index);
 
 		OnFieldDeletedDelegate.Broadcast(GroupId, FieldId, Index);
+
+		FRCTransactionListenerHelper<const FGuid&, const FGuid&, int32>(ERCTransaction::Undo, OnFieldAdded(), GroupId, FieldId, Index);
+		FRCTransactionListenerHelper<const FGuid&, const FGuid&, int32>(ERCTransaction::Redo, OnFieldDeleted(), GroupId, FieldId, Index);
+
 		Owner->OnPresetLayoutModified().Broadcast(Owner.Get());
+
+		FRCTransactionListenerHelper<URemoteControlPreset*>(ERCTransaction::Undo, Owner->OnPresetLayoutModified(), Owner.Get());
+		FRCTransactionListenerHelper<URemoteControlPreset*>(ERCTransaction::Redo, Owner->OnPresetLayoutModified(), Owner.Get());
 	}
 }
 
@@ -571,6 +641,10 @@ URCVirtualPropertyInContainer* URemoteControlPreset::AddController(TSubclassOf<U
 	if (NewController)
 	{
 		OnControllerAdded().Broadcast(this, NewPropertyName, NewController->Id);
+
+		FRCTransactionListenerHelper<URemoteControlPreset*, const FGuid&>(ERCTransaction::Undo, OnControllerRemoved(), this, NewController->Id);
+		FRCTransactionListenerHelper<URemoteControlPreset*, const FName, const FGuid&>(ERCTransaction::Redo, OnControllerAdded(), this, NewPropertyName, NewController->Id);
+
 		InitializeEntityMetadata(NewController);
 	}
 	
@@ -589,6 +663,10 @@ bool URemoteControlPreset::RemoveController(const FName& InPropertyName)
 		{
 			ControllerContainer->RemoveProperty(ControllerToDelete->GetPropertyName());
 			OnControllerRemoved().Broadcast(this, ControllerToDelete->Id);
+
+			FRCTransactionListenerHelper<URemoteControlPreset*, const FName, const FGuid&>(ERCTransaction::Undo, OnControllerAdded(), this, ControllerToDelete->DisplayName, ControllerToDelete->Id);
+			FRCTransactionListenerHelper<URemoteControlPreset*, const FGuid&>(ERCTransaction::Redo, OnControllerRemoved(), this, ControllerToDelete->Id);
+
 			return true;
 		}
 	}
@@ -619,6 +697,9 @@ void URemoteControlPreset::ResetControllers()
 		for (TObjectPtr<URCVirtualPropertyBase> Controller : ControllerContainer->VirtualProperties)
 		{
 			OnControllerRemoved().Broadcast(this, Controller->Id);
+
+			FRCTransactionListenerHelper<URemoteControlPreset*, const FName, const FGuid&>(ERCTransaction::Undo, OnControllerAdded(), this, Controller->DisplayName, Controller->Id);
+			FRCTransactionListenerHelper<URemoteControlPreset*, const FGuid&>(ERCTransaction::Redo, OnControllerRemoved(), this, Controller->Id);
 		}
 		ControllerContainer->Reset();
 	}
@@ -674,6 +755,9 @@ void URemoteControlPreset::OnModifyController(const FPropertyChangedEvent& Prope
 		{
 			ControllerContainer->OnModifyPropertyValue(PropertyChangedEvent);
 			OnControllerModified().Broadcast(this, {ModifiedController->Id});
+
+			FRCTransactionListenerHelper<URemoteControlPreset*, const TSet<FGuid>&>(ERCTransaction::Undo, OnControllerModified(), this, {ModifiedController->Id});
+			FRCTransactionListenerHelper<URemoteControlPreset*, const TSet<FGuid>&>(ERCTransaction::Redo, OnControllerModified(), this, {ModifiedController->Id});
 		}
 	}	
 } 
@@ -815,6 +899,9 @@ TSharedPtr<FRemoteControlEntity> URemoteControlPreset::Expose(FRemoteControlEnti
 	
 	OnEntityExposed().Broadcast(this, RCEntity->GetId());
 
+	FRCTransactionListenerHelper<URemoteControlPreset*, const FGuid&>(ERCTransaction::Undo, OnEntityUnexposed(), this, RCEntity->GetId());
+	FRCTransactionListenerHelper<URemoteControlPreset*, const FGuid&>(ERCTransaction::Redo, OnEntityExposed(), this, RCEntity->GetId());
+
 	return RCEntity;
 }
 
@@ -1362,6 +1449,9 @@ void URemoteControlPreset::Unexpose(const FGuid& EntityId)
 {
 	if (EntityId.IsValid() && Registry->GetExposedEntity(EntityId).IsValid())
 	{
+		FRCTransactionListenerHelper<URemoteControlPreset*, const FGuid&>(ERCTransaction::Undo, OnEntityExposed(), this, EntityId);
+		FRCTransactionListenerHelper<URemoteControlPreset*, const FGuid&>(ERCTransaction::Redo, OnEntityUnexposed(), this, EntityId);
+
 		OnEntityUnexposedDelegate.Broadcast(this, EntityId);
 
 		Registry->Modify();
@@ -1528,11 +1618,17 @@ void URemoteControlPreset::OnExpose(const FExposeInfo& Info)
 		NameToGuidMap.Add(Field->GetLabel(), Info.FieldId);
 		Layout.AddField(FieldGroupId, Info.FieldId);
 		OnEntityExposed().Broadcast(this, Info.FieldId);
+
+		FRCTransactionListenerHelper<URemoteControlPreset*, const FGuid&>(ERCTransaction::Undo, OnEntityUnexposed(), this, Info.FieldId);
+		FRCTransactionListenerHelper<URemoteControlPreset*, const FGuid&>(ERCTransaction::Redo, OnEntityExposed(), this, Info.FieldId);
 	}
 }
 
 void URemoteControlPreset::OnUnexpose(FGuid UnexposedFieldId)
  {
+	FRCTransactionListenerHelper<URemoteControlPreset*, const FGuid&>(ERCTransaction::Undo, OnEntityExposed(), this, UnexposedFieldId);
+	FRCTransactionListenerHelper<URemoteControlPreset*, const FGuid&>(ERCTransaction::Redo, OnEntityUnexposed(), this, UnexposedFieldId);
+
 	OnEntityUnexposed().Broadcast(this, UnexposedFieldId);
 	
 	FRCCachedFieldData CachedData = FieldCache.FindChecked(UnexposedFieldId);
diff --git a/Engine/Plugins/VirtualProduction/RemoteControl/Source/RemoteControl/Private/RemoteControlTransactionListenerHelper.h b/Engine/Plugins/VirtualProduction/RemoteControl/Source/RemoteControl/Private/RemoteControlTransactionListenerHelper.h
new file mode 100644
--- /dev/null
+++ b/Engine/Plugins/VirtualProduction/RemoteControl/Source/RemoteControl/Private/RemoteControlTransactionListenerHelper.h
@@ -0,0 +1,30 @@
+﻿// Copyright Epic Games, Inc. All Rights Reserved.
+#pragma once
+
+#include "CoreMinimal.h"
+#include "RemoteControlTransactionListenerManager.h"
+
+/**
+ * Helper class which create a Listener and Register to the Manager
+ * @tparam CallbackTypeArgs Callback Args type
+ */
+template <typename  ...CallbackTypeArgs>
+class FRCTransactionListenerHelper
+{
+	FRCTransactionListenerHelper() = delete;
+	
+public:
+	
+	/**
+	 * Create and Register a Listener to the ListenerManager
+	 * @param InType When the broadcast should be called (Undo or Redo), see ERCTransactionType
+	 * @param InDelegateToBroadcast Multicast delegate to broadcast
+	 * @param InOrderedArgs Ordered arguments that should be passed to the MulticastDelegate
+	 */
+	FRCTransactionListenerHelper(ERCTransaction::Type InType, TMulticastDelegate<void(CallbackTypeArgs...)> InDelegateToBroadcast, CallbackTypeArgs... InOrderedArgs)
+	{
+#if WITH_EDITOR
+		FRCTransactionListenerManager<CallbackTypeArgs...>::CreateListenerAndRegister(InType, InDelegateToBroadcast, Forward<CallbackTypeArgs>(InOrderedArgs)...);
+#endif // WITH_EDITOR
+	}
+};
\ No newline at end of file
diff --git a/Engine/Plugins/VirtualProduction/RemoteControl/Source/RemoteControl/Private/RemoteControlTransactionListenerManager.h b/Engine/Plugins/VirtualProduction/RemoteControl/Source/RemoteControl/Private/RemoteControlTransactionListenerManager.h
new file mode 100644
--- /dev/null
+++ b/Engine/Plugins/VirtualProduction/RemoteControl/Source/RemoteControl/Private/RemoteControlTransactionListenerManager.h
@@ -0,0 +1,269 @@
+﻿// Copyright Epic Games, Inc. All Rights Reserved.
+#pragma once
+
+#include "CoreMinimal.h"
+#if WITH_EDITOR
+#include "Editor.h"
+#include "Editor/TransBuffer.h"
+#endif // WITH_EDITOR
+#include "Misc/ITransaction.h"
+
+/**
+ * Enum that helps to distinguish between Undo and Redo
+ */
+namespace ERCTransaction
+{
+	enum Type
+	{
+		None,
+		Undo,
+		Redo
+	};
+}
+
+/**
+ * Collection of FRCTransactionListener which holds them and executes their Callback on the respective UTransBuffer Callback.
+ * @tparam CallbackArgs Callback Args type
+ */
+template <typename... CallbackArgs>
+class FRCTransactionListenerManager
+{
+	/**
+ 	* Helper class that stores a callback and its Args to then call it later on during Transaction Undo/Redo.
+ 	* Assign the Callback and the Args (if needed) before calling RegisterToCollection.
+ 	* @tparam ListenerCallbackArgs Args types to pass to the callback
+ 	*/
+	template <typename... ListenerCallbackArgs>
+	class FRCTransactionListener
+	{	
+		FRCTransactionListener() = delete;
+
+	public:
+
+		/**
+		* Actual listener which holds the broadcast to make at Undo or Redo
+	 	* @param InType When the broadcast should be called (Undo or Redo)
+	 	* @param InDelegateToBroadcast Multicast delegate to broadcast
+	 	* @param InArgs Ordered arguments that should be passed to the MulticastDelegate
+		*/
+		FRCTransactionListener(ERCTransaction::Type InType, TMulticastDelegate<void(ListenerCallbackArgs...)> InDelegateToBroadcast, ListenerCallbackArgs&&... InArgs) : Type(InType)
+		{
+			OnUndoRedoDelegate.AddLambda([InDelegateToBroadcast, InArgs...]()
+			{
+				InDelegateToBroadcast.Broadcast(InArgs...);	
+			});
+		}
+
+		/**
+		 * Getter for the holder of the callback called during Undo or Redo
+		 * @return The UndoRedo Multicast Delegate
+		 */
+		const TMulticastDelegate<void()>& OnUndoRedo()
+		{
+			return OnUndoRedoDelegate;
+		}
+
+	public:
+		/**
+ 		 * Type of the transaction
+ 		 */
+		ERCTransaction::Type Type;
+
+	private:
+		/**
+ 		 * Holds the callback to be executed during Undo or Redo
+ 		 */
+		TMulticastDelegate<void()> OnUndoRedoDelegate;
+	};
+
+public:
+
+	/**
+ 	 * Create a listener with the parameter passed and Register it into the Manager to be called during Undo/Redo callbacks
+ 	 */
+	static void CreateListenerAndRegister(ERCTransaction::Type InType, TMulticastDelegate<void(CallbackArgs...)> InDelegateToBroadcast, CallbackArgs&&... InOrderedArgs)
+	{
+		FRCTransactionListener<CallbackArgs...> Listener = FRCTransactionListener<CallbackArgs...>(InType, InDelegateToBroadcast, Forward<CallbackArgs>(InOrderedArgs)...);
+		Get().Register(Listener);
+	}
+	
+private:
+	
+	/**
+	 * Return the only Instance for CallbackType and CallbackArgs of this class.
+	 * @return Singleton of the class.
+	 */
+	static FRCTransactionListenerManager<CallbackArgs...>& Get()
+	{
+		if (!Instance.IsValid())
+		{
+			Instance = TUniquePtr<FRCTransactionListenerManager<CallbackArgs...>>(new FRCTransactionListenerManager<CallbackArgs...>());
+		}
+		return *(Instance.Get());
+	}
+
+public:
+
+	/**
+ 	 * UnSubscribe to TransactionBuffer callback
+ 	 */
+	~FRCTransactionListenerManager()
+	{
+#if WITH_EDITOR
+		if (GEditor)
+		{
+			UTransBuffer* TransBuffer = CastChecked<UTransBuffer>(GEditor->Trans);
+
+			TransBuffer->OnTransactionStateChanged().RemoveAll(this);
+			TransBuffer->OnUndo().RemoveAll(this);
+			TransBuffer->OnRedo().RemoveAll(this);
+			TransBuffer->OnUndoBufferChanged().RemoveAll(this);
+		}
+#endif
+	}
+	
+private:
+	
+	/**
+	 * Add the Listener to the PendingListener.
+	 * @param Listener Listener to add to the PendingListeners
+	 */
+	void Register(FRCTransactionListener<CallbackArgs...>& Listener)
+	{
+		PendingListeners.Add(Listener);
+	}
+	
+	/**
+	 * Private constructor used to bind to TransactionBuffer callback
+	 */
+	FRCTransactionListenerManager()
+	{
+#if WITH_EDITOR
+		if (GEditor)
+		{
+			UTransBuffer* TransBuffer = CastChecked<UTransBuffer>(GEditor->Trans);
+
+			TransBuffer->OnTransactionStateChanged().AddRaw(this, &FRCTransactionListenerManager::OnTransactionChanged);
+			TransBuffer->OnUndo().AddRaw(this, &FRCTransactionListenerManager::OnUndo);
+			TransBuffer->OnRedo().AddRaw(this, &FRCTransactionListenerManager::OnRedo);
+			TransBuffer->OnUndoBufferChanged().AddRaw(this, &FRCTransactionListenerManager::OnUndoBufferChanged);
+		}
+#endif
+	}
+
+	/**
+	 * Execute all the Listeners that are binded to the current TransactionId and have the same TransactionType.
+	 * @param InListeners Listener binded to the current TransactionId.
+	 * @param InTransactionType Current Transaction type (Undo or Redo).
+	 */
+	void ExecuteAllListeners(TArray<FRCTransactionListener<CallbackArgs...>>& InListeners, ERCTransaction::Type InTransactionType)
+	{
+		for (FRCTransactionListener<CallbackArgs...>& InListener : InListeners)
+		{
+			if (InListener.Type == InTransactionType)
+			{
+				InListener.OnUndoRedo().Broadcast();
+			}
+		}
+	}
+
+	/**
+	 * Gets executed when an Undo is attempted.
+	 * @param InTransactionContext Transaction information.
+	 * @param bSucceeded True if the Undo succeeded otherwise false.
+	 */
+	void OnUndo(const FTransactionContext& InTransactionContext, bool bSucceeded)
+	{
+		if (!bSucceeded)
+		{
+			return;
+		}
+		if (TArray<FRCTransactionListener<CallbackArgs...>>* FoundListeners = Listeners.Find(
+			InTransactionContext.TransactionId))
+		{
+			ExecuteAllListeners(*FoundListeners, ERCTransaction::Undo);
+		}
+	}
+
+	/**
+	 * Gets executed when a Redo is attempted.
+	 * @param InTransactionContext Transaction information.
+	 * @param bSucceeded True if the Redo succeeded otherwise false.
+	 */
+	void OnRedo(const FTransactionContext& InTransactionContext, bool bSucceeded)
+	{
+		if (!bSucceeded)
+		{
+			return;
+		}
+		if (TArray<FRCTransactionListener<CallbackArgs...>>* FoundListeners = Listeners.Find(
+			InTransactionContext.TransactionId))
+		{
+			ExecuteAllListeners(*FoundListeners, ERCTransaction::Redo);
+		}
+	}
+
+	/**
+ 	 * Gets executed when the undo buffer change.
+ 	 */
+	void OnUndoBufferChanged()
+	{
+		DeleteUnusedCallbacks();
+	}
+
+	/**
+	 * Delete from the Listeners all the Callbacks which are not binded to any Transaction inside the TransactionBuffer by their TransactionID.
+	 */
+	void DeleteUnusedCallbacks()
+	{
+#if WITH_EDITOR
+		if (GEditor)
+		{
+			UTransBuffer* TransBuffer = CastChecked<UTransBuffer>(GEditor->Trans);
+			TArray<FGuid> Keys;
+			Listeners.GetKeys(Keys);
+			for (FGuid Key : Keys)
+			{
+				if (!TransBuffer->UndoBuffer.ContainsByPredicate([Key](TSharedPtr<FTransaction> InTransaction)
+				{
+					return InTransaction->GetId() == Key;
+				}))
+				{
+					Listeners.Remove(Key);
+				}
+			}
+		}
+		if (PendingListeners.Num() == 0 && Listeners.Num() == 0)
+		{
+			Instance.Reset();
+		}
+#endif
+	}
+
+	/**
+	 * Gets executed when a transaction state changes.
+	 * @param InTransactionContext Transaction information.
+	 * @param InTransactionState Current state of the Transaction, see ETransactionStateEventType to know all types.
+	 */
+	void OnTransactionChanged(const FTransactionContext& InTransactionContext, ETransactionStateEventType InTransactionState)
+	{
+		if (InTransactionState == ETransactionStateEventType::TransactionFinalized
+			&& !Listeners.Contains(InTransactionContext.TransactionId))
+		{
+			if (!PendingListeners.IsEmpty())
+			{
+				Listeners.FindOrAdd(InTransactionContext.TransactionId).Append(PendingListeners);
+				PendingListeners.Empty();
+			}
+			DeleteUnusedCallbacks();
+		}
+	}
+
+private:
+	static TUniquePtr<FRCTransactionListenerManager<CallbackArgs...>> Instance;
+	TArray<FRCTransactionListener<CallbackArgs...>> PendingListeners;
+	TMap<FGuid, TArray<FRCTransactionListener<CallbackArgs...>>> Listeners;
+};
+
+template <typename... CallbackArgs>
+TUniquePtr<FRCTransactionListenerManager<CallbackArgs...>> FRCTransactionListenerManager<CallbackArgs...>::Instance(nullptr);
