RCP Bind Logic is a step behind in transactions

diff --git a/Engine/Plugins/VirtualProduction/RemoteControl/Source/RemoteControl/Private/RemoteControlPreset.cpp b/Engine/Plugins/VirtualProduction/RemoteControl/Source/RemoteControl/Private/RemoteControlPreset.cpp
--- a/Engine/Plugins/VirtualProduction/RemoteControl/Source/RemoteControl/Private/RemoteControlPreset.cpp
+++ b/Engine/Plugins/VirtualProduction/RemoteControl/Source/RemoteControl/Private/RemoteControlPreset.cpp
@@ -748,7 +748,7 @@ void URemoteControlPreset::OnModifyController(const FPropertyChangedEvent& Prope
 
 		if (!ModifiedController)
 		{
-			ModifiedController = ControllerContainer->GetVirtualProperty(PropertyChangedEvent.MemberProperty->GetFName());
+			ModifiedController = ControllerContainer->GetVirtualProperty(PropertyChangedEvent.Property->GetFName());
 		}
 
 		if (ModifiedController)
