1. Crashes related to missing exposed property

2. Creating Controller with an Unbound Entity Crashes Editor

diff --git a/Engine/Plugins/VirtualProduction/RemoteControl/Source/RemoteControlLogic/Private/Behaviour/Builtin/Bind/RCBehaviourBind.cpp b/Engine/Plugins/VirtualProduction/RemoteControl/Source/RemoteControlLogic/Private/Behaviour/Builtin/Bind/RCBehaviourBind.cpp
--- a/Engine/Plugins/VirtualProduction/RemoteControl/Source/RemoteControlLogic/Private/Behaviour/Builtin/Bind/RCBehaviourBind.cpp
+++ b/Engine/Plugins/VirtualProduction/RemoteControl/Source/RemoteControlLogic/Private/Behaviour/Builtin/Bind/RCBehaviourBind.cpp
@@ -124,6 +124,11 @@ bool URCBehaviourBind::CanHaveActionForField(URCController* Controller, TSharedR
 	const FProperty* RemoteControlProperty = RemoteControlEntityAsProperty->GetProperty();
 	const FStructProperty* RemoteControlEntityAsStructProperty = CastField<FStructProperty>(RemoteControlProperty);
 
+	if (!ControllerAsProperty || !RemoteControlProperty)
+	{
+		return false; // Unbound Property
+	}
+
 	// Indirect Binding (related types)
 	//
 	const static TMap<FFieldClass*, TArray<FFieldClass*>> SupportedIndirectBindsMap =
diff --git a/Engine/Plugins/VirtualProduction/RemoteControl/Source/RemoteControlUI/Private/UI/Action/Bind/RCActionBindModel.cpp b/Engine/Plugins/VirtualProduction/RemoteControl/Source/RemoteControlUI/Private/UI/Action/Bind/RCActionBindModel.cpp
--- a/Engine/Plugins/VirtualProduction/RemoteControl/Source/RemoteControlUI/Private/UI/Action/Bind/RCActionBindModel.cpp
+++ b/Engine/Plugins/VirtualProduction/RemoteControl/Source/RemoteControlUI/Private/UI/Action/Bind/RCActionBindModel.cpp
@@ -125,7 +125,10 @@ FLinearColor FRCPropertyActionBindModel::GetActionTypeColor() const
 	{
 		if (TSharedPtr<FRemoteControlProperty> RemoteControlProperty = PropertyAction->GetRemoteControlProperty())
 		{
-			return UE::RCUIHelpers::GetFieldClassTypeColor(RemoteControlProperty->GetProperty());
+			if (RemoteControlProperty->IsBound())
+			{
+				return UE::RCUIHelpers::GetFieldClassTypeColor(RemoteControlProperty->GetProperty());
+			}
 		}
 	}
 
diff --git a/Engine/Plugins/VirtualProduction/RemoteControl/Source/RemoteControlUI/Private/UI/Action/RCActionModel.cpp b/Engine/Plugins/VirtualProduction/RemoteControl/Source/RemoteControlUI/Private/UI/Action/RCActionModel.cpp
--- a/Engine/Plugins/VirtualProduction/RemoteControl/Source/RemoteControlUI/Private/UI/Action/RCActionModel.cpp
+++ b/Engine/Plugins/VirtualProduction/RemoteControl/Source/RemoteControlUI/Private/UI/Action/RCActionModel.cpp
@@ -252,14 +252,20 @@ TSharedRef<SWidget> FRCPropertyActionType::GetPropertyWidget() const
 
 FLinearColor FRCPropertyActionType::GetPropertyTypeColor() const
 {
+	FLinearColor TypeColor = FLinearColor::White;
+
 	if (!ensure(PropertyActionWeakPtr.IsValid()))
 	{
-		return FLinearColor::White;
+		return TypeColor;
 	}
 
-	const URCPropertyAction* PropertyAction = PropertyActionWeakPtr.Get();
-
-	const FLinearColor TypeColor = UE::RCUIHelpers::GetFieldClassTypeColor(PropertyAction->PropertySelfContainer->GetProperty());
+	if (const URCPropertyAction* PropertyAction = PropertyActionWeakPtr.Get())
+	{
+		if (const FProperty* Property = PropertyAction->PropertySelfContainer->GetProperty())
+		{
+			TypeColor = UE::RCUIHelpers::GetFieldClassTypeColor(Property);
+		}
+	}
 
 	return TypeColor;
 }
diff --git a/Engine/Plugins/VirtualProduction/RemoteControl/Source/RemoteControlUI/Private/UI/Action/SRCBehaviourDetails.cpp b/Engine/Plugins/VirtualProduction/RemoteControl/Source/RemoteControlUI/Private/UI/Action/SRCBehaviourDetails.cpp
--- a/Engine/Plugins/VirtualProduction/RemoteControl/Source/RemoteControlUI/Private/UI/Action/SRCBehaviourDetails.cpp
+++ b/Engine/Plugins/VirtualProduction/RemoteControl/Source/RemoteControlUI/Private/UI/Action/SRCBehaviourDetails.cpp
@@ -36,12 +36,15 @@ void SRCBehaviourDetails::Construct(const FArguments& InArgs, TSharedRef<SRCActi
 
 	const FSlateFontInfo& FontBehaviorDesc = FRemoteControlPanelStyle::Get()->GetFontStyle("RemoteControlPanel.Behaviours.BehaviorDescription");
 
-	FLinearColor TypeColor;
-	FString TypeDisplayName;
+	FLinearColor TypeColor = FLinearColor::White;
+	FString TypeDisplayName = (LOCTEXT("NoneDisplayName", "None")).ToString();
 	if (URCController* Controller = Behaviour->ControllerWeakPtr.Get())
 	{
-		TypeColor = UE::RCUIHelpers::GetFieldClassTypeColor(Behaviour->ControllerWeakPtr->GetProperty());
-		TypeDisplayName = FName::NameToDisplayString(UE::RCUIHelpers::GetFieldClassDisplayName(Behaviour->ControllerWeakPtr->GetProperty()).ToString(), false);
+		if (FProperty* Property = Behaviour->ControllerWeakPtr->GetProperty())
+		{
+			TypeColor = UE::RCUIHelpers::GetFieldClassTypeColor(Property);
+			TypeDisplayName = FName::NameToDisplayString(UE::RCUIHelpers::GetFieldClassDisplayName(Property).ToString(), false);
+		}
 	}
 
 	BehaviourDetailsWidget = InBehaviourItem->GetBehaviourDetailsWidget();
