- Fix for AllowOrigin and AllowedIP being verified even if RestrictServer is disabled

diff --git a/Engine/Plugins/VirtualProduction/RemoteControl/Source/WebRemoteControl/Private/RemoteControlWebSocketServer.cpp b/Engine/Plugins/VirtualProduction/RemoteControl/Source/WebRemoteControl/Private/RemoteControlWebSocketServer.cpp
--- a/Engine/Plugins/VirtualProduction/RemoteControl/Source/WebRemoteControl/Private/RemoteControlWebSocketServer.cpp
+++ b/Engine/Plugins/VirtualProduction/RemoteControl/Source/WebRemoteControl/Private/RemoteControlWebSocketServer.cpp
@@ -258,37 +258,40 @@ void FRCWebSocketServer::OnSocketClose(INetworkingWebSocket* Socket)
 
 EWebsocketConnectionFilterResult FRCWebSocketServer::FilterConnection(FString OriginHeader, FString ClientIP) const
 {
-	OriginHeader.RemoveSpacesInline();
-	OriginHeader.TrimStartAndEndInline();
-	
-	auto SimplifyAddress = [] (FString Address)
-	{
-		Address.RemoveFromStart(TEXT("https://www."));
-		Address.RemoveFromStart(TEXT("http://www."));
-		Address.RemoveFromStart(TEXT("https://"));
-		Address.RemoveFromStart(TEXT("http://"));
-		Address.RemoveFromEnd(TEXT("/"));
-		return Address;
-	};
-
-	const FString SimplifiedOrigin = SimplifyAddress(OriginHeader);
-	const FWildcardString SimplifiedAllowedOrigin = SimplifyAddress(GetDefault<URemoteControlSettings>()->AllowedOrigin);
-	if (GetDefault<URemoteControlSettings>()->AllowedOrigin != TEXT("*"))
+	const URemoteControlSettings* Settings = GetDefault<URemoteControlSettings>();
+	if (Settings->bRestrictServerAccess)
 	{
-		if (!SimplifiedAllowedOrigin.IsMatch(SimplifiedOrigin))
+		OriginHeader.RemoveSpacesInline();
+		OriginHeader.TrimStartAndEndInline();
+		
+		auto SimplifyAddress = [] (FString Address)
+		{
+			Address.RemoveFromStart(TEXT("https://www."));
+			Address.RemoveFromStart(TEXT("http://www."));
+			Address.RemoveFromStart(TEXT("https://"));
+			Address.RemoveFromStart(TEXT("http://"));
+			Address.RemoveFromEnd(TEXT("/"));
+			return Address;
+		};
+
+		const FString SimplifiedOrigin = SimplifyAddress(OriginHeader);
+		if (Settings->AllowedOrigin != TEXT("*"))
 		{
-			return EWebsocketConnectionFilterResult::ConnectionRefused;
+			const FWildcardString SimplifiedAllowedOrigin = SimplifyAddress(Settings->AllowedOrigin);
+			if (!SimplifiedAllowedOrigin.IsMatch(SimplifiedOrigin))
+			{
+				return EWebsocketConnectionFilterResult::ConnectionRefused;
+			}
 		}
-	}
-
-	const FWildcardString WildcardAllowedIP = SimplifyAddress(GetDefault<URemoteControlSettings>()->AllowedIP);
 
-	// Allow requests from localhost
-	if (ClientIP != TEXT("localhost") && ClientIP != TEXT("127.0.0.1"))
-	{
-		if (!WildcardAllowedIP.IsEmpty() && WildcardAllowedIP.IsMatch(ClientIP))
+		// Allow requests from localhost
+		if (ClientIP != TEXT("localhost") && ClientIP != TEXT("127.0.0.1"))
 		{
-			return EWebsocketConnectionFilterResult::ConnectionRefused;
+			const FWildcardString WildcardAllowedIP = SimplifyAddress(Settings->AllowedIP);
+			if (!WildcardAllowedIP.IsEmpty() && WildcardAllowedIP.IsMatch(ClientIP))
+			{
+				return EWebsocketConnectionFilterResult::ConnectionRefused;
+			}
 		}
 	}
 
diff --git a/Engine/Plugins/VirtualProduction/RemoteControl/Source/WebRemoteControl/Private/WebRemoteControl.cpp b/Engine/Plugins/VirtualProduction/RemoteControl/Source/WebRemoteControl/Private/WebRemoteControl.cpp
--- a/Engine/Plugins/VirtualProduction/RemoteControl/Source/WebRemoteControl/Private/WebRemoteControl.cpp
+++ b/Engine/Plugins/VirtualProduction/RemoteControl/Source/WebRemoteControl/Private/WebRemoteControl.cpp
@@ -2364,60 +2364,64 @@ void FWebRemoteControlModule::RegisterDefaultPreprocessors()
 
 	const FHttpRequestHandler IPValidationRequestHandler = FHttpRequestHandler([this](const FHttpServerRequest& Request, const FHttpResultCallback& OnComplete)
 	{
-		TUniquePtr<FHttpServerResponse> Response = WebRemoteControlInternalUtils::CreateHttpResponse();
-
-		FString OriginHeader;
-		if (const TArray<FString>* OriginHeaders = Request.Headers.Find(WebRemoteControlInternalUtils::OriginHeader))
+		const URemoteControlSettings* Settings = GetDefault<URemoteControlSettings>();
+		if (Settings->bRestrictServerAccess)
 		{
-			if (OriginHeaders->Num())
-			{
-				OriginHeader = (*OriginHeaders)[0];
-			}
-		}
+			TUniquePtr<FHttpServerResponse> Response = WebRemoteControlInternalUtils::CreateHttpResponse();
 
-		OriginHeader.RemoveSpacesInline();
-		OriginHeader.TrimStartAndEndInline();
-		
-		auto SimplifyAddress = [] (FString Address)
-		{
-			Address.RemoveFromStart(TEXT("https://www."));
-			Address.RemoveFromStart(TEXT("http://www."));
-			Address.RemoveFromStart(TEXT("https://"));
-			Address.RemoveFromStart(TEXT("http://"));
-			Address.RemoveFromEnd(TEXT("/"));
-			return Address;
-		};
-
-		const FString SimplifiedOrigin = SimplifyAddress(OriginHeader);
-		const FWildcardString SimplifiedAllowedOrigin = SimplifyAddress(GetDefault<URemoteControlSettings>()->AllowedOrigin);
-		if (GetDefault<URemoteControlSettings>()->AllowedOrigin != TEXT("*"))
-		{
-			if (!SimplifiedAllowedOrigin.IsMatch(SimplifiedOrigin))
+			FString OriginHeader;
+			if (const TArray<FString>* OriginHeaders = Request.Headers.Find(WebRemoteControlInternalUtils::OriginHeader))
 			{
-				WebRemoteControlInternalUtils::CreateUTF8ErrorMessage(FString::Printf(TEXT("Client origin %s does not respect the allowed origin set in Remote Control Settings."), *OriginHeader), Response->Body);
-				Response->Code = EHttpServerResponseCodes::Denied;
-				OnComplete(MoveTemp(Response));
-				return true;
+				if (OriginHeaders->Num())
+				{
+					OriginHeader = (*OriginHeaders)[0];
+				}
 			}
-		}
 
-		if (Request.PeerAddress)
-		{
-			constexpr bool bAppendPort = false;
-			FString ClientIP = Request.PeerAddress->ToString(bAppendPort);
-			const FWildcardString WildcardAllowedIP = SimplifyAddress(GetDefault<URemoteControlSettings>()->AllowedIP);
-				
-			// Allow requests from localhost
-			if (ClientIP != TEXT("localhost") && ClientIP != TEXT("127.0.0.1"))
+			OriginHeader.RemoveSpacesInline();
+			OriginHeader.TrimStartAndEndInline();
+			
+			auto SimplifyAddress = [] (FString Address)
 			{
-				if (!WildcardAllowedIP.IsEmpty() && WildcardAllowedIP.IsMatch(ClientIP))
+				Address.RemoveFromStart(TEXT("https://www."));
+				Address.RemoveFromStart(TEXT("http://www."));
+				Address.RemoveFromStart(TEXT("https://"));
+				Address.RemoveFromStart(TEXT("http://"));
+				Address.RemoveFromEnd(TEXT("/"));
+				return Address;
+			};
+
+			const FString SimplifiedOrigin = SimplifyAddress(OriginHeader);
+			if (Settings->AllowedOrigin != TEXT("*"))
+			{
+				const FWildcardString SimplifiedAllowedOrigin = SimplifyAddress(Settings->AllowedOrigin);
+				if (!SimplifiedAllowedOrigin.IsMatch(SimplifiedOrigin))
 				{
-					WebRemoteControlInternalUtils::CreateUTF8ErrorMessage(FString::Printf(TEXT("Client IP %s does not respect the allowed IP set in Remote Control Settings."), *ClientIP), Response->Body);
+					WebRemoteControlInternalUtils::CreateUTF8ErrorMessage(FString::Printf(TEXT("Client origin %s does not respect the allowed origin set in Remote Control Settings."), *OriginHeader), Response->Body);
 					Response->Code = EHttpServerResponseCodes::Denied;
 					OnComplete(MoveTemp(Response));
 					return true;
 				}
 			}
+
+			if (Request.PeerAddress)
+			{
+				constexpr bool bAppendPort = false;
+				FString ClientIP = Request.PeerAddress->ToString(bAppendPort);
+
+				// Allow requests from localhost
+				if (ClientIP != TEXT("localhost") && ClientIP != TEXT("127.0.0.1"))
+				{
+					const FWildcardString WildcardAllowedIP = SimplifyAddress(Settings->AllowedIP);
+					if (!WildcardAllowedIP.IsEmpty() && WildcardAllowedIP.IsMatch(ClientIP))
+					{
+						WebRemoteControlInternalUtils::CreateUTF8ErrorMessage(FString::Printf(TEXT("Client IP %s does not respect the allowed IP set in Remote Control Settings."), *ClientIP), Response->Body);
+						Response->Code = EHttpServerResponseCodes::Denied;
+						OnComplete(MoveTemp(Response));
+						return true;
+					}
+				}
+			}
 		}
 
 		return false;
