[RemoteControl] HTTP routes to add items to/remove items from arrays

diff --git a/Engine/Plugins/VirtualProduction/EpicStageApp/Source/EpicStageApp/Public/StageAppVersion.h b/Engine/Plugins/VirtualProduction/EpicStageApp/Source/EpicStageApp/Public/StageAppVersion.h
--- a/Engine/Plugins/VirtualProduction/EpicStageApp/Source/EpicStageApp/Public/StageAppVersion.h
+++ b/Engine/Plugins/VirtualProduction/EpicStageApp/Source/EpicStageApp/Public/StageAppVersion.h
@@ -11,7 +11,7 @@ struct EPICSTAGEAPP_API FEpicStageAppAPIVersion
 	static const uint16 Major = 1;
 
 	/** The minor component of the current version. */
-	static const uint16 Minor = 0;
+	static const uint16 Minor = 1;
 	
 	/** The patch component of the current version. */
 	static const uint16 Patch = 0;
diff --git a/Engine/Plugins/VirtualProduction/RemoteControl/Source/RemoteControl/Private/RemoteControlModule.cpp b/Engine/Plugins/VirtualProduction/RemoteControl/Source/RemoteControl/Private/RemoteControlModule.cpp
--- a/Engine/Plugins/VirtualProduction/RemoteControl/Source/RemoteControl/Private/RemoteControlModule.cpp
+++ b/Engine/Plugins/VirtualProduction/RemoteControl/Source/RemoteControl/Private/RemoteControlModule.cpp
@@ -1206,6 +1206,121 @@ bool FRemoteControlModule::GetObjectProperties(const FRCObjectReference& ObjectA
 	return false;
 }
 
+#if WITH_EDITOR
+bool FRemoteControlModule::StartPropertyTransaction(FRCObjectReference& ObjectReference, const FText& InChangeDescription, bool& bOutGeneratedTransaction)
+{
+	const bool bUseOngoingChangeOptimization = CVarRemoteControlEnableOngoingChangeOptimization.GetValueOnAnyThread() == 1;
+	if (bUseOngoingChangeOptimization)
+	{
+		const FString ObjectPath = ObjectReference.Object->GetPathName();
+
+		// If we have a change that hasn't yet generated a post edit change property, do that before handling the next change.
+		EndOngoingModificationIfMismatched(GetTypeHash(ObjectReference));
+
+		// This step is necessary because the object might get recreated by a PostEditChange called in TestOrFinalizeOngoingChange.
+		if (!ObjectReference.Object.IsValid())
+		{
+			ObjectReference.Object = StaticFindObject(UObject::StaticClass(), nullptr, *ObjectPath);
+			if (ObjectReference.Object.IsValid() && ObjectReference.PropertyPathInfo.IsResolved())
+			{
+				// Update ContainerAddress as well if the path was resolved.
+				if (ObjectReference.PropertyPathInfo.Resolve(ObjectReference.Object.Get()))
+				{
+					ObjectReference.ContainerAdress = ObjectReference.PropertyPathInfo.GetResolvedData().ContainerAddress;
+				}
+			}
+			else
+			{
+				return false;
+			}
+		}
+	}
+
+	bOutGeneratedTransaction = ObjectReference.Access == ERCAccess::WRITE_TRANSACTION_ACCESS;
+	const bool bIsNewTransaction = bOutGeneratedTransaction && (!bUseOngoingChangeOptimization || !OngoingModification);
+
+	if (bIsNewTransaction && GEditor)
+	{
+		GEditor->BeginTransaction(InChangeDescription);
+
+		if (bUseOngoingChangeOptimization)
+		{
+			// Call modify since it's not called by PreEditChange until the end of the ongoing change.
+			ObjectReference.Object->Modify();
+		}
+	}
+
+	if (!bUseOngoingChangeOptimization)
+	{
+		FEditPropertyChain PreEditChain;
+		ObjectReference.PropertyPathInfo.ToEditPropertyChain(PreEditChain);
+		ObjectReference.Object->PreEditChange(PreEditChain);
+	}
+
+	return true;
+}
+#endif
+
+#if WITH_EDITOR
+void FRemoteControlModule::SnapshotOrEndTransaction(FRCObjectReference& ObjectReference, bool bGeneratedTransaction)
+{
+	if (CVarRemoteControlEnableOngoingChangeOptimization.GetValueOnAnyThread() == 1)
+	{
+		// If we have modified the same object and property in the last frames,
+		// update the triggered flag and snapshot the object to the transaction buffer.
+		if (OngoingModification && GetTypeHash(*OngoingModification) == GetTypeHash(ObjectReference))
+		{
+			OngoingModification->bWasTriggeredSinceLastPass = true;
+			if (OngoingModification->bHasStartedTransaction)
+			{
+				SnapshotTransactionBuffer(ObjectReference.Object.Get());
+				FPropertyChangedEvent PropertyEvent(ObjectReference.PropertyPathInfo.ToPropertyChangedEvent());
+				PropertyEvent.ChangeType = EPropertyChangeType::Interactive;
+				ObjectReference.Object->PostEditChangeProperty(PropertyEvent);
+
+				if (UActorComponent* Component = Cast<UActorComponent>(ObjectReference.Object.Get()))
+				{
+					Component->MarkRenderStateDirty();
+					Component->UpdateComponentToWorld();
+				}
+			}
+
+			// Update the world lighting if we're modifying a color.
+			if (ObjectReference.IsValid())
+			{
+				if (ObjectReference.Object->IsA<ULightComponent>())
+				{
+					UClass* OwnerClass = ObjectReference.Property->GetOwnerClass();
+					if (OwnerClass && OwnerClass->IsChildOf(ULightComponentBase::StaticClass()))
+					{
+						UWorld* World = ObjectReference.Object->GetWorld();
+						if (World && World->Scene)
+						{
+							World->Scene->UpdateLightColorAndBrightness(Cast<ULightComponent>(ObjectReference.Object.Get()));
+						}
+					}
+				}
+			}
+		}
+		else if (ObjectReference.Access != ERCAccess::WRITE_MANUAL_TRANSACTION_ACCESS)
+		{
+			OngoingModification = ObjectReference;
+			OngoingModification->bHasStartedTransaction = bGeneratedTransaction;
+		}
+	}
+	else
+	{
+		if (GEditor && bGeneratedTransaction)
+		{
+			GEditor->EndTransaction();
+		}
+
+		FPropertyChangedEvent PropertyEvent(ObjectReference.PropertyPathInfo.ToPropertyChangedEvent());
+		ObjectReference.Object->PostEditChangeProperty(PropertyEvent);
+	}
+}
+#endif
+
 bool FRemoteControlModule::SetObjectProperties(const FRCObjectReference& ObjectAccess, IStructDeserializerBackend& Backend, ERCPayloadType InPayloadType, const TArray<uint8>& InPayload, ERCModifyOperation Operation)
 {
 	TRACE_CPUPROFILER_EVENT_SCOPE(FRemoteControlModule::SetObjectProperties);
@@ -1324,52 +1439,10 @@ bool FRemoteControlModule::SetObjectProperties(const FRCObjectReference& ObjectA
 		FRCObjectReference MutableObjectReference = ObjectAccess;
 
 #if WITH_EDITOR
-		const bool bUseOngoingChangeOptimization = CVarRemoteControlEnableOngoingChangeOptimization.GetValueOnAnyThread() == 1;
-		if (bUseOngoingChangeOptimization)
-		{
-			const FString ObjectPath = MutableObjectReference.Object->GetPathName();
-
-			// If we have a change that hasn't yet generated a post edit change property, do that before handling the next change.
-			EndOngoingModificationIfMismatched(GetTypeHash(MutableObjectReference));
-
-			// This step is necessary because the object might get recreated by a PostEditChange called in TestOrFinalizeOngoingChange.
-			if (!MutableObjectReference.Object.IsValid())
-			{
-				MutableObjectReference.Object = StaticFindObject(UObject::StaticClass(), nullptr, *ObjectPath);
-				if (MutableObjectReference.Object.IsValid() && MutableObjectReference.PropertyPathInfo.IsResolved())
-				{
-					// Update ContainerAddress as well if the path was resolved.
-					if (MutableObjectReference.PropertyPathInfo.Resolve(MutableObjectReference.Object.Get()))
-					{
-						MutableObjectReference.ContainerAdress = MutableObjectReference.PropertyPathInfo.GetResolvedData().ContainerAddress;
-					}
-				}
-				else
-				{
-					return false;
-				}
-			}
-		}
-
-		const bool bGenerateTransaction = MutableObjectReference.Access == ERCAccess::WRITE_TRANSACTION_ACCESS;
-		const bool bIsNewTransaction = bGenerateTransaction && (!bUseOngoingChangeOptimization || !OngoingModification);
-
-		if (bIsNewTransaction && GEditor)
-		{
-			GEditor->BeginTransaction(LOCTEXT("RemoteSetPropertyTransaction", "Remote Set Object Property"));
-
-			if (bUseOngoingChangeOptimization)
-			{
-				// Call modify since it's not called by PreEditChange until the end of the ongoing change.
-				MutableObjectReference.Object->Modify();
-			}
-		}
-
-		if (!bUseOngoingChangeOptimization)
+		bool bGeneratedTransaction;
+		if (!StartPropertyTransaction(MutableObjectReference, LOCTEXT("RemoteSetPropertyTransaction", "Remote Set Object Property"), bGeneratedTransaction))
 		{
-			FEditPropertyChain PreEditChain;
-			MutableObjectReference.PropertyPathInfo.ToEditPropertyChain(PreEditChain);
-			MutableObjectReference.Object->PreEditChange(PreEditChain);
+			return false;
 		}
 #endif
 
@@ -1421,64 +1494,12 @@ bool FRemoteControlModule::SetObjectProperties(const FRCObjectReference& ObjectA
 		}
 
 #if WITH_EDITOR
-		if (CVarRemoteControlEnableOngoingChangeOptimization.GetValueOnAnyThread() == 1)
-		{
-			// If we have modified the same object and property in the last frames,
-			// update the triggered flag and snapshot the object to the transaction buffer.
-			if (OngoingModification && GetTypeHash(*OngoingModification) == GetTypeHash(MutableObjectReference))
-			{
-				OngoingModification->bWasTriggeredSinceLastPass = true;
-				if (OngoingModification->bHasStartedTransaction)
-				{
-					SnapshotTransactionBuffer(MutableObjectReference.Object.Get());
-					FPropertyChangedEvent PropertyEvent(MutableObjectReference.PropertyPathInfo.ToPropertyChangedEvent());
-					PropertyEvent.ChangeType = EPropertyChangeType::Interactive;
-					MutableObjectReference.Object->PostEditChangeProperty(PropertyEvent);
-
-					if (UActorComponent* Component = Cast<UActorComponent>(MutableObjectReference.Object.Get()))
-					{
-						Component->MarkRenderStateDirty();
-						Component->UpdateComponentToWorld();
-					}
-				}
-
-				// Update the world lighting if we're modifying a color.
-				if (MutableObjectReference.IsValid() )
-				{
-					if (MutableObjectReference.Object->IsA<ULightComponent>())
-					{
-						UClass* OwnerClass = MutableObjectReference.Property->GetOwnerClass();
-						if (OwnerClass && OwnerClass->IsChildOf(ULightComponentBase::StaticClass()))
-						{
-							UWorld* World = MutableObjectReference.Object->GetWorld();
-							if (World && World->Scene)
-							{
-								World->Scene->UpdateLightColorAndBrightness(Cast<ULightComponent>(MutableObjectReference.Object.Get()));
-							}
-						}
-					}
-				}
-			}
-			else if (MutableObjectReference.Access != ERCAccess::WRITE_MANUAL_TRANSACTION_ACCESS)
-			{
-				OngoingModification = MutableObjectReference;
-				OngoingModification->bHasStartedTransaction = bGenerateTransaction;
-			}
-		}
-		else
-		{
-			if (GEditor && bGenerateTransaction)
-			{
-				GEditor->EndTransaction();
-			}
-
-			FPropertyChangedEvent PropertyEvent(MutableObjectReference.PropertyPathInfo.ToPropertyChangedEvent());
-			MutableObjectReference.Object->PostEditChangeProperty(PropertyEvent);
-		}
+		SnapshotOrEndTransaction(MutableObjectReference, bGeneratedTransaction);
 #endif
+
 		for (const TPair<FName, TSharedPtr<IRemoteControlPropertyFactory>>& EntityFactoryPair : EntityFactories)
 		{
-			EntityFactoryPair.Value->PostSetObjectProperties(ObjectAccess.Object.Get(), bSuccess);
+			EntityFactoryPair.Value->PostSetObjectProperties(MutableObjectReference.Object.Get(), bSuccess);
 		}
 
 		if (bSuccess)
@@ -1491,6 +1512,51 @@ bool FRemoteControlModule::SetObjectProperties(const FRCObjectReference& ObjectA
 	return false;
 }
 
+bool FRemoteControlModule::AppendToObjectArrayProperty(const FRCObjectReference& ObjectAccess, IStructDeserializerBackend& Backend, ERCPayloadType InPayloadType, const TArray<uint8>& InPayload)
+{
+	TRACE_CPUPROFILER_EVENT_SCOPE(FRemoteControlModule::AppendToObjectArrayProperty);
+	UE_LOG(LogRemoteControl, VeryVerbose, TEXT("Append to Object Array Property"));
+
+	return AddToArrayProperty(ObjectAccess, Backend, [](FScriptArrayHelper& ArrayHelper)
+	{
+		return ArrayHelper.AddValue();
+	});
+}
+
+bool FRemoteControlModule::InsertToObjectArrayProperty(int32 Index, const FRCObjectReference& ObjectAccess, IStructDeserializerBackend& Backend, ERCPayloadType InPayloadType, const TArray<uint8>& InPayload)
+{
+	TRACE_CPUPROFILER_EVENT_SCOPE(FRemoteControlModule::InsertToObjectArrayProperty);
+	UE_LOG(LogRemoteControl, VeryVerbose, TEXT("Insert to Object Array Property"));
+
+	return AddToArrayProperty(ObjectAccess, Backend, [Index](FScriptArrayHelper& ArrayHelper)
+	{
+		if (!ArrayHelper.IsValidIndex(Index))
+		{
+			return (int32)INDEX_NONE;
+		}
+
+		ArrayHelper.InsertValues(Index, 1);
+		return Index;
+	});
+}
+
+bool FRemoteControlModule::RemoveFromObjectArrayProperty(int32 Index, const FRCObjectReference& ObjectAccess)
+{
+	TRACE_CPUPROFILER_EVENT_SCOPE(FRemoteControlModule::RemoveFromObjectArrayProperty);
+	UE_LOG(LogRemoteControl, VeryVerbose, TEXT("Remove from Object Array Property"));
+
+	return ModifyArrayProperty(ObjectAccess, [Index](FScriptArrayHelper& ArrayHelper)
+	{
+		if (!ArrayHelper.IsValidIndex(Index))
+		{
+			return false;
+		}
+
+		ArrayHelper.RemoveValues(Index);
+		return true;
+	});
+}
+
 void FRemoteControlModule::RefreshEditorPostSetObjectProperties(const FRCObjectReference& ObjectAccess)
 {
 #if WITH_EDITOR
@@ -2329,6 +2395,71 @@ void FRemoteControlModule::RegisterMaskingFactories()
 	RegisterMaskingFactoryForType(TBaseStructure<FLinearColor>::Get(), FLinearColorMaskingFactory::MakeInstance());
 }
 
+bool FRemoteControlModule::AddToArrayProperty(const FRCObjectReference& ObjectAccess, IStructDeserializerBackend& Backend, TFunctionRef<int32(FScriptArrayHelper&)> ModifyFunction)
+{
+	return ModifyArrayProperty(ObjectAccess, [&ObjectAccess, &Backend, &ModifyFunction](FScriptArrayHelper& ArrayHelper) {
+		const int32 Index = ModifyFunction(ArrayHelper);
+
+		if (Index == INDEX_NONE)
+		{
+			return false;
+		}
+
+		const FRCFieldPathSegment& LastSegment = ObjectAccess.PropertyPathInfo.GetFieldSegment(ObjectAccess.PropertyPathInfo.GetSegmentCount() - 1);
+
+		FStructDeserializerPolicies Policies;
+		Policies.PropertyFilter = [&ObjectAccess](const FProperty* CurrentProp, const FProperty* ParentProp)
+		{
+			return CurrentProp == ObjectAccess.Property || ParentProp != nullptr;
+		};
+
+		return FStructDeserializer::DeserializeElement(ObjectAccess.ContainerAdress, *LastSegment.ResolvedData.Struct, Index, Backend, Policies);
+	});
+}
+
+bool FRemoteControlModule::ModifyArrayProperty(const FRCObjectReference& ObjectAccess, TFunctionRef<bool(FScriptArrayHelper&)> ModifyFunction)
+{
+	if (ObjectAccess.IsValid()
+		&& (RemoteControlUtil::IsWriteAccess(ObjectAccess.Access))
+		&& ObjectAccess.Property.IsValid()
+		&& ObjectAccess.PropertyPathInfo.IsResolved())
+	{
+		if (FArrayProperty* ArrayProperty = CastField<FArrayProperty>(ObjectAccess.Property.Get()))
+		{
+#if WITH_EDITOR
+			FRCObjectReference MutableObjectReference = ObjectAccess;
+
+			bool bGeneratedTransaction;
+			if (!StartPropertyTransaction(MutableObjectReference, LOCTEXT("RemoteArrayAppendTransaction", "Remote Append to Object Array Property"), bGeneratedTransaction))
+			{
+				return false;
+			}
+#endif
+
+			FScriptArrayHelper ArrayHelper(ArrayProperty, ArrayProperty->ContainerPtrToValuePtr<void>(ObjectAccess.ContainerAdress));
+
+			const bool bSuccess = ModifyFunction(ArrayHelper);
+
+#if WITH_EDITOR
+			SnapshotOrEndTransaction(MutableObjectReference, bGeneratedTransaction);
+
+			for (const TPair<FName, TSharedPtr<IRemoteControlPropertyFactory>>& EntityFactoryPair : EntityFactories)
+			{
+				EntityFactoryPair.Value->PostSetObjectProperties(MutableObjectReference.Object.Get(), bSuccess);
+			}
+#endif
+
+			if (bSuccess)
+			{
+				RefreshEditorPostSetObjectProperties(ObjectAccess);
+			}
+
+			return bSuccess;
+		}
+	}
+
+	return false;
+}
 
 #if WITH_EDITOR
 void FRemoteControlModule::EndOngoingModificationIfMismatched(uint32 TypeHash)
diff --git a/Engine/Plugins/VirtualProduction/RemoteControl/Source/RemoteControl/Private/RemoteControlModule.h b/Engine/Plugins/VirtualProduction/RemoteControl/Source/RemoteControl/Private/RemoteControlModule.h
--- a/Engine/Plugins/VirtualProduction/RemoteControl/Source/RemoteControl/Private/RemoteControlModule.h
+++ b/Engine/Plugins/VirtualProduction/RemoteControl/Source/RemoteControl/Private/RemoteControlModule.h
@@ -45,6 +45,9 @@ public:
 	virtual bool GetObjectProperties(const FRCObjectReference& ObjectAccess, IStructSerializerBackend& Backend) override;
 	virtual bool SetObjectProperties(const FRCObjectReference& ObjectAccess, IStructDeserializerBackend& Backend, ERCPayloadType InPayloadType, const TArray<uint8>& InPayload, ERCModifyOperation Operation) override;
 	virtual bool ResetObjectProperties(const FRCObjectReference& ObjectAccess, const bool bAllowIntercept) override;
+	virtual bool InsertToObjectArrayProperty(int32 Index, const FRCObjectReference& ObjectAccess, IStructDeserializerBackend& Backend, ERCPayloadType InPayloadType, const TArray<uint8>& InInterceptPayload) override;
+	virtual bool RemoveFromObjectArrayProperty(int32 Index, const FRCObjectReference& ObjectAccess) override;
+	virtual bool AppendToObjectArrayProperty(const FRCObjectReference& ObjectAccess, IStructDeserializerBackend& Backend, ERCPayloadType InPayloadType, const TArray<uint8>& InInterceptPayload) override;
 	virtual bool SetPresetController(const FName PresetName, const FName ControllerName, IStructDeserializerBackend& Backend, const TArray<uint8>& InPayload, const bool bAllowIntercept) override;
 	virtual bool SetPresetController(const FName PresetName, class URCVirtualPropertyBase* VirtualProperty, IStructDeserializerBackend& Backend, const TArray<uint8>& InPayload, const bool bAllowIntercept) override;
 	virtual TOptional<FExposedFunction> ResolvePresetFunction(const FResolvePresetFieldArgs& Args) const override;
@@ -121,7 +124,40 @@ private:
 	 */
 	void RegisterMaskingFactories();
 
+	/**
+	 * Add an item to an array property, handling transactions and editor events automatically.
+	 * @param ObjectAccess Reference to the array property to be modified.
+	 * @param Backend Deserialization backend for the value to be added to the array. This will be deserialized into the array at the index provided by ModifyFunction.
+	 * @param ModifyFunction Function which takes an array helper, attempts to modify the array, and returns the index of the new item (or INDEX_NONE if the request is invalid).
+	 */
+	bool AddToArrayProperty(const FRCObjectReference& ObjectAccess, IStructDeserializerBackend& Backend, TFunctionRef<int32(FScriptArrayHelper&)> ModifyFunction);
+
+	/**
+	 * Modify an array property, handling transactions and editor events automatically.
+	 * @param ObjectAccess Reference to the array property to be modified.
+	 * @param ModifyFunction Function which takes an array helper, attempts to modify the array, and returns true if the request was valid.
+	 */
+	bool ModifyArrayProperty(const FRCObjectReference& ObjectAccess, TFunctionRef<bool(FScriptArrayHelper&)> ModifyFunction);
+
 #if WITH_EDITOR
+
+	/**
+	 * Create a transaction to modify an object's property, taking into account any ongoing transactions before this one.
+	 * @param ObjectReference A reference to the object to be modified.
+	 * @param InChangeDescription The description of the transaction to be created.
+	 * @param bOutGeneratedTransaction Set to true if a transaction needed to be created, else false.
+	 * @return true if the property is ready to be modified, else false.
+	 */
+	bool StartPropertyTransaction(FRCObjectReference& ObjectReference, const FText& InChangeDescription, bool& bOutGeneratedTransaction);
+
+	/**
+	 * Depending on transaction mode, either snapshot the object to the transaction buffer or end the transaction.
+	 * This should be called after StartPropertyTransaction and then modifying the object's properties.
+	 * @param ObjectReference A reference to the object to be modified.
+	 * @param bGeneratedTransaction The value of bOutGeneratedTransaction provided by StartPropertyTransaction.
+	 */
+	void SnapshotOrEndTransaction(FRCObjectReference& ObjectReference, bool bGeneratedTransaction);
+
 	/**
 	 * End the ongoing modification if one exists and is a mismatch for the new object to edit.
 	 * @param TypeHash The type hash of the object we want to modify after this check.
diff --git a/Engine/Plugins/VirtualProduction/RemoteControl/Source/RemoteControl/Public/IRemoteControlModule.h b/Engine/Plugins/VirtualProduction/RemoteControl/Source/RemoteControl/Public/IRemoteControlModule.h
--- a/Engine/Plugins/VirtualProduction/RemoteControl/Source/RemoteControl/Public/IRemoteControlModule.h
+++ b/Engine/Plugins/VirtualProduction/RemoteControl/Source/RemoteControl/Public/IRemoteControlModule.h
@@ -394,6 +394,35 @@ public:
 	 */
 	virtual bool ResetObjectProperties(const FRCObjectReference& ObjectAccess, const bool bAllowIntercept = false) = 0;
 
+	/**
+	 * Deserialize the Object Reference from the specified backend and append it to an array property.
+	 * @param ObjectAccess the object reference to deserialize into, it should be a write access reference. if the object is WRITE_TRANSACTION_ACCESS, the setting will be wrapped in a transaction.
+	 * @param Backend the struct deserializer backend to use to deserialize the object properties.
+	 * @param InPayloadType the payload type archive.
+	 * @param InInterceptPayload the payload reference archive for the interception.
+	 * @return true if the operation succeeded
+	 */
+	virtual bool AppendToObjectArrayProperty(const FRCObjectReference& ObjectAccess, IStructDeserializerBackend& Backend, ERCPayloadType InPayloadType = ERCPayloadType::Json, const TArray<uint8>& InInterceptPayload = TArray<uint8>()) = 0;
+
+	/**
+	 * Deserialize the Object Reference from the specified backend and insert it at the given position in an array property.
+	 * @param Index The index into the array to insert at.
+	 * @param ObjectAccess the object reference to deserialize into, it should be a write access reference. if the object is WRITE_TRANSACTION_ACCESS, the setting will be wrapped in a transaction.
+	 * @param Backend the struct deserializer backend to use to deserialize the object properties.
+	 * @param InPayloadType the payload type archive.
+	 * @param InInterceptPayload the payload reference archive for the interception.
+	 * @return true if the operation succeeded
+	 */
+	virtual bool InsertToObjectArrayProperty(int32 Index, const FRCObjectReference& ObjectAccess, IStructDeserializerBackend& Backend, ERCPayloadType InPayloadType = ERCPayloadType::Json, const TArray<uint8>& InInterceptPayload = TArray<uint8>()) = 0;
+
+	/**
+	 * Remove an item from the referenced array.
+	 * @param Index The index into the array to remove.
+	 * @param ObjectAccess The reference to the array. If the object is WRITE_TRANSACTION_ACCESS, the operation will be wrapped in a transaction.
+	 * @return true if the operation succeeded
+	 */
+	virtual bool RemoveFromObjectArrayProperty(int32 Index, const FRCObjectReference& ObjectAccess) = 0;
+
 	/**
 	* Set a controller's value on a Remote Control Preset
 	* @param PresetName - The Remote Control Preset asset's name
diff --git a/Engine/Plugins/VirtualProduction/RemoteControl/Source/WebRemoteControl/Private/WebRemoteControl.cpp b/Engine/Plugins/VirtualProduction/RemoteControl/Source/WebRemoteControl/Private/WebRemoteControl.cpp
--- a/Engine/Plugins/VirtualProduction/RemoteControl/Source/WebRemoteControl/Private/WebRemoteControl.cpp
+++ b/Engine/Plugins/VirtualProduction/RemoteControl/Source/WebRemoteControl/Private/WebRemoteControl.cpp
@@ -267,6 +267,65 @@ namespace WebRemoteControl
 			}
 		}
 	}
+
+	/**
+	 * Resolve the object referenced in a WebRemoteControl request.
+	 * Returns true if the object was resolved, or else returns false and calls OnComplete with an error response.
+	 */
+	bool ResolveObjectPropertyForRequest(const FHttpServerRequest& Request, const FHttpResultCallback& OnComplete, TUniquePtr<FHttpServerResponse>& InResponse,
+		FRCObjectReference& OutObjectRef, FRCObjectRequest& OutDeserializedRequest)
+	{
+		if (!WebRemoteControlInternalUtils::ValidateContentType(Request, TEXT("application/json"), OnComplete))
+		{
+			return false;
+		}
+
+		if (!RemotePayloadSerializer::DeserializeObjectRef(Request, OutObjectRef, OutDeserializedRequest, OnComplete))
+		{
+			return false;
+		}
+
+		// If we haven't found the object, return a not found error code
+		if (!OutObjectRef.IsValid())
+		{
+			InResponse->Code = EHttpServerResponseCodes::NotFound;
+			WebRemoteControlInternalUtils::CreateUTF8ErrorMessage(TEXT("Unable to find the object."), InResponse->Body);
+			OnComplete(MoveTemp(InResponse));
+			return false;
+		}
+
+		return true;
+	}
+
+	/**
+	 * Create a memory reader for the deserialized property data contained in a request to add an element to an array.
+	 */
+	FMemoryReader MakePropertyMemoryReaderForArrayAddRequest(FRCObjectRequest& InDeserializedRequest)
+	{
+		const FBlockDelimiters& PropertyValueDelimiters = InDeserializedRequest.GetStructParameters().FindChecked(FRCObjectRequest::PropertyValueLabel());
+		FMemoryReader Reader(InDeserializedRequest.TCHARBody);
+		Reader.Seek(PropertyValueDelimiters.BlockStart);
+		Reader.SetLimitSize(PropertyValueDelimiters.BlockEnd);
+
+		return MoveTemp(Reader);
+	}
+
+	/**
+	 * Retrieve the integer "index" query parameter from a request, or if unable to, call OnComplete with an error response and return false.
+	 */
+	bool GetIndexQueryParamForRequest(const FHttpServerRequest& Request, const FHttpResultCallback& OnComplete, TUniquePtr<FHttpServerResponse>& InResponse, int32& OutIndex)
+	{
+		const FString* IndexParam = Request.QueryParams.Find(TEXT("index"));
+		if (!IndexParam) {
+			InResponse->Code = EHttpServerResponseCodes::BadRequest;
+			WebRemoteControlInternalUtils::CreateUTF8ErrorMessage(TEXT("No 'index' query parameter was provided."), InResponse->Body);
+			OnComplete(MoveTemp(InResponse));
+			return false;
+		}
+
+		OutIndex = FCString::Atoi(**IndexParam);
+		return true;
+	}
 }
 
 void FWebRemoteControlModule::StartupModule()
@@ -579,6 +638,27 @@ void FWebRemoteControlModule::RegisterRoutes()
 		FRequestHandlerDelegate::CreateRaw(this, &FWebRemoteControlModule::HandleObjectPropertyRoute)
 		});
 
+	RegisterRoute({
+		TEXT("Append a new item to an array property on a remote object."),
+		FHttpPath(TEXT("/remote/object/property/append")),
+		EHttpServerRequestVerbs::VERB_PUT,
+		FRequestHandlerDelegate::CreateRaw(this, &FWebRemoteControlModule::HandleObjectPropertyAppendRoute)
+		});
+
+	RegisterRoute({
+		TEXT("Insert a new item into an array property on a remote object."),
+		FHttpPath(TEXT("/remote/object/property/insert")),
+		EHttpServerRequestVerbs::VERB_PUT,
+		FRequestHandlerDelegate::CreateRaw(this, &FWebRemoteControlModule::HandleObjectPropertyInsertRoute)
+		});
+
+	RegisterRoute({
+		TEXT("Remove an item from an array property on a remote object."),
+		FHttpPath(TEXT("/remote/object/property/remove")),
+		EHttpServerRequestVerbs::VERB_PUT,
+		FRequestHandlerDelegate::CreateRaw(this, &FWebRemoteControlModule::HandleObjectPropertyRemoveRoute)
+		});
+
 	RegisterRoute({
 		TEXT("Describe an object."),
 		FHttpPath(TEXT("/remote/object/describe")),
@@ -937,27 +1017,11 @@ bool FWebRemoteControlModule::HandleObjectPropertyRoute(const FHttpServerRequest
 {
 	TUniquePtr<FHttpServerResponse> Response = WebRemoteControlInternalUtils::CreateHttpResponse();
 
-	if (!WebRemoteControlInternalUtils::ValidateContentType(Request, TEXT("application/json"), OnComplete))
-	{
-		return true;
-	}
-
-	FRCObjectReference ObjectRef;
 	bool bResetToDefault = false;
-	FString ErrorText;
+	FRCObjectReference ObjectRef;
 	FRCObjectRequest DeserializedRequest;
-
-	if (!RemotePayloadSerializer::DeserializeObjectRef(Request, ObjectRef, DeserializedRequest, OnComplete))
-	{
-		return true;
-	}
-
-	// If we haven't found the object, return a not found error code
-	if (!ObjectRef.IsValid())
+	if (!WebRemoteControl::ResolveObjectPropertyForRequest(Request, OnComplete, Response, ObjectRef, DeserializedRequest))
 	{
-		Response->Code = EHttpServerResponseCodes::NotFound;
-		WebRemoteControlInternalUtils::CreateUTF8ErrorMessage(TEXT("Unable to find the object."), Response->Body);
-		OnComplete(MoveTemp(Response));
 		return true;
 	}
 
@@ -1011,6 +1075,84 @@ bool FWebRemoteControlModule::HandleObjectPropertyRoute(const FHttpServerRequest
 	return true;
 }
 
+bool FWebRemoteControlModule::HandleObjectPropertyAppendRoute(const FHttpServerRequest& Request, const FHttpResultCallback& OnComplete)
+{
+	TUniquePtr<FHttpServerResponse> Response = WebRemoteControlInternalUtils::CreateHttpResponse();
+	
+	FRCObjectReference ObjectRef;
+	FRCObjectRequest DeserializedRequest;
+	if (!WebRemoteControl::ResolveObjectPropertyForRequest(Request,OnComplete, Response, ObjectRef, DeserializedRequest))
+	{
+		return true;
+	}
+
+	FMemoryReader Reader = WebRemoteControl::MakePropertyMemoryReaderForArrayAddRequest(DeserializedRequest);
+	FRCJsonStructDeserializerBackend DeserializerBackend(Reader);
+
+	if (IRemoteControlModule::Get().AppendToObjectArrayProperty(ObjectRef, DeserializerBackend, ERCPayloadType::Json, DeserializedRequest.TCHARBody))
+	{
+		Response->Code = EHttpServerResponseCodes::Ok;
+	}
+
+	OnComplete(MoveTemp(Response));
+	return true;
+}
+
+bool FWebRemoteControlModule::HandleObjectPropertyInsertRoute(const FHttpServerRequest& Request, const FHttpResultCallback& OnComplete)
+{
+	TUniquePtr<FHttpServerResponse> Response = WebRemoteControlInternalUtils::CreateHttpResponse();
+
+	FRCObjectReference ObjectRef;
+	FRCObjectRequest DeserializedRequest;
+	if (!WebRemoteControl::ResolveObjectPropertyForRequest(Request, OnComplete, Response, ObjectRef, DeserializedRequest))
+	{
+		return true;
+	}
+
+	FMemoryReader Reader = WebRemoteControl::MakePropertyMemoryReaderForArrayAddRequest(DeserializedRequest);
+	FRCJsonStructDeserializerBackend DeserializerBackend(Reader);
+
+	int32 Index;
+	if (!WebRemoteControl::GetIndexQueryParamForRequest(Request, OnComplete, Response, Index))
+	{
+		return true;
+	}
+
+	if (IRemoteControlModule::Get().InsertToObjectArrayProperty(Index, ObjectRef, DeserializerBackend, ERCPayloadType::Json, DeserializedRequest.TCHARBody))
+	{
+		Response->Code = EHttpServerResponseCodes::Ok;
+	}
+
+	OnComplete(MoveTemp(Response));
+	return true;
+}
+
+bool FWebRemoteControlModule::HandleObjectPropertyRemoveRoute(const FHttpServerRequest& Request, const FHttpResultCallback& OnComplete)
+{
+	TUniquePtr<FHttpServerResponse> Response = WebRemoteControlInternalUtils::CreateHttpResponse();
+
+	FRCObjectReference ObjectRef;
+	FRCObjectRequest DeserializedRequest;
+	if (!WebRemoteControl::ResolveObjectPropertyForRequest(Request, OnComplete, Response, ObjectRef, DeserializedRequest))
+	{
+		return true;
+	}
+
+	int32 Index;
+	if (!WebRemoteControl::GetIndexQueryParamForRequest(Request, OnComplete, Response, Index))
+	{
+		return true;
+	}
+
+	if (IRemoteControlModule::Get().RemoveFromObjectArrayProperty(Index, ObjectRef))
+	{
+		Response->Code = EHttpServerResponseCodes::Ok;
+	}
+
+	OnComplete(MoveTemp(Response));
+	return true;
+}
+
 bool FWebRemoteControlModule::HandlePresetCallFunctionRoute(const FHttpServerRequest& Request, const FHttpResultCallback& OnComplete)
 {
 	TUniquePtr<FHttpServerResponse> Response = WebRemoteControlInternalUtils::CreateHttpResponse();
diff --git a/Engine/Plugins/VirtualProduction/RemoteControl/Source/WebRemoteControl/Private/WebRemoteControl.h b/Engine/Plugins/VirtualProduction/RemoteControl/Source/WebRemoteControl/Private/WebRemoteControl.h
--- a/Engine/Plugins/VirtualProduction/RemoteControl/Source/WebRemoteControl/Private/WebRemoteControl.h
+++ b/Engine/Plugins/VirtualProduction/RemoteControl/Source/WebRemoteControl/Private/WebRemoteControl.h
@@ -111,6 +111,9 @@ private:
 	bool HandleOptionsRoute(const FHttpServerRequest& Request, const FHttpResultCallback& OnComplete);
 	bool HandleObjectCallRoute(const FHttpServerRequest& Request, const FHttpResultCallback& OnComplete);
 	bool HandleObjectPropertyRoute(const FHttpServerRequest& Request, const FHttpResultCallback& OnComplete);
+	bool HandleObjectPropertyAppendRoute(const FHttpServerRequest& Request, const FHttpResultCallback& OnComplete);
+	bool HandleObjectPropertyInsertRoute(const FHttpServerRequest& Request, const FHttpResultCallback& OnComplete);
+	bool HandleObjectPropertyRemoveRoute(const FHttpServerRequest& Request, const FHttpResultCallback& OnComplete);
 	bool HandlePresetCallFunctionRoute(const FHttpServerRequest& Request, const FHttpResultCallback& OnComplete);
 	bool HandlePresetSetPropertyRoute(const FHttpServerRequest& Request, const FHttpResultCallback& OnComplete);
 	bool HandlePresetGetPropertyRoute(const FHttpServerRequest& Request, const FHttpResultCallback& OnComplete);
