Enforce passphrase for clients outside localhost

diff --git a/Engine/Plugins/VirtualProduction/RemoteControl/Source/RemoteControl/RemoteControl.Build.cs b/Engine/Plugins/VirtualProduction/RemoteControl/Source/RemoteControl/RemoteControl.Build.cs
--- a/Engine/Plugins/VirtualProduction/RemoteControl/Source/RemoteControl/RemoteControl.Build.cs
+++ b/Engine/Plugins/VirtualProduction/RemoteControl/Source/RemoteControl/RemoteControl.Build.cs
@@ -37,6 +37,7 @@ public class RemoteControl : ModuleRules
 					"AssetTools",
 					"DeveloperSettings",
 					"MessageLog",
+					"SharedSettingsWidgets",
 					"UnrealEd",
 				}
 			);
diff --git a/Engine/Plugins/VirtualProduction/RemoteControl/Source/RemoteControlCommon/Public/RemoteControlSettings.h b/Engine/Plugins/VirtualProduction/RemoteControl/Source/RemoteControlCommon/Public/RemoteControlSettings.h
--- a/Engine/Plugins/VirtualProduction/RemoteControl/Source/RemoteControlCommon/Public/RemoteControlSettings.h
+++ b/Engine/Plugins/VirtualProduction/RemoteControl/Source/RemoteControlCommon/Public/RemoteControlSettings.h
@@ -123,6 +123,10 @@ public:
 
 	UPROPERTY(Config, EditAnywhere, Category = "Remote Control Preset")
 	bool bIgnoreWarnings = false;
+
+	/** Whether to restrict access to a list of hostname/IPs in the AllowedOrigins setting. */
+	UPROPERTY(Config, EditAnywhere, Category = "Remote Control | Security")
+	bool bRestrictServerAccess = false;
 	
 	/** Whether communication with the Web Interface should only be allowed with an Passphrase */
 	UPROPERTY(Config, EditAnywhere, Category = "Remote Control | Security", DisplayName = "Use Passphrase to block Access")
@@ -135,10 +139,6 @@ public:
 	/** Enable remote python execution, enabling this could open you open to vulnerabilities if an outside actor has access to your server. */
 	UPROPERTY(Config, EditAnywhere, Category = "Remote Control | Security")
 	bool bEnableRemotePythonExecution = false;
-
-	/** Whether to restrict access to a list of hostname/IPs in the AllowedOrigins setting. */
-	UPROPERTY(Config, EditAnywhere, Category = "Remote Control | Security")
-	bool bRestrictServerAccess = false;
 	
 	/** 
 	 * Origin that can make requests to the remote control server. Should contain the hostname or IP of the server making requests to remote control. ie. "http://yourdomain.com", or "*" to allow all origins. 
@@ -156,6 +156,18 @@ public:
 	 */
 	UPROPERTY(Config, EditAnywhere, Category = "Remote Control | Security", meta = (EditCondition = bRestrictServerAccess))
 	FString AllowedIP = TEXT("127.0.0.1");
+
+	/**
+	 * Controls whether a passphrase should be required when remote control is accessed by a client outside of localhost.
+	 */
+	UPROPERTY(config, EditAnywhere, Category = "Remote Control | Security")
+	bool bEnforcePassphraseForRemoteClients = false;
+	
+	/**
+     * List of IPs that are explicitly *not* required to have a passphrase when accessing remote control.
+     */
+    UPROPERTY(config, EditAnywhere, Category = "Remote Control | Security", meta= (EditCondition = bEnforcePassphraseForRemoteClients))
+    TArray<FString> AllowedIPsForRemotePassphrases;
 	
 	UPROPERTY(config, EditAnywhere, Category = "Remote Control | Security", DisplayName = "Remote Control Passphrase")
 	TArray<FRCPassphrase> Passphrases = {};
diff --git a/Engine/Plugins/VirtualProduction/RemoteControl/Source/WebRemoteControl/Private/RemoteControlDefaultPreprocessors.h b/Engine/Plugins/VirtualProduction/RemoteControl/Source/WebRemoteControl/Private/RemoteControlDefaultPreprocessors.h
new file mode 100644
--- /dev/null
+++ b/Engine/Plugins/VirtualProduction/RemoteControl/Source/WebRemoteControl/Private/RemoteControlDefaultPreprocessors.h
@@ -0,0 +1,449 @@
+// Copyright Epic Games, Inc. All Rights Reserved.
+
+#pragma once
+
+#include "CoreMinimal.h"
+#include "HttpResultCallback.h"
+#include "HttpServerRequest.h"
+#include "Misc/WildcardString.h"
+#include "RemoteControlSettings.h"
+#include "WebRemoteControlInternalUtils.h"
+
+#if WITH_EDITOR
+#include "Dialogs/Dialogs.h"
+#include "Editor.h"
+#include "Framework/Application/SlateApplication.h"
+#include "Framework/Notifications/NotificationManager.h"
+#include "HAL/PlatformFileManager.h"
+#include "ISourceControlModule.h"
+#include "Misc/MessageDialog.h"
+#include "SSettingsEditorCheckoutNotice.h"
+#include "Widgets/Input/SButton.h"
+#include "Widgets/Input/SEditableTextBox.h"
+#include "Widgets/Notifications/SNotificationList.h"
+#include "Widgets/Text/STextBlock.h"
+#endif
+
+#define LOCTEXT_NAMESPACE "WebRemoteControl"
+
+namespace UE::WebRemoteControl
+{
+	enum class EPreprocessorResult : uint8
+	{
+		RequestPassthrough,
+		RequestHandled
+	};
+	
+	using FRCPreprocessorHandler = TFunction<EPreprocessorResult(const FHttpServerRequest& Request, const FHttpResultCallback& OnComplete)>;
+
+	/** Utility function to wrap a preprocessor handler to a http request handler than the HttpRouter can take. */
+	FHttpRequestHandler MakeHttpRequestHandler(FRCPreprocessorHandler Handler)
+	{
+		return [WrappedHandler = MoveTemp(Handler)](const FHttpServerRequest& Request, const FHttpResultCallback& OnComplete)
+		{
+			return WrappedHandler(Request, OnComplete) == EPreprocessorResult::RequestPassthrough ? false : true;
+		};
+	}
+
+#if WITH_EDITOR
+	TWeakPtr<SNotificationItem> NoPassphrasePrompt;
+	FText NoActionTakenText = LOCTEXT("NoActionTaken", "No action was taken, further requests from this IP will be denied.");
+
+	const FString& GetRemoteControlConfigPath()
+	{
+		static const FString RemoteControlConfigPath = FPackageName::LongPackageNameToFilename(GetDefault<URemoteControlSettings>()->GetPackage()->GetName(), TEXT(".ini"));
+		return RemoteControlConfigPath;
+	}
+	void SaveRemoteControlConfig()
+	{
+		FString ConfigFilename = FPaths::ConvertRelativePathToFull(URemoteControlSettings::StaticClass()->GetConfigName());
+		
+		if (!FPlatformFileManager::Get().GetPlatformFile().IsReadOnly(*ConfigFilename))
+		{
+			GetMutableDefault<URemoteControlSettings>()->SaveConfig();
+		}
+		else
+		{
+			if (FPlatformFileManager::Get().GetPlatformFile().FileExists(*ConfigFilename))
+			{
+				if (ISourceControlModule::Get().IsEnabled())
+				{
+					if (SettingsHelpers::IsSourceControlled(ConfigFilename))
+					{
+						FText DisplayMessage = LOCTEXT("SaveConfigCheckoutMessage", "The configuration file for these settings is currently not checked out. Would you like to check it out from revision control?");
+						if (FMessageDialog::Open(EAppMsgType::YesNo, DisplayMessage) == EAppReturnType::Yes)
+						{
+							constexpr bool bForceSourceControlUpdate = true;
+							constexpr bool bShowErrorInNotification = true;
+							SettingsHelpers::CheckOutOrAddFile(ConfigFilename, bForceSourceControlUpdate, bShowErrorInNotification);
+							GetMutableDefault<URemoteControlSettings>()->SaveConfig();
+							return;
+						}
+					}
+				}
+
+				FText DisplayMessage = LOCTEXT("MakeConfigWritable", "The configuration file for these settings is currently read only. Would you like to make it writable?");
+				if (FMessageDialog::Open(EAppMsgType::YesNo, DisplayMessage) == EAppReturnType::Yes)
+				{
+					if (FPlatformFileManager::Get().GetPlatformFile().SetReadOnly(*ConfigFilename, false))
+					{
+						GetMutableDefault<URemoteControlSettings>()->SaveConfig();
+					}
+					else
+					{
+						FText ErrorMessage = FText::Format(LOCTEXT("FailedToMakeWritable", "Could not make {0} writable."), FText::FromString(ConfigFilename));;
+
+						FNotificationInfo MakeWritableNotification(MoveTemp(ErrorMessage));
+						MakeWritableNotification.ExpireDuration = 3.0f;
+						FSlateNotificationManager::Get().AddNotification(MakeWritableNotification);
+					}
+				}
+			}
+		}
+	}
+
+	void AddIPToAllowlist(FString IPAddress)
+	{
+		if (TSharedPtr<SNotificationItem> Notification = NoPassphrasePrompt.Pin())
+		{
+			Notification->SetCompletionState(SNotificationItem::ECompletionState::CS_Pending);
+
+			GetMutableDefault<URemoteControlSettings>()->AllowedIPsForRemotePassphrases.Add(IPAddress);
+			SaveRemoteControlConfig();
+			
+			Notification->SetCompletionState(SNotificationItem::ECompletionState::CS_Success);
+			Notification->ExpireAndFadeout();
+		}
+	}
+
+	void CreatePassphrase()
+	{
+		if (TSharedPtr<SNotificationItem> Notification = NoPassphrasePrompt.Pin())
+		{
+			Notification->SetCompletionState(SNotificationItem::ECompletionState::CS_Pending);
+
+			TSharedPtr<SEditableTextBox> PassphraseIdentifierTextbox;
+			TSharedPtr<SEditableTextBox> PassphraseTextbox;
+
+			TSharedPtr<SWindow> Window = SNew(SWindow)
+				.Title(LOCTEXT("RemoteControlCreatePassphrase", "Create Passphrase"))
+				.SizingRule(ESizingRule::Autosized)
+				.SupportsMaximize(false)
+				.SupportsMinimize(false)
+				.FocusWhenFirstShown(true);
+
+			auto OnPassphraseEntered = [Window, Notification](const FText& PassphraseIdentifier, const FText& PassphraseText)
+			{
+				FRCPassphrase PassphraseStruct;
+				PassphraseStruct.Identifier = PassphraseIdentifier.ToString();
+				PassphraseStruct.Passphrase = FMD5::HashAnsiString(*PassphraseText.ToString());
+				GetMutableDefault<URemoteControlSettings>()->Passphrases.Add(MoveTemp(PassphraseStruct));
+				SaveRemoteControlConfig();
+
+				Notification->SetCompletionState(SNotificationItem::ECompletionState::CS_Success);
+				Notification->ExpireAndFadeout();
+
+				if (Window)
+				{
+					Window->RequestDestroyWindow();
+				}
+			};
+
+			Window->SetContent(
+				SNew(SBorder)
+				.Padding(4.f)
+				.VAlign(VAlign_Center)
+				.BorderImage(FAppStyle::GetBrush("ToolPanel.GroupBorder"))
+				[
+					SNew(SVerticalBox)
+					+ SVerticalBox::Slot()
+					.Padding(0.f, 5.f)
+					.AutoHeight()
+					[
+						SNew(STextBlock)
+						.Text(LOCTEXT("CreatePassphraseDescription", "Create a passphrase for the client to use."))
+					]
+					+ SVerticalBox::Slot()
+					.AutoHeight()
+					[
+						SAssignNew(PassphraseIdentifierTextbox, SEditableTextBox)
+						.HintText(LOCTEXT("PassphraseIdentifier", "Passphrase Identifier (Optional)"))
+					]
+					+ SVerticalBox::Slot()
+					.AutoHeight()
+					[
+						SAssignNew(PassphraseTextbox, SEditableTextBox)
+						.HintText(LOCTEXT("Passphrase", "Passphrase"))
+						.OnTextCommitted_Lambda([PassphraseIdentifierTextbox, OnPassphraseEntered](const FText& InText, ETextCommit::Type CommitType)
+						{
+							if (PassphraseIdentifierTextbox && !InText.IsEmpty() && CommitType == ETextCommit::OnEnter)
+							{
+								OnPassphraseEntered(PassphraseIdentifierTextbox->GetText(), InText);
+							}
+						})
+					]
+					+ SVerticalBox::Slot()
+					.HAlign(HAlign_Right)
+					.AutoHeight()
+					.Padding(0.0f, 2.0f, 0.0f, 0.0f)
+					[
+						SNew(SButton)
+						.Text(LOCTEXT("OK", "OK"))
+						.OnClicked_Lambda([PassphraseTextbox, PassphraseIdentifierTextbox, OnPassphraseEntered]()
+						{
+							if (PassphraseIdentifierTextbox && PassphraseTextbox && !PassphraseTextbox->GetText().IsEmpty())
+							{
+								OnPassphraseEntered(PassphraseIdentifierTextbox->GetText(), PassphraseTextbox->GetText());
+							}
+
+							return FReply::Handled();
+						})
+					]
+				]
+			);
+
+			Window->GetOnWindowClosedEvent().AddLambda([Notification](const TSharedRef<SWindow>&)
+			{
+				if (Notification->GetCompletionState() == SNotificationItem::ECompletionState::CS_Pending)
+				{
+					Notification->SetText(NoActionTakenText);
+					Notification->SetSubText(FText::GetEmpty());
+					Notification->SetCompletionState(SNotificationItem::ECompletionState::CS_Fail);
+					Notification->ExpireAndFadeout();
+				}
+			});
+
+			Window->GetOnWindowActivatedEvent().AddLambda([PassphraseTextbox]()
+			{
+				if (PassphraseTextbox)
+				{
+					constexpr uint32 UserIndex = 0;
+					FSlateApplication::Get().SetUserFocus(UserIndex, PassphraseTextbox, EFocusCause::SetDirectly);
+				}
+			});
+
+			GEditor->EditorAddModalWindow(Window.ToSharedRef());
+		}
+	}
+
+	void DisableRemotePassphrases()
+	{
+		if (TSharedPtr<SNotificationItem> Notification = NoPassphrasePrompt.Pin())
+		{
+			Notification->SetCompletionState(SNotificationItem::ECompletionState::CS_Pending);
+			
+			if (FMessageDialog::Open(EAppMsgType::YesNo, LOCTEXT("DisableRemotePassphrases", "Are you sure you want to disable passphrase requirement for remote clients? This will allow anyone on your network to access the remote control servers and could open you up to vulnerabilities.")) == EAppReturnType::Yes)
+			{
+				GetMutableDefault<URemoteControlSettings>()->bEnforcePassphraseForRemoteClients = false;
+				SaveRemoteControlConfig();
+				
+				Notification->SetCompletionState(SNotificationItem::ECompletionState::CS_Success);
+                Notification->ExpireAndFadeout();
+			}
+			else
+			{
+				Notification->SetText(NoActionTakenText);
+				Notification->SetSubText(FText::GetEmpty());
+				Notification->SetCompletionState(SNotificationItem::ECompletionState::CS_Fail);
+				Notification->ExpireAndFadeout();
+			}
+		}
+	}
+#endif
+
+	// Notifies the editor if a client tries to access the RC server without a passphrase.
+	EPreprocessorResult RemotePassphraseEnforcementPreprocessor(const FHttpServerRequest& Request, const FHttpResultCallback& OnComplete)
+	{
+		if (GetDefault<URemoteControlSettings>()->bRestrictServerAccess)
+		{
+			if (Request.PeerAddress)
+			{
+				constexpr bool bAppendPort = false;
+				const FString PeerAddress = Request.PeerAddress->ToString(bAppendPort);
+
+				if (!GetDefault<URemoteControlSettings>()->bEnforcePassphraseForRemoteClients)
+				{
+					return EPreprocessorResult::RequestPassthrough;
+				}
+
+				if (PeerAddress.Contains(TEXT("127.0.0.1")) || PeerAddress.Contains(TEXT("localhost")))
+				{
+					if (const TArray<FString>* ForwardedIP = Request.Headers.Find(WebRemoteControlInternalUtils::ForwardedIPHeader))
+					{
+						if (ForwardedIP->Num())
+						{
+							// We will need to rework this when the IP range change is in.
+							if (GetDefault<URemoteControlSettings>()->AllowedIPsForRemotePassphrases.Contains(ForwardedIP->Last()))
+							{
+								return EPreprocessorResult::RequestPassthrough;
+							}
+						}
+						else
+						{
+							return EPreprocessorResult::RequestPassthrough;
+						}
+					}
+					else
+					{
+						return EPreprocessorResult::RequestPassthrough;
+					}
+				}
+
+				if (GetDefault<URemoteControlSettings>()->AllowedIPsForRemotePassphrases.Contains(PeerAddress))
+				{
+					return EPreprocessorResult::RequestPassthrough;
+				}
+
+				if (GetDefault<URemoteControlSettings>()->Passphrases.Num())
+				{
+					const TArray<FString>* PassphraseHeader = Request.Headers.Find(WebRemoteControlInternalUtils::PassphraseHeader);
+					if (!PassphraseHeader || PassphraseHeader->Num() == 0)
+					{
+						TUniquePtr<FHttpServerResponse> Response = WebRemoteControlInternalUtils::CreateHttpResponse();
+
+						const FString ErrorMessage = TEXT("Remote passphrase enforcement is enabled but no passphrase was specified!");
+						UE_LOG(LogRemoteControl, Error, TEXT("%s"), *ErrorMessage);
+						IRemoteControlModule::BroadcastError(ErrorMessage);
+						WebRemoteControlInternalUtils::CreateUTF8ErrorMessage(ErrorMessage, Response->Body);
+
+						Response->Code = EHttpServerResponseCodes::Denied;
+						OnComplete(MoveTemp(Response));
+						return EPreprocessorResult::RequestHandled;
+					}
+
+					if (!WebRemoteControlInternalUtils::CheckPassphrase(PassphraseHeader->Last()))
+					{
+						OnComplete(WebRemoteControlInternalUtils::CreatedInvalidPassphraseResponse());
+						return EPreprocessorResult::RequestHandled;
+					}
+					else
+					{
+						return EPreprocessorResult::RequestPassthrough;
+					}
+				}
+				else
+				{
+					/* Prompt user, checkout settings if needed,
+					 * Either
+					 * 1: Explicitely add peer IP to Allowlist
+					 * 2: Create a passphrase, tell user he must enter that on the app he's using or put it in the header
+					 * 3: Disable remote passphrase enforcement. (Warn that this is dangerous)
+					 */
+	#if WITH_EDITOR
+					if (GEditor)
+					{
+						if (!NoPassphrasePrompt.IsValid() || NoPassphrasePrompt.Pin()->GetCompletionState() == SNotificationItem::ECompletionState::CS_None)
+						{
+							FNotificationInfo Info(LOCTEXT("NoPassphraseNotificationHeader", "Remote control request denied!"));
+							Info.SubText = FText::Format(LOCTEXT("NoPassphraseNotificationSubtext", "A remote control request was made by an external client ({0}) without a passphrase."), FText::FromString(PeerAddress));
+							Info.bFireAndForget = false;
+							Info.FadeInDuration = 0.5f;
+							Info.ButtonDetails.Add(FNotificationButtonInfo(LOCTEXT("AllowListIP", "AllowList IP"), LOCTEXT("AllowListTooltip", "Add this IP to the list of allowed IPs that can make remote control requests without a passphrase."), FSimpleDelegate::CreateStatic(&AddIPToAllowlist, PeerAddress), SNotificationItem::CS_None));
+							Info.ButtonDetails.Add(FNotificationButtonInfo(LOCTEXT("CreatePassphrase", "Create Passphrase"), LOCTEXT("CreatePassPhraseTooltip", "Create a passphrase for this client to use."), FSimpleDelegate::CreateStatic(&CreatePassphrase), SNotificationItem::CS_None));
+							Info.ButtonDetails.Add(FNotificationButtonInfo(LOCTEXT("DisableRemotePassphrase", "Disable Passphrases"), LOCTEXT("DisableRemotePassphrasesTooltip", "Disable the requirement for remote control requests coming from external clients to have a passphrase.\nWarning: This should only be done as a last resort since all clients on the network will be able to access your servers."), FSimpleDelegate::CreateStatic(&DisableRemotePassphrases), SNotificationItem::CS_None));
+							Info.WidthOverride = 450.f;
+							NoPassphrasePrompt = FSlateNotificationManager::Get().AddNotification(MoveTemp(Info));	
+						}
+					}
+	#endif
+					// In -game or packaged, we won't prompt the user so just deny the request.
+					const FString ErrorMessage = TEXT("A passphrase for remote control is required but we can't create one in -game or packaged.");
+					UE_LOG(LogRemoteControl, Error, TEXT("%s"), *ErrorMessage);
+					IRemoteControlModule::BroadcastError(ErrorMessage);
+
+					OnComplete(WebRemoteControlInternalUtils::CreatedInvalidPassphraseResponse());
+					return EPreprocessorResult::RequestHandled;
+				}
+			}
+		}
+
+		return EPreprocessorResult::RequestPassthrough;
+	}
+	
+	EPreprocessorResult PassphrasePreprocessor(const FHttpServerRequest& Request, const FHttpResultCallback& OnComplete)
+	{
+		TUniquePtr<FHttpServerResponse> Response = WebRemoteControlInternalUtils::CreateHttpResponse();
+
+		TArray<FString> ValueArray = {};
+		if (Request.Headers.Find(WebRemoteControlInternalUtils::PassphraseHeader))
+		{
+			ValueArray = Request.Headers[WebRemoteControlInternalUtils::PassphraseHeader];
+		}
+		
+		const FString Passphrase = !ValueArray.IsEmpty() ? ValueArray.Last() : "";
+
+		if (!WebRemoteControlInternalUtils::CheckPassphrase(Passphrase))
+		{
+			OnComplete(WebRemoteControlInternalUtils::CreatedInvalidPassphraseResponse());
+			return EPreprocessorResult::RequestHandled;
+		}
+
+		return EPreprocessorResult::RequestPassthrough;
+	}
+
+	EPreprocessorResult IPValidationPreprocessor(const FHttpServerRequest& Request, const FHttpResultCallback& OnComplete)
+	{
+		if (GetDefault<URemoteControlSettings>()->bRestrictServerAccess)
+		{
+			TUniquePtr<FHttpServerResponse> Response = WebRemoteControlInternalUtils::CreateHttpResponse();
+
+			FString OriginHeader;
+			if (const TArray<FString>* OriginHeaders = Request.Headers.Find(WebRemoteControlInternalUtils::OriginHeader))
+			{
+				if (OriginHeaders->Num())
+				{
+					OriginHeader = (*OriginHeaders)[0];
+				}
+			}
+
+			OriginHeader.RemoveSpacesInline();
+			OriginHeader.TrimStartAndEndInline();
+		
+			auto SimplifyAddress = [] (FString Address)
+			{
+				Address.RemoveFromStart(TEXT("https://www."));
+				Address.RemoveFromStart(TEXT("http://www."));
+				Address.RemoveFromStart(TEXT("https://"));
+				Address.RemoveFromStart(TEXT("http://"));
+				Address.RemoveFromEnd(TEXT("/"));
+				return Address;
+			};
+
+			const FString SimplifiedOrigin = SimplifyAddress(OriginHeader);
+			const FWildcardString SimplifiedAllowedOrigin = SimplifyAddress(GetDefault<URemoteControlSettings>()->AllowedOrigin);
+			if (!SimplifiedOrigin.IsEmpty() && GetDefault<URemoteControlSettings>()->AllowedOrigin != TEXT("*"))
+			{
+				if (!SimplifiedAllowedOrigin.IsMatch(SimplifiedOrigin))
+				{
+					WebRemoteControlInternalUtils::CreateUTF8ErrorMessage(FString::Printf(TEXT("Client origin %s does not respect the allowed origin set in Remote Control Settings."), *OriginHeader), Response->Body);
+					Response->Code = EHttpServerResponseCodes::Denied;
+					OnComplete(MoveTemp(Response));
+					return EPreprocessorResult::RequestHandled;
+				}
+			}
+
+			if (Request.PeerAddress)
+			{
+				constexpr bool bAppendPort = false;
+				FString ClientIP = Request.PeerAddress->ToString(bAppendPort);
+				const FWildcardString WildcardAllowedIP = SimplifyAddress(GetDefault<URemoteControlSettings>()->AllowedIP);
+				
+				// Allow requests from localhost
+				if (ClientIP != TEXT("localhost") && ClientIP != TEXT("127.0.0.1"))
+				{
+					if (!WildcardAllowedIP.IsEmpty() && WildcardAllowedIP.IsMatch(ClientIP))
+					{
+						WebRemoteControlInternalUtils::CreateUTF8ErrorMessage(FString::Printf(TEXT("Client IP %s does not respect the allowed IP set in Remote Control Settings."), *ClientIP), Response->Body);
+						Response->Code = EHttpServerResponseCodes::Denied;
+						OnComplete(MoveTemp(Response));
+						return EPreprocessorResult::RequestHandled;
+					}
+				}
+			}
+		}
+
+		return EPreprocessorResult::RequestPassthrough;
+	}
+}
+
+#undef LOCTEXT_NAMESPACE /* WebRemoteControl */
\ No newline at end of file
diff --git a/Engine/Plugins/VirtualProduction/RemoteControl/Source/WebRemoteControl/Private/RemoteControlWebSocketServer.cpp b/Engine/Plugins/VirtualProduction/RemoteControl/Source/WebRemoteControl/Private/RemoteControlWebSocketServer.cpp
--- a/Engine/Plugins/VirtualProduction/RemoteControl/Source/WebRemoteControl/Private/RemoteControlWebSocketServer.cpp
+++ b/Engine/Plugins/VirtualProduction/RemoteControl/Source/WebRemoteControl/Private/RemoteControlWebSocketServer.cpp
@@ -263,8 +263,8 @@ EWebsocketConnectionFilterResult FRCWebSocketServer::FilterConnection(FString Or
 	{
 		OriginHeader.RemoveSpacesInline();
 		OriginHeader.TrimStartAndEndInline();
-		
-		auto SimplifyAddress = [] (FString Address)
+
+		auto SimplifyAddress = [](FString Address)
 		{
 			Address.RemoveFromStart(TEXT("https://www."));
 			Address.RemoveFromStart(TEXT("http://www."));
@@ -275,22 +275,23 @@ EWebsocketConnectionFilterResult FRCWebSocketServer::FilterConnection(FString Or
 		};
 
 		const FString SimplifiedOrigin = SimplifyAddress(OriginHeader);
-		if (Settings->AllowedOrigin != TEXT("*"))
+		const FWildcardString SimplifiedAllowedOrigin = SimplifyAddress(GetDefault<URemoteControlSettings>()->AllowedOrigin);
+
+		if (!SimplifiedOrigin.IsEmpty() && GetDefault<URemoteControlSettings>()->AllowedOrigin != TEXT("*"))
 		{
-			const FWildcardString SimplifiedAllowedOrigin = SimplifyAddress(Settings->AllowedOrigin);
 			if (!SimplifiedAllowedOrigin.IsMatch(SimplifiedOrigin))
 			{
 				return EWebsocketConnectionFilterResult::ConnectionRefused;
 			}
-		}
 
-		// Allow requests from localhost
-		if (ClientIP != TEXT("localhost") && ClientIP != TEXT("127.0.0.1"))
-		{
-			const FWildcardString WildcardAllowedIP = SimplifyAddress(Settings->AllowedIP);
-			if (!WildcardAllowedIP.IsEmpty() && WildcardAllowedIP.IsMatch(ClientIP))
+			// Allow requests from localhost
+			if (ClientIP != TEXT("localhost") && ClientIP != TEXT("127.0.0.1"))
 			{
-				return EWebsocketConnectionFilterResult::ConnectionRefused;
+				const FWildcardString WildcardAllowedIP = SimplifyAddress(Settings->AllowedIP);
+				if (!WildcardAllowedIP.IsEmpty() && WildcardAllowedIP.IsMatch(ClientIP))
+				{
+					return EWebsocketConnectionFilterResult::ConnectionRefused;
+				}
 			}
 		}
 	}
diff --git a/Engine/Plugins/VirtualProduction/RemoteControl/Source/WebRemoteControl/Private/WebRemoteControl.cpp b/Engine/Plugins/VirtualProduction/RemoteControl/Source/WebRemoteControl/Private/WebRemoteControl.cpp
--- a/Engine/Plugins/VirtualProduction/RemoteControl/Source/WebRemoteControl/Private/WebRemoteControl.cpp
+++ b/Engine/Plugins/VirtualProduction/RemoteControl/Source/WebRemoteControl/Private/WebRemoteControl.cpp
@@ -8,6 +8,7 @@
 #include "IRemoteControlModule.h"
 #include "RCVirtualPropertyContainer.h"
 #include "RCVirtualProperty.h"
+#include "RemoteControlDefaultPreprocessors.h"
 #include "RemoteControlReflectionUtils.h"
 #include "RemoteControlRoute.h"
 #include "RemoteControlSettings.h"
@@ -361,7 +362,7 @@ void FWebRemoteControlModule::StartupModule()
 		const TArray<FString>* InPassphrase = Message.Header.Find(WebRemoteControlInternalUtils::PassphraseHeader);
 		const FString Passphrase = InPassphrase ? InPassphrase->Last() : FString("");
 		
-		const bool bCanBeDispatched = CheckPassphrase(Passphrase);
+		const bool bCanBeDispatched = WebRemoteControlInternalUtils::CheckPassphrase(Passphrase);
 
 		if (!bCanBeDispatched)
 		{
@@ -371,7 +372,7 @@ void FWebRemoteControlModule::StartupModule()
 			Wrapper.Passphrase = Passphrase;
 			Wrapper.Verb = "401";
 			
-			WebRemoteControlInternalUtils::CreateUTF8ErrorMessage(FString::Printf(TEXT("Given Passphrase is not correct!")), Wrapper.TCHARBody);
+			WebRemoteControlInternalUtils::CreateUTF8ErrorMessage(WebRemoteControlInternalUtils::InvalidPassphraseError, Wrapper.TCHARBody);
 			WebRemoteControlUtils::SerializeMessage(Wrapper, Response);
 
 			WebSocketServer.Send(Message.ClientId, MoveTemp(Response));
@@ -1844,7 +1845,7 @@ bool FWebRemoteControlModule::HandlePassphraseRoute(const FHttpServerRequest& Re
 		Passphrase = PassphraseHeader->Last();
 	}
 	
-	if (bool bIsCorrect = CheckPassphrase(Passphrase))
+	if (bool bIsCorrect = WebRemoteControlInternalUtils::CheckPassphrase(Passphrase))
 	{
 		WebRemoteControlUtils::SerializeMessage(FCheckPassphraseResponse{ bIsCorrect}, Response->Body);
 		Response->Code = EHttpServerResponseCodes::Ok;
@@ -2129,34 +2130,6 @@ bool FWebRemoteControlModule::HandleEntityMetadataOperationsRoute(const FHttpSer
 	return true;
 }
 
-bool FWebRemoteControlModule::CheckPassphrase(const FString& HashedPassphrase) const
-{
-	bool bOutResult = !(GetMutableDefault<URemoteControlSettings>()->bUseRemoteControlPassphrase);
-
-	if (bOutResult)
-	{
-		return true;
-	}
-
-	TArray<FString> HashedPassphrases = GetMutableDefault<URemoteControlSettings>()->GetHashedPassphrases();
-	if (HashedPassphrases.IsEmpty())
-	{
-		return true;
-	}
-	
-	for (const FString& InPassphrase : HashedPassphrases)
-	{
-		bOutResult = bOutResult || InPassphrase == HashedPassphrase;
-
-		if (bOutResult)
-		{
-			break;
-		}
-	}
-	
-	return bOutResult;
-}
-
 bool FWebRemoteControlModule::HandleEntitySetLabelRoute(const FHttpServerRequest& Request, const FHttpResultCallback& OnComplete)
 {
 	TUniquePtr<FHttpServerResponse> Response = WebRemoteControlInternalUtils::CreateHttpResponse();
@@ -2481,96 +2454,15 @@ void FWebRemoteControlModule::InvokeWrappedRequest(const FRCRequestWrapper& Wrap
 
 void FWebRemoteControlModule::RegisterDefaultPreprocessors()
 {
-	const FHttpRequestHandler PassphraseRequestHandler = FHttpRequestHandler([this](const FHttpServerRequest& Request, const FHttpResultCallback& OnComplete)
-		{
-			TUniquePtr<FHttpServerResponse> Response = WebRemoteControlInternalUtils::CreateHttpResponse();
-
-			TArray<FString> ValueArray = {};
-			if (Request.Headers.Find(WebRemoteControlInternalUtils::PassphraseHeader))
-			{
-				ValueArray = Request.Headers[WebRemoteControlInternalUtils::PassphraseHeader];
-			}
-			
-			const FString Passphrase = !ValueArray.IsEmpty() ? ValueArray.Last() : "";
-
-			if (!CheckPassphrase(Passphrase))
-			{
-				WebRemoteControlInternalUtils::CreateUTF8ErrorMessage(FString::Printf(TEXT("Given Passphrase is not correct!")), Response->Body);
-				Response->Code = EHttpServerResponseCodes::Denied;
-				OnComplete(MoveTemp(Response));
-				return true;
-			}
-
-			return false;
-		});
-
-	const FHttpRequestHandler IPValidationRequestHandler = FHttpRequestHandler([this](const FHttpServerRequest& Request, const FHttpResultCallback& OnComplete)
+	using namespace UE::WebRemoteControl;
+	auto RegisterInternalPreprocessor = [this](FRCPreprocessorHandler PreprocessorHandler)
 	{
-		const URemoteControlSettings* Settings = GetDefault<URemoteControlSettings>();
-		if (Settings->bRestrictServerAccess)
-		{
-			TUniquePtr<FHttpServerResponse> Response = WebRemoteControlInternalUtils::CreateHttpResponse();
-
-			FString OriginHeader;
-			if (const TArray<FString>* OriginHeaders = Request.Headers.Find(WebRemoteControlInternalUtils::OriginHeader))
-			{
-				if (OriginHeaders->Num())
-				{
-					OriginHeader = (*OriginHeaders)[0];
-				}
-			}
-
-			OriginHeader.RemoveSpacesInline();
-			OriginHeader.TrimStartAndEndInline();
-			
-			auto SimplifyAddress = [] (FString Address)
-			{
-				Address.RemoveFromStart(TEXT("https://www."));
-				Address.RemoveFromStart(TEXT("http://www."));
-				Address.RemoveFromStart(TEXT("https://"));
-				Address.RemoveFromStart(TEXT("http://"));
-				Address.RemoveFromEnd(TEXT("/"));
-				return Address;
-			};
-
-			const FString SimplifiedOrigin = SimplifyAddress(OriginHeader);
-			if (Settings->AllowedOrigin != TEXT("*"))
-			{
-				const FWildcardString SimplifiedAllowedOrigin = SimplifyAddress(Settings->AllowedOrigin);
-				if (!SimplifiedAllowedOrigin.IsMatch(SimplifiedOrigin))
-				{
-					WebRemoteControlInternalUtils::CreateUTF8ErrorMessage(FString::Printf(TEXT("Client origin %s does not respect the allowed origin set in Remote Control Settings."), *OriginHeader), Response->Body);
-					Response->Code = EHttpServerResponseCodes::Denied;
-					OnComplete(MoveTemp(Response));
-					return true;
-				}
-			}
-
-			if (Request.PeerAddress)
-			{
-				constexpr bool bAppendPort = false;
-				FString ClientIP = Request.PeerAddress->ToString(bAppendPort);
-
-				// Allow requests from localhost
-				if (ClientIP != TEXT("localhost") && ClientIP != TEXT("127.0.0.1"))
-				{
-					const FWildcardString WildcardAllowedIP = SimplifyAddress(Settings->AllowedIP);
-					if (!WildcardAllowedIP.IsEmpty() && WildcardAllowedIP.IsMatch(ClientIP))
-					{
-						WebRemoteControlInternalUtils::CreateUTF8ErrorMessage(FString::Printf(TEXT("Client IP %s does not respect the allowed IP set in Remote Control Settings."), *ClientIP), Response->Body);
-						Response->Code = EHttpServerResponseCodes::Denied;
-						OnComplete(MoveTemp(Response));
-						return true;
-					}
-				}
-			}
-		}
-
-		return false;
-	});
+		AllRegisteredPreprocessorHandlers.Add(HttpRouter->RegisterRequestPreprocessor(MakeHttpRequestHandler(MoveTemp(PreprocessorHandler))));
+	};
 
-	AllRegisteredPreprocessorHandlers.Add(HttpRouter->RegisterRequestPreprocessor(PassphraseRequestHandler));
-	AllRegisteredPreprocessorHandlers.Add(HttpRouter->RegisterRequestPreprocessor(IPValidationRequestHandler));
+	RegisterInternalPreprocessor(&RemotePassphraseEnforcementPreprocessor);
+	RegisterInternalPreprocessor(&PassphrasePreprocessor);
+	RegisterInternalPreprocessor(&IPValidationPreprocessor);
 }
 
 void FWebRemoteControlModule::UnregisterAllPreprocessors()
diff --git a/Engine/Plugins/VirtualProduction/RemoteControl/Source/WebRemoteControl/Private/WebRemoteControl.h b/Engine/Plugins/VirtualProduction/RemoteControl/Source/WebRemoteControl/Private/WebRemoteControl.h
--- a/Engine/Plugins/VirtualProduction/RemoteControl/Source/WebRemoteControl/Private/WebRemoteControl.h
+++ b/Engine/Plugins/VirtualProduction/RemoteControl/Source/WebRemoteControl/Private/WebRemoteControl.h
@@ -102,9 +102,6 @@ private:
 	/** Unregister console commands. */
 	void UnregisterConsoleCommands();
 
-	/** Checking ApiKey using Md5. */
-	bool CheckPassphrase(const FString& HashedPassphrase) const;
-
 	//~ Route handlers
 	bool HandleInfoRoute(const FHttpServerRequest& Request, const FHttpResultCallback& OnComplete);
 	bool HandleBatchRequest(const FHttpServerRequest& Request, const FHttpResultCallback& OnComplete);
diff --git a/Engine/Plugins/VirtualProduction/RemoteControl/Source/WebRemoteControl/Private/WebRemoteControlInternalUtils.cpp b/Engine/Plugins/VirtualProduction/RemoteControl/Source/WebRemoteControl/Private/WebRemoteControlInternalUtils.cpp
--- a/Engine/Plugins/VirtualProduction/RemoteControl/Source/WebRemoteControl/Private/WebRemoteControlInternalUtils.cpp
+++ b/Engine/Plugins/VirtualProduction/RemoteControl/Source/WebRemoteControl/Private/WebRemoteControlInternalUtils.cpp
@@ -261,6 +261,13 @@ TUniquePtr<FHttpServerResponse> WebRemoteControlInternalUtils::CreateHttpRespons
 	return Response;
 }
 
+TUniquePtr<FHttpServerResponse> WebRemoteControlInternalUtils::CreatedInvalidPassphraseResponse()
+{
+	TUniquePtr<FHttpServerResponse> Response = CreateHttpResponse(EHttpServerResponseCodes::Denied);
+	CreateUTF8ErrorMessage(InvalidPassphraseError, Response->Body);
+	return Response;
+}
+
 void WebRemoteControlInternalUtils::CreateUTF8ErrorMessage(const FString& InMessage, TArray<uint8>& OutUTF8Message)
 {
 	WebRemoteControlUtils::ConvertToUTF8(FString::Printf(TEXT("{ \"errorMessage\": \"%s\" }"), *InMessage), OutUTF8Message);
@@ -593,3 +600,31 @@ bool WebRemoteControlInternalUtils::IsRequestContentType(const FHttpServerReques
 	
 	return false;
 }
+
+bool WebRemoteControlInternalUtils::CheckPassphrase(const FString& HashedPassphrase)
+{
+	bool bOutResult = !(GetDefault<URemoteControlSettings>()->bUseRemoteControlPassphrase);
+
+	if (bOutResult)
+	{
+		return true;
+	}
+
+	TArray<FString> HashedPassphrases = GetDefault<URemoteControlSettings>()->GetHashedPassphrases();
+	if (HashedPassphrases.IsEmpty())
+	{
+		return true;
+	}
+
+	for (const FString& InPassphrase : HashedPassphrases)
+	{
+		bOutResult = bOutResult || InPassphrase == HashedPassphrase;
+
+		if (bOutResult)
+		{
+			break;
+		}
+	}
+
+	return bOutResult;
+}
diff --git a/Engine/Plugins/VirtualProduction/RemoteControl/Source/WebRemoteControl/Private/WebRemoteControlInternalUtils.h b/Engine/Plugins/VirtualProduction/RemoteControl/Source/WebRemoteControl/Private/WebRemoteControlInternalUtils.h
--- a/Engine/Plugins/VirtualProduction/RemoteControl/Source/WebRemoteControl/Private/WebRemoteControlInternalUtils.h
+++ b/Engine/Plugins/VirtualProduction/RemoteControl/Source/WebRemoteControl/Private/WebRemoteControlInternalUtils.h
@@ -51,6 +51,8 @@ namespace WebRemoteControlInternalUtils
 	static const TCHAR* WrappedRequestHeader = TEXT("UE-Wrapped-Request");
 	static const FString PassphraseHeader = TEXT("Passphrase");
 	static const TCHAR* OriginHeader = TEXT("Origin");
+	static const TCHAR* ForwardedIPHeader = TEXT("x-forwarded-for");
+	static const TCHAR* InvalidPassphraseError = TEXT("Given Passphrase is not correct!");
 
 	/**
 	 * Construct a default http response with CORS headers.
@@ -59,6 +61,11 @@ namespace WebRemoteControlInternalUtils
 	 */
 	TUniquePtr<FHttpServerResponse> CreateHttpResponse(EHttpServerResponseCodes InResponseCode = EHttpServerResponseCodes::BadRequest);
 
+	/**
+	 * Create a http response for a request denied because of an invalid passphrase.
+	 */
+	TUniquePtr<FHttpServerResponse> CreatedInvalidPassphraseResponse();
+
 	/**
 	 * Create a json structure containing an error message.
 	 * @param InMessage The desired error message.
@@ -447,4 +454,7 @@ namespace WebRemoteControlInternalUtils
 
 		return bSuccess;
 	}
+
+	/** Checking ApiKey using Md5. */
+	bool CheckPassphrase(const FString& HashedPassphrase);
 }
diff --git a/Engine/Plugins/VirtualProduction/RemoteControl/Source/WebRemoteControl/WebRemoteControl.Build.cs b/Engine/Plugins/VirtualProduction/RemoteControl/Source/WebRemoteControl/WebRemoteControl.Build.cs
--- a/Engine/Plugins/VirtualProduction/RemoteControl/Source/WebRemoteControl/WebRemoteControl.Build.cs
+++ b/Engine/Plugins/VirtualProduction/RemoteControl/Source/WebRemoteControl/WebRemoteControl.Build.cs
@@ -36,8 +36,10 @@ public class WebRemoteControl : ModuleRules
 					"ImageWrapper",
 					"RemoteControlUI",
 					"Settings",
+					"SharedSettingsWidgets",
 					"Slate",
 					"SlateCore",
+					"SourceControl",
 					"UnrealEd",
 				}
 			);
