Remote API Protocols - Restarting the editor resets masking Transform

diff --git a/Engine/Plugins/VirtualProduction/RemoteControl/Source/RemoteControl/Private/RemoteControlField.cpp b/Engine/Plugins/VirtualProduction/RemoteControl/Source/RemoteControl/Private/RemoteControlField.cpp
--- a/Engine/Plugins/VirtualProduction/RemoteControl/Source/RemoteControl/Private/RemoteControlField.cpp
+++ b/Engine/Plugins/VirtualProduction/RemoteControl/Source/RemoteControl/Private/RemoteControlField.cpp
@@ -148,17 +148,17 @@ bool FRemoteControlField::CanBindObject(const UObject* InObjectToBind) const
 
 void FRemoteControlField::ClearMask(ERCMask InMaskBit)
 {
-	ActiveMasks &= ~InMaskBit;
+	ActiveMasks &= ~(uint8)InMaskBit;
 }
 
 void FRemoteControlField::EnableMask(ERCMask InMaskBit)
 {
-	ActiveMasks |= InMaskBit;
+	ActiveMasks |= (uint8)InMaskBit;
 }
 
 bool FRemoteControlField::HasMask(ERCMask InMaskBit) const
 {
-	return (ActiveMasks & InMaskBit) != ERCMask::NoMask;
+	return (ActiveMasks & (uint8)InMaskBit) != (uint8)ERCMask::NoMask;
 }
 
 void FRemoteControlField::PostSerialize(const FArchive& Ar)
diff --git a/Engine/Plugins/VirtualProduction/RemoteControl/Source/RemoteControl/Public/RemoteControlField.h b/Engine/Plugins/VirtualProduction/RemoteControl/Source/RemoteControl/Public/RemoteControlField.h
--- a/Engine/Plugins/VirtualProduction/RemoteControl/Source/RemoteControl/Public/RemoteControlField.h
+++ b/Engine/Plugins/VirtualProduction/RemoteControl/Source/RemoteControl/Public/RemoteControlField.h
@@ -55,7 +55,7 @@ struct REMOTECONTROL_API FRemoteControlField : public FRemoteControlEntity
 	virtual bool SupportsMasking() const { return false; }
 
 	/** Retrieves the enabled masks. */
-	ERCMask GetActiveMasks() const { return ActiveMasks; }
+	ERCMask GetActiveMasks() const { return (ERCMask)ActiveMasks; }
 
 	/**	Returns whether the field is only available in editor. */
 	bool IsEditorOnly() const { return bIsEditorOnly; }
@@ -106,7 +106,7 @@ protected:
 	
 	/** Holds the actively enabled masks. */
 	UPROPERTY()
-	ERCMask ActiveMasks = RC_AllMasks;
+	uint8 ActiveMasks = (uint8)RC_AllMasks;
 
 protected:
 	FRemoteControlField(URemoteControlPreset* InPreset, EExposedFieldType InType, FName InLabel, FRCFieldPathInfo InFieldPathInfo, const TArray<URemoteControlBinding*> InBindings);
