Remote Control Websocket server should not bind to 0.0.0.0 by default

diff --git a/Engine/Plugins/VirtualProduction/RemoteControl/Source/RemoteControlCommon/Public/RemoteControlSettings.h b/Engine/Plugins/VirtualProduction/RemoteControl/Source/RemoteControlCommon/Public/RemoteControlSettings.h
--- a/Engine/Plugins/VirtualProduction/RemoteControl/Source/RemoteControlCommon/Public/RemoteControlSettings.h
+++ b/Engine/Plugins/VirtualProduction/RemoteControl/Source/RemoteControlCommon/Public/RemoteControlSettings.h
@@ -7,6 +7,141 @@
 
 #include "RemoteControlSettings.generated.h"
 
+/**
+ * Utility struct to represent IPv4 Network addresses.
+ */
+USTRUCT(BlueprintType)
+struct FRCNetworkAddress
+{
+	GENERATED_BODY()
+
+	FRCNetworkAddress() = default;
+
+	bool operator==(const FRCNetworkAddress& OtherNetworkAddress) const
+	{
+		return ClassA == OtherNetworkAddress.ClassA
+			&& ClassB == OtherNetworkAddress.ClassB
+			&& ClassC == OtherNetworkAddress.ClassC
+			&& ClassD == OtherNetworkAddress.ClassD;
+	}
+
+	/**
+	 * Calculates the hash for a Network Address.
+	 *
+	 * @param NetworkAddress The Network address to calculate the hash for.
+	 * @return The hash.
+	 */
+	friend uint32 GetTypeHash(const FRCNetworkAddress& NetworkAddress)
+	{
+		const uint32 NetworkPartHash = HashCombine(GetTypeHash(NetworkAddress.ClassA), GetTypeHash(NetworkAddress.ClassB));
+		const uint32 HostPartHash = HashCombine(GetTypeHash(NetworkAddress.ClassC), GetTypeHash(NetworkAddress.ClassD));
+
+		return HashCombine(NetworkPartHash, HostPartHash);
+	}
+
+	/**
+	 * Retrieves the network address as string e.g. 192.168.1.1
+	 */
+	const FString ToString() const
+	{
+		return FString::Printf(TEXT("%d.%d.%d.%d"), ClassA, ClassB, ClassC, ClassD);
+	}
+
+	/** Denotes the first octet of the IPv4 address (0-255.xxx.xxx.xxx) */
+	UPROPERTY(EditAnywhere, Category="Network Address")
+		uint8 ClassA = 192;
+
+	/** Denotes the second octet of the IPv4 address (xxx.0-255.xxx.xxx) */
+	UPROPERTY(EditAnywhere, Category = "Network Address")
+		uint8 ClassB = 168;
+
+	/** Denotes the third octet of the IPv4 address (xxx.xxx.0-255.xxx) */
+	UPROPERTY(EditAnywhere, Category = "Network Address")
+		uint8 ClassC = 1;
+
+	/** Denotes the fourth octet of the IPv4 address (xxx.xxx.xxx.0-255) */
+	UPROPERTY(EditAnywhere, Category = "Network Address")
+		uint8 ClassD = 1;
+};
+
+/**
+ * Utility struct to represent range of IPv4 Network addresses.
+ */
+USTRUCT(BlueprintType)
+struct FRCNetworkAddressRange
+{
+	GENERATED_BODY()
+
+	FRCNetworkAddressRange() = default;
+
+	bool operator==(const FRCNetworkAddressRange& OtherNetworkAddressRange) const
+	{
+		return LowerBound == OtherNetworkAddressRange.LowerBound
+			&& UpperBound == OtherNetworkAddressRange.UpperBound;
+	}
+	
+	/**
+	 * Calculates the hash for a Network Address Range.
+	 *
+	 * @param NetworkAddress The Network address range to calculate the hash for.
+	 * @return The hash.
+	 */
+	friend uint32 GetTypeHash(const FRCNetworkAddressRange& NetworkAddressRange)
+	{
+		const uint32 LowerBoundHash = GetTypeHash(NetworkAddressRange.LowerBound);
+		const uint32 UpperBoundHash = GetTypeHash(NetworkAddressRange.UpperBound);
+
+		return HashCombine(LowerBoundHash, UpperBoundHash);
+	}
+
+	bool IsInRange(const FString& InNetworkAddress) const
+	{
+		TArray<FString> IndividualClasses;
+
+		if (InNetworkAddress.ParseIntoArray(IndividualClasses, TEXT(".")) == 4)
+		{
+			return IsInRange_Internal(FCString::Atoi(*IndividualClasses[0])
+				, FCString::Atoi(*IndividualClasses[1])
+				, FCString::Atoi(*IndividualClasses[2])
+				, FCString::Atoi(*IndividualClasses[3])
+			);
+		}
+
+		return false;
+	}
+	
+	bool IsInRange(const FRCNetworkAddress& InNetworkAddress) const
+	{
+		return IsInRange_Internal(InNetworkAddress.ClassA
+			, InNetworkAddress.ClassB
+			, InNetworkAddress.ClassC
+			, InNetworkAddress.ClassD
+		);
+	}
+
+private:
+
+	bool IsInRange_Internal(uint8 InClassA, uint8 InClassB, uint8 InClassC, uint8 InClassD) const
+	{
+		bool bClassAIsInRange = InClassA >= LowerBound.ClassA && InClassA <= UpperBound.ClassA;
+		bool bClassBIsInRange = InClassB >= LowerBound.ClassB && InClassB <= UpperBound.ClassB;
+		bool bClassCIsInRange = InClassC >= LowerBound.ClassC && InClassC <= UpperBound.ClassC;
+		bool bClassDIsInRange = InClassD >= LowerBound.ClassD && InClassD <= UpperBound.ClassD;
+		
+		return bClassAIsInRange && bClassBIsInRange && bClassCIsInRange && bClassDIsInRange;
+	}
+
+public:
+
+	/** Denotes the lower bound IPv4 address. */
+	UPROPERTY(EditAnywhere, Category="Network Address Range")
+		FRCNetworkAddress LowerBound = { 192, 168, 1, 1 };
+
+	/** Denotes the upper bound IPv4 address. */
+	UPROPERTY(EditAnywhere, Category = "Network Address Range")
+		FRCNetworkAddress UpperBound = { 192, 168, 255, 255 };
+};
+
 /**
  * Passphrase Struct
  */
@@ -68,6 +203,19 @@ public:
 		return OutArray;
 	}
 
+	virtual bool IsClientAllowed(const FString& InClientAddressStr) const
+	{
+		for (const FRCNetworkAddressRange& AllowlistedClient : AllowlistedClients)
+		{
+			if (AllowlistedClient.IsInRange(InClientAddressStr))
+			{
+				return true;
+			}
+		}
+
+		return false;
+	}
+
 	/**
 	 * Should transactions be generated for events received through protocols (ie. MIDI, DMX etc.)
 	 * Disabling transactions improves performance but will prevent events from being transacted to Multi-User
@@ -79,7 +227,7 @@ public:
 	/** The remote control web app http port. */
 	UPROPERTY(config, EditAnywhere, Category = "Remote Control Web Interface", DisplayName = "Remote Control Web Interface http Port")
 	uint32 RemoteControlWebInterfacePort = 30000;
-
+	
 	/** Should force a build of the WebApp at startup. */
 	UPROPERTY(config, EditAnywhere, Category = "Remote Control Web Interface", DisplayName = "Force WebApp build at startup")
 	bool bForceWebAppBuildAtStartup = false;
@@ -103,7 +251,7 @@ public:
 	/** The web remote control WebSocket server's port. */
 	UPROPERTY(config, EditAnywhere, Category = "Remote Control Web Server", DisplayName = "Remote Control WebSocket Server Port")
 	uint32 RemoteControlWebSocketServerPort = 30020;
-
+	
 	/** Show a warning icon for exposed editor-only fields. */
 	UPROPERTY(config, EditAnywhere, Category = "Remote Control Preset", DisplayName = "Show a warning when exposing editor-only entities.")
 	bool bDisplayInEditorOnlyWarnings = false;
@@ -139,7 +287,11 @@ public:
 	/** Enable remote python execution, enabling this could open you open to vulnerabilities if an outside actor has access to your server. */
 	UPROPERTY(Config, EditAnywhere, Category = "Remote Control | Security")
 	bool bEnableRemotePythonExecution = false;
-	
+
+	/** List of IP Addresses that are allowed to access the Web API. */
+	UPROPERTY(config, EditAnywhere, Category = "Remote Control | Security", DisplayName = "Range of Allowlisted Clients", Meta = (EditCondition = "bRestrictServerAccess", EditConditionHides))
+	TSet<FRCNetworkAddressRange> AllowlistedClients;
+
 	/** 
 	 * Origin that can make requests to the remote control server. Should contain the hostname or IP of the server making requests to remote control. ie. "http://yourdomain.com", or "*" to allow all origins. 
 	 * @Note: This is used to block requests coming from a browser (ie. Coming from a script running on a website), ideally you should use both this setting and AllowedIPs, as a request coming from a websocket client can have an empty Origin.
diff --git a/Engine/Plugins/VirtualProduction/RemoteControl/Source/RemoteControlUI/Private/RemoteControlUIModule.cpp b/Engine/Plugins/VirtualProduction/RemoteControl/Source/RemoteControlUI/Private/RemoteControlUIModule.cpp
--- a/Engine/Plugins/VirtualProduction/RemoteControl/Source/RemoteControlUI/Private/RemoteControlUIModule.cpp
+++ b/Engine/Plugins/VirtualProduction/RemoteControl/Source/RemoteControlUI/Private/RemoteControlUIModule.cpp
@@ -35,6 +35,7 @@
 #include "UI/Behaviour/Builtin/Conditional/SRCBehaviourConditional.h"
 #include "UI/Behaviour/SRCBehaviourPanel.h"
 #include "UI/Customizations/FPassphraseCustomization.h"
+#include "UI/Customizations/NetworkAddressCustomization.h"
 #include "UI/Customizations/RemoteControlEntityCustomization.h"
 #include "UI/RemoteControlPanelStyle.h"
 #include "UI/SRCPanelExposedActor.h"
@@ -840,6 +841,7 @@ void FRemoteControlUIModule::RegisterStructCustomizations()
 		PropertyEditorModule.RegisterCustomClassLayout(Name, FOnGetDetailCustomizationInstance::CreateStatic(&FRemoteControlEntityCustomization::MakeInstance));
 	}
 	
+	PropertyEditorModule.RegisterCustomPropertyTypeLayout(FRCNetworkAddress::StaticStruct()->GetFName(), FOnGetPropertyTypeCustomizationInstance::CreateStatic(&FNetworkAddressCustomization::MakeInstance));
 	PropertyEditorModule.RegisterCustomPropertyTypeLayout(FRCPassphrase::StaticStruct()->GetFName(), FOnGetPropertyTypeCustomizationInstance::CreateStatic(&FPassphraseCustomization::MakeInstance));
 }
 
@@ -855,6 +857,7 @@ void FRemoteControlUIModule::UnregisterStructCustomizations()
 
 	if (UObjectInitialized())
 	{
+		PropertyEditorModule.UnregisterCustomPropertyTypeLayout(FRCNetworkAddress::StaticStruct()->GetFName());
 		PropertyEditorModule.UnregisterCustomPropertyTypeLayout(FRCPassphrase::StaticStruct()->GetFName());
 	}
 }
diff --git a/Engine/Plugins/VirtualProduction/RemoteControl/Source/RemoteControlUI/Private/UI/Customizations/NetworkAddressCustomization.cpp b/Engine/Plugins/VirtualProduction/RemoteControl/Source/RemoteControlUI/Private/UI/Customizations/NetworkAddressCustomization.cpp
new file mode 100644
--- /dev/null
+++ b/Engine/Plugins/VirtualProduction/RemoteControl/Source/RemoteControlUI/Private/UI/Customizations/NetworkAddressCustomization.cpp
@@ -0,0 +1,112 @@
+// Copyright Epic Games, Inc. All Rights Reserved.
+
+#include "NetworkAddressCustomization.h"
+#include "DetailLayoutBuilder.h"
+#include "DetailWidgetRow.h"
+#include "IDetailChildrenBuilder.h"
+#include "PropertyHandle.h"
+#include "RemoteControlSettings.h"
+#include "Widgets/Layout/SBorder.h"
+#include "Widgets/Layout/SWrapBox.h"
+
+#define LOCTEXT_NAMESPACE "NetworkAddressCustomization"
+
+TSharedRef<IPropertyTypeCustomization> FNetworkAddressCustomization::MakeInstance()
+{
+	// Create the instance and returned a SharedRef
+	return MakeShareable(new FNetworkAddressCustomization());
+}
+
+void FNetworkAddressCustomization::CustomizeHeader(TSharedRef<IPropertyHandle> StructPropertyHandle, FDetailWidgetRow& HeaderRow, IPropertyTypeCustomizationUtils& StructCustomizationUtils)
+{
+	// First we need to retrieve every Property handles
+	TSharedPtr<IPropertyHandle> ClassAPropertyHandle = StructPropertyHandle->GetChildHandle(GET_MEMBER_NAME_CHECKED(FRCNetworkAddress, ClassA));
+	TSharedPtr<IPropertyHandle> ClassBPropertyHandle = StructPropertyHandle->GetChildHandle(GET_MEMBER_NAME_CHECKED(FRCNetworkAddress, ClassB));
+	TSharedPtr<IPropertyHandle> ClassCPropertyHandle = StructPropertyHandle->GetChildHandle(GET_MEMBER_NAME_CHECKED(FRCNetworkAddress, ClassC));
+	TSharedPtr<IPropertyHandle> ClassDPropertyHandle = StructPropertyHandle->GetChildHandle(GET_MEMBER_NAME_CHECKED(FRCNetworkAddress, ClassD));
+
+	// Check them before using them
+	check(ClassAPropertyHandle.IsValid()
+		&& ClassBPropertyHandle.IsValid()
+		&& ClassCPropertyHandle.IsValid()
+		&& ClassDPropertyHandle.IsValid());
+
+	constexpr bool bUseAllottedWidth = true;
+
+	const float IconSizeSquared = 4;
+
+	const TSharedRef<SWidget> DotIconWidget = SNew(SBox)
+		.HeightOverride(IconSizeSquared)
+		.WidthOverride(IconSizeSquared)
+		[
+			SNew(SImage)
+			.Image(FAppStyle::Get().GetBrush("Icons.FilledCircle"))
+		];
+
+	HeaderRow
+		.NameContent()
+		[
+			StructPropertyHandle->CreatePropertyNameWidget()
+		]
+		.ValueContent()
+		.HAlign(HAlign_Fill)
+		.VAlign(VAlign_Center)
+		[
+			SNew(SWrapBox)
+			.UseAllottedWidth(bUseAllottedWidth)
+
+			+ SWrapBox::Slot()
+			.Padding(2.f, 0.f)
+			[
+				ClassAPropertyHandle->CreatePropertyValueWidget()
+			]
+		
+			+ SWrapBox::Slot()
+			.HAlign(HAlign_Center)
+			.VAlign(VAlign_Center)
+			.Padding(2.f, 2.f)
+			[
+				DotIconWidget
+			]
+
+			+ SWrapBox::Slot()
+			.Padding(2.f, 0.f)
+			[
+				ClassBPropertyHandle->CreatePropertyValueWidget()
+			]
+			
+			+ SWrapBox::Slot()
+			.HAlign(HAlign_Center)
+			.VAlign(VAlign_Center)
+			.Padding(2.f, 2.f)
+			[
+				DotIconWidget
+			]
+
+			+ SWrapBox::Slot()
+			.Padding(2.f, 0.f)
+			[
+				ClassCPropertyHandle->CreatePropertyValueWidget()
+			]
+			
+			+ SWrapBox::Slot()
+			.HAlign(HAlign_Center)
+			.VAlign(VAlign_Center)
+			.Padding(2.f, 2.f)
+			[
+				DotIconWidget
+			]
+
+			+ SWrapBox::Slot()
+			.Padding(2.f, 0.f)
+			[
+				ClassDPropertyHandle->CreatePropertyValueWidget()
+			]
+		];
+}
+
+void FNetworkAddressCustomization::CustomizeChildren(TSharedRef<IPropertyHandle> StructPropertyHandle, IDetailChildrenBuilder& StructBuilder, IPropertyTypeCustomizationUtils& StructCustomizationUtils)
+{
+}
+
+#undef LOCTEXT_NAMESPACE
diff --git a/Engine/Plugins/VirtualProduction/RemoteControl/Source/RemoteControlUI/Private/UI/Customizations/NetworkAddressCustomization.h b/Engine/Plugins/VirtualProduction/RemoteControl/Source/RemoteControlUI/Private/UI/Customizations/NetworkAddressCustomization.h
new file mode 100644
--- /dev/null
+++ b/Engine/Plugins/VirtualProduction/RemoteControl/Source/RemoteControlUI/Private/UI/Customizations/NetworkAddressCustomization.h
@@ -0,0 +1,31 @@
+// Copyright Epic Games, Inc. All Rights Reserved.
+
+#pragma once
+
+#include "IPropertyTypeCustomization.h"
+
+class FDetailWidgetRow;
+class IDetailChildrenBuilder;
+
+/**
+ * 
+ */
+class FNetworkAddressCustomization : public IPropertyTypeCustomization
+{
+public:
+
+	/**
+	 * It is just a convenient helpers which will be used to register our customization.
+	 * When the propertyEditor module find our FRCNetworkAddress property,
+	 * it will use this static method to instanciate our customization object.
+	 */
+	static TSharedRef<IPropertyTypeCustomization> MakeInstance();
+
+	//~ BEGIN : IPropertyTypeCustomization Interface
+	
+	virtual void CustomizeHeader(TSharedRef<IPropertyHandle> StructPropertyHandle, FDetailWidgetRow& HeaderRow, IPropertyTypeCustomizationUtils& StructCustomizationUtils) override;
+	
+	virtual void CustomizeChildren(TSharedRef<IPropertyHandle> StructPropertyHandle, IDetailChildrenBuilder& StructBuilder, IPropertyTypeCustomizationUtils& StructCustomizationUtils) override;
+	
+	//~ END : IPropertyTypeCustomization Interface
+};
diff --git a/Engine/Plugins/VirtualProduction/RemoteControl/Source/WebRemoteControl/Private/RemoteControlWebSocketServer.cpp b/Engine/Plugins/VirtualProduction/RemoteControl/Source/WebRemoteControl/Private/RemoteControlWebSocketServer.cpp
--- a/Engine/Plugins/VirtualProduction/RemoteControl/Source/WebRemoteControl/Private/RemoteControlWebSocketServer.cpp
+++ b/Engine/Plugins/VirtualProduction/RemoteControl/Source/WebRemoteControl/Private/RemoteControlWebSocketServer.cpp
@@ -287,6 +287,11 @@ EWebsocketConnectionFilterResult FRCWebSocketServer::FilterConnection(FString Or
 			// Allow requests from localhost
 			if (ClientIP != TEXT("localhost") && ClientIP != TEXT("127.0.0.1"))
 			{
+				if (!Settings->IsClientAllowed(ClientIP))
+				{
+					return EWebsocketConnectionFilterResult::ConnectionRefused;
+				}
+
 				const FWildcardString WildcardAllowedIP = SimplifyAddress(Settings->AllowedIP);
 				if (!WildcardAllowedIP.IsEmpty() && WildcardAllowedIP.IsMatch(ClientIP))
 				{
