Fix network address customization CIS

diff --git a/Engine/Plugins/VirtualProduction/RemoteControl/Source/RemoteControlUI/Private/UI/Customizations/NetworkAddressCustomization.cpp b/Engine/Plugins/VirtualProduction/RemoteControl/Source/RemoteControlUI/Private/UI/Customizations/NetworkAddressCustomization.cpp
--- a/Engine/Plugins/VirtualProduction/RemoteControl/Source/RemoteControlUI/Private/UI/Customizations/NetworkAddressCustomization.cpp
+++ b/Engine/Plugins/VirtualProduction/RemoteControl/Source/RemoteControlUI/Private/UI/Customizations/NetworkAddressCustomization.cpp
@@ -6,7 +6,9 @@
 #include "IDetailChildrenBuilder.h"
 #include "PropertyHandle.h"
 #include "RemoteControlSettings.h"
+#include "Widgets/Images/SImage.h"
 #include "Widgets/Layout/SBorder.h"
+#include "Widgets/Layout/SBox.h"
 #include "Widgets/Layout/SWrapBox.h"
 
 #define LOCTEXT_NAMESPACE "NetworkAddressCustomization"
