Fix functions with out parameters not working with Remote Control Interception

When running in a display cluster context, remote control requests get intercepted by the Display Cluster interceptor, but requests going through that system are sent in a "fire and forget" fashion, therefore we cannot fetch the returned values from the intecepted payloads.

This fix works around this by disabling interception when dealing with functions calls with out parameters.

diff --git a/Engine/Plugins/VirtualProduction/RemoteControl/Source/RemoteControl/Private/RemoteControlModule.cpp b/Engine/Plugins/VirtualProduction/RemoteControl/Source/RemoteControl/Private/RemoteControlModule.cpp
--- a/Engine/Plugins/VirtualProduction/RemoteControl/Source/RemoteControl/Private/RemoteControlModule.cpp
+++ b/Engine/Plugins/VirtualProduction/RemoteControl/Source/RemoteControl/Private/RemoteControlModule.cpp
@@ -57,6 +57,8 @@ struct FRCInterceptionPayload
 static TAutoConsoleVariable<int32> CVarRemoteControlEnableOngoingChangeOptimization(TEXT("RemoteControl.EnableOngoingChangeOptimization"), 1, TEXT("Enable an optimization that keeps track of the ongoing remote control change in order to improve performance."));
 #endif
 
+#define REMOTE_CONTROL_LOG_ONCE(Verbosity, Format, ...) (LogOnce(ELogVerbosity::Verbosity, *FString::Printf(Format, ##__VA_ARGS__), FString(__FILE__), __LINE__))
+
 namespace RemoteControlUtil
 {
 	const FName NAME_DisplayName(TEXT("DisplayName"));
@@ -945,7 +947,7 @@ bool FRemoteControlModule::InvokeCall(FRCCall& InCall, ERCPayloadType InPayloadT
 		const bool bGenerateTransaction = InCall.TransactionMode == ERCTransactionMode::AUTOMATIC;
 
 		// Check the replication path before apply property values
-		if (InInterceptPayload.Num() != 0)
+		if (InInterceptPayload.Num() != 0 && CanInterceptFunction(InCall))
 		{
 			FRCIFunctionMetadata FunctionMetadata(InCall.CallRef.Object->GetPathName(), InCall.CallRef.Function->GetPathName(), bGenerateTransaction, ToExternal(InPayloadType), InInterceptPayload);
 
@@ -1589,7 +1591,7 @@ bool FRemoteControlModule::SetPresetController(const FName PresetName, const FNa
 		return false;
 	}
 
-	URCVirtualPropertyBase* VirtualProperty = Preset->GetControllerByDisplayName(ControllerName);
+	URCVirtualPropertyBase* VirtualProperty = Preset->GetController(ControllerName);
 	if (!ensure(VirtualProperty))
 	{
 		return false;
@@ -2395,6 +2397,91 @@ void FRemoteControlModule::RegisterMaskingFactories()
 	RegisterMaskingFactoryForType(TBaseStructure<FLinearColor>::Get(), FLinearColorMaskingFactory::MakeInstance());
 }
 
+bool FRemoteControlModule::CanInterceptFunction(const FRCCall& RCCall) const
+{
+	if (RCCall.IsValid())
+	{
+		for (TFieldIterator<FProperty> It(RCCall.CallRef.Function.Get()); It; ++It)
+		{
+			// At the moment interceptors do not support return values.
+			if (It->HasAnyPropertyFlags(CPF_ReturnParm | CPF_OutParm))
+			{
+				
+#if WITH_EDITOR
+				FString FunctionName = RCCall.CallRef.Function->GetDisplayNameText().ToString();
+#else
+				FString FunctionName = RCCall.CallRef.Function->GetName();
+#endif
+				REMOTE_CONTROL_LOG_ONCE(Warning, TEXT("Function \"%s\" on object \"%s\" could not be intercepted because it contains out parameters."), *FunctionName, *RCCall.CallRef.Object->GetName());
+				return false;
+			}
+		}
+	}
+
+	return true;
+}
+
+void FRemoteControlModule::LogOnce(ELogVerbosity::Type InVerbosity, const FString& LogDetails, const FString& FileName, int32 LineNumber) const
+{
+	// This code was taken from AudioMixer.h
+#if !NO_LOGGING
+	// Log once to avoid Spam.
+	static FCriticalSection Cs;
+	static TSet<uint32> LogHistory;
+
+	FScopeLock Lock(&Cs);
+	const FString MessageToHash = FString::Printf(TEXT("%s (File %s, Line %d)"), *LogDetails, *FileName, LineNumber);
+
+	uint32 Hash = GetTypeHash(MessageToHash);
+	if (!LogHistory.Contains(Hash))
+	{
+		switch (InVerbosity)
+		{
+		case ELogVerbosity::Fatal:
+			UE_LOG(LogRemoteControl, Fatal, TEXT("%s"), *LogDetails);
+			break;
+
+		case ELogVerbosity::Error:
+			UE_LOG(LogRemoteControl, Error, TEXT("%s"), *LogDetails);
+			break;
+
+		case ELogVerbosity::Warning:
+			UE_LOG(LogRemoteControl, Warning, TEXT("%s"), *LogDetails);
+			break;
+
+		case ELogVerbosity::Display:
+			UE_LOG(LogRemoteControl, Display, TEXT("%s"), *LogDetails);
+			break;
+
+		case ELogVerbosity::Log:
+			UE_LOG(LogRemoteControl, Log, TEXT("%s"), *LogDetails);
+			break;
+
+		case ELogVerbosity::Verbose:
+			UE_LOG(LogRemoteControl, Verbose, TEXT("%s"), *LogDetails);
+			break;
+
+		case ELogVerbosity::VeryVerbose:
+			UE_LOG(LogRemoteControl, VeryVerbose, TEXT("%s"), *LogDetails);
+			break;
+
+		default:
+			UE_LOG(LogRemoteControl, Error, TEXT("%s"), *LogDetails);
+			{
+				static_assert(static_cast<uint8>(ELogVerbosity::NumVerbosity) == 8, "Missing ELogVerbosity case coverage");
+			}
+			break;
+		}
+
+		LogHistory.Add(Hash);
+		if (InVerbosity == ELogVerbosity::Fatal || InVerbosity == ELogVerbosity::Error || InVerbosity == ELogVerbosity::Warning)
+		{
+			BroadcastError(LogDetails, InVerbosity);
+		}
+	}
+#endif
+}
+
 bool FRemoteControlModule::AddToArrayProperty(const FRCObjectReference& ObjectAccess, IStructDeserializerBackend& Backend, TFunctionRef<int32(FScriptArrayHelper&)> ModifyFunction)
 {
 	return ModifyArrayProperty(ObjectAccess, [&ObjectAccess, &Backend, &ModifyFunction](FScriptArrayHelper& ArrayHelper) {
diff --git a/Engine/Plugins/VirtualProduction/RemoteControl/Source/RemoteControl/Private/RemoteControlModule.h b/Engine/Plugins/VirtualProduction/RemoteControl/Source/RemoteControl/Private/RemoteControlModule.h
--- a/Engine/Plugins/VirtualProduction/RemoteControl/Source/RemoteControl/Private/RemoteControlModule.h
+++ b/Engine/Plugins/VirtualProduction/RemoteControl/Source/RemoteControl/Private/RemoteControlModule.h
@@ -124,6 +124,17 @@ private:
 	 */
 	void RegisterMaskingFactories();
 
+	/**
+	 * Whether function can be intercepted by a remote control interceptor.
+	 */
+	bool CanInterceptFunction(const FRCCall& RCCall) const;
+
+	/**
+	 * Used to log once per program execution. Called by the REMOTE_CONTROL_LOG_ONCE macro.
+	 */
+	void LogOnce(ELogVerbosity::Type InVerbosity, const FString& LogDetails, const FString& FileName, int32 LineNumber) const;
+
+
 	/**
 	 * Add an item to an array property, handling transactions and editor events automatically.
 	 * @param ObjectAccess Reference to the array property to be modified.
