Fix remote control crashes

- GEditor was accessed without testing for nullptr
- Make sure propertywatcher is only used in -game except for transform properties
- In PropertyWatcher->CheckForChange, IsIdentical was call on LastFrameValue.GetData() without first checking if the last frame value had any data in it. LastFrameValue should now be correctly initalized if it couldn't be done at the creation of the property watcher

diff --git a/Engine/Plugins/VirtualProduction/RemoteControl/Source/RemoteControl/Private/RemoteControlPreset.cpp b/Engine/Plugins/VirtualProduction/RemoteControl/Source/RemoteControl/Private/RemoteControlPreset.cpp
--- a/Engine/Plugins/VirtualProduction/RemoteControl/Source/RemoteControl/Private/RemoteControlPreset.cpp
+++ b/Engine/Plugins/VirtualProduction/RemoteControl/Source/RemoteControl/Private/RemoteControlPreset.cpp
@@ -114,6 +114,31 @@ namespace
 		}
 		return OutputName;
 	}
+
+	enum class ELaunchConfiguration : uint8
+	{
+		Editor,
+		DashGame,
+		Packaged
+	};
+
+	ELaunchConfiguration GetLaunchConfiguration()
+	{
+		ELaunchConfiguration CurrentConfiguration = ELaunchConfiguration::Packaged;
+
+#if WITH_EDITOR
+		if (GEditor)
+		{
+			CurrentConfiguration = ELaunchConfiguration::Editor;
+		}
+		else 
+		{
+			CurrentConfiguration = ELaunchConfiguration::DashGame;
+		}
+#endif
+
+		return CurrentConfiguration;
+	}
 }
 
 FRemoteControlPresetExposeArgs::FRemoteControlPresetExposeArgs()
@@ -1137,28 +1162,29 @@ void URemoteControlPreset::CreatePropertyWatcher(const TSharedPtr<FRemoteControl
 
 bool URemoteControlPreset::PropertyShouldBeWatched(const TSharedPtr<FRemoteControlProperty>& RCProperty) const
 {
-#if WITH_EDITOR
-	if (GEditor && !CVarRemoteControlEnablePropertyWatchInEditor.GetValueOnAnyThread())
-	{
-		// Don't use property watchers in editor unless explicitely specified.
-		return false;
-	}
-#endif
-
-	// If we are not running in editor, we need to watch all properties as there is no object modified callback.
-	if (!GIsEditor)
-	{
-		return true;	
-	}
-	
+	// We know these properties will be redirected to their setter in -game or packaged, so we won't get the object modified callback.
 	static const TSet<FName> WatchedPropertyNames =
 		{
 			UStaticMeshComponent::GetRelativeLocationPropertyName(),
 			UStaticMeshComponent::GetRelativeRotationPropertyName(),
 			UStaticMeshComponent::GetRelativeScale3DPropertyName()
 		};
-	
-	return RCProperty && WatchedPropertyNames.Contains(RCProperty->FieldName);
+
+	static const ELaunchConfiguration LaunchConfiguration = GetLaunchConfiguration();
+
+	// In packaged, we don't have access to object modified callbacks, so we have to track them manually.
+	if (LaunchConfiguration == ELaunchConfiguration::Packaged || CVarRemoteControlEnablePropertyWatchInEditor.GetValueOnAnyThread())
+	{
+		return true;
+	}
+	else if (LaunchConfiguration == ELaunchConfiguration::DashGame)
+	{
+		return RCProperty && WatchedPropertyNames.Contains(RCProperty->FieldName);
+	}
+	else
+	{
+		return false;
+	}
 }
 
 void URemoteControlPreset::CreatePropertyWatchers()
@@ -1204,7 +1230,8 @@ void URemoteControlPreset::HandleDisplayClusterConfigChange(UObject* DisplayClus
 {
 #if WITH_EDITOR
 	AActor* OwnerActor = DisplayClusterConfigData->GetTypedOuter<AActor>();
-	GEditor->GetTimerManager()->SetTimerForNextTick(FTimerDelegate::CreateLambda([OwnerActor, PresetPtr = TWeakObjectPtr<URemoteControlPreset>{ this }]()
+	FTimerManager& TimerManager = GEditor ? *GEditor->GetTimerManager() : OwnerActor->GetWorld()->GetTimerManager();
+	TimerManager.SetTimerForNextTick(FTimerDelegate::CreateLambda([OwnerActor, PresetPtr = TWeakObjectPtr<URemoteControlPreset>{ this }]()
 	{
 		TSet<URemoteControlBinding*> ModifiedBindings;
 
@@ -2265,7 +2292,11 @@ void URemoteControlPreset::FRCPropertyWatcher::CheckForChange()
 	TRACE_CPUPROFILER_EVENT_SCOPE(FRCPropertyWatcher::CheckForChange);
 	if (TOptional<FRCFieldResolvedData> ResolvedData = GetWatchedPropertyResolvedData())
 	{
-		if (ensure(ResolvedData->Field && ResolvedData->ContainerAddress))
+		if (!LastFrameValue.Num())
+		{
+			SetLastFrameValue(*ResolvedData);
+		}
+		else if (ensure(ResolvedData->Field && ResolvedData->ContainerAddress))
 		{
 			const void* NewValueAddress = ResolvedData->Field->ContainerPtrToValuePtr<void>(ResolvedData->ContainerAddress);
 			if (NewValueAddress && (ResolvedData->Field->GetSize() != LastFrameValue.Num() || !ResolvedData->Field->Identical(LastFrameValue.GetData(), NewValueAddress)))
