Remote Control - HTTP preset property routes fail to set values for array elements

diff --git a/Engine/Plugins/VirtualProduction/RemoteControl/Source/WebRemoteControl/Private/WebRemoteControlInternalUtils.h b/Engine/Plugins/VirtualProduction/RemoteControl/Source/WebRemoteControl/Private/WebRemoteControlInternalUtils.h
--- a/Engine/Plugins/VirtualProduction/RemoteControl/Source/WebRemoteControl/Private/WebRemoteControlInternalUtils.h
+++ b/Engine/Plugins/VirtualProduction/RemoteControl/Source/WebRemoteControl/Private/WebRemoteControlInternalUtils.h
@@ -385,7 +385,14 @@ namespace WebRemoteControlInternalUtils
 
 		// Replace PropertyValue with the underlying property name.
 		TArray<uint8> NewPayload;
-		RemotePayloadSerializer::ReplaceFirstOccurence(Payload, TEXT("PropertyValue"), Property.FieldName.ToString(), NewPayload);
+		FString FieldName = Property.FieldName.ToString();
+		if (Property.GetProperty()->IsA<FArrayProperty>() ||
+			Property.GetProperty()->IsA<FMapProperty>() ||
+			Property.GetProperty()->IsA<FSetProperty>())
+		{
+			FieldName = Property.FieldPathInfo.Segments.Last().Name.ToString();
+		}
+		RemotePayloadSerializer::ReplaceFirstOccurence(Payload, TEXT("PropertyValue"), FieldName, NewPayload);
 
 		// Then deserialize the payload onto all the bound objects.
 		FMemoryReader NewPayloadReader(NewPayload);
