Remote Control Web Interface bind address should not be 0.0.0.0

diff --git a/Engine/Plugins/VirtualProduction/RemoteControl/Source/RemoteControlCommon/Public/RemoteControlSettings.h b/Engine/Plugins/VirtualProduction/RemoteControl/Source/RemoteControlCommon/Public/RemoteControlSettings.h
--- a/Engine/Plugins/VirtualProduction/RemoteControl/Source/RemoteControlCommon/Public/RemoteControlSettings.h
+++ b/Engine/Plugins/VirtualProduction/RemoteControl/Source/RemoteControlCommon/Public/RemoteControlSettings.h
@@ -203,6 +203,31 @@ public:
 		return OutArray;
 	}
 
+	virtual void AllowClient(const FString& InClientAddressStr)
+	{
+		// Return early if we have this in range.
+		if (IsClientAllowed(InClientAddressStr))
+		{
+			return;
+		}
+
+		// Add a new range.
+		TArray<FString> IndividualClasses;
+
+		if (InClientAddressStr.ParseIntoArray(IndividualClasses, TEXT(".")) == 4)
+		{
+			FRCNetworkAddress LowerAndUpperBounds { FCString::Atoi(*IndividualClasses[0]) // Class A
+				, FCString::Atoi(*IndividualClasses[1]) // Class B
+				, FCString::Atoi(*IndividualClasses[2]) // Class C
+				, FCString::Atoi(*IndividualClasses[3]) // Class D
+			};
+
+			FRCNetworkAddressRange NewRange = { LowerAndUpperBounds, LowerAndUpperBounds };
+
+			AllowlistedClients.Add(NewRange);
+		}
+	}
+	
 	virtual bool IsClientAllowed(const FString& InClientAddressStr) const
 	{
 		for (const FRCNetworkAddressRange& AllowlistedClient : AllowlistedClients)
@@ -315,12 +340,6 @@ public:
 	UPROPERTY(config, EditAnywhere, Category = "Remote Control | Security")
 	bool bEnforcePassphraseForRemoteClients = false;
 	
-	/**
-     * List of IPs that are explicitly *not* required to have a passphrase when accessing remote control.
-     */
-    UPROPERTY(config, EditAnywhere, Category = "Remote Control | Security", meta= (EditCondition = bEnforcePassphraseForRemoteClients))
-    TArray<FString> AllowedIPsForRemotePassphrases;
-	
 	UPROPERTY(config, EditAnywhere, Category = "Remote Control | Security", DisplayName = "Remote Control Passphrase")
 	TArray<FRCPassphrase> Passphrases = {};
 
diff --git a/Engine/Plugins/VirtualProduction/RemoteControl/Source/WebRemoteControl/Private/RemoteControlDefaultPreprocessors.h b/Engine/Plugins/VirtualProduction/RemoteControl/Source/WebRemoteControl/Private/RemoteControlDefaultPreprocessors.h
--- a/Engine/Plugins/VirtualProduction/RemoteControl/Source/WebRemoteControl/Private/RemoteControlDefaultPreprocessors.h
+++ b/Engine/Plugins/VirtualProduction/RemoteControl/Source/WebRemoteControl/Private/RemoteControlDefaultPreprocessors.h
@@ -108,7 +108,7 @@ namespace UE::WebRemoteControl
 		{
 			Notification->SetCompletionState(SNotificationItem::ECompletionState::CS_Pending);
 
-			GetMutableDefault<URemoteControlSettings>()->AllowedIPsForRemotePassphrases.Add(IPAddress);
+			GetMutableDefault<URemoteControlSettings>()->AllowClient(IPAddress);
 			SaveRemoteControlConfig();
 			
 			Notification->SetCompletionState(SNotificationItem::ECompletionState::CS_Success);
@@ -273,7 +273,7 @@ namespace UE::WebRemoteControl
 						if (ForwardedIP->Num())
 						{
 							// We will need to rework this when the IP range change is in.
-							if (GetDefault<URemoteControlSettings>()->AllowedIPsForRemotePassphrases.Contains(ForwardedIP->Last()))
+							if (GetDefault<URemoteControlSettings>()->IsClientAllowed(ForwardedIP->Last()))
 							{
 								return EPreprocessorResult::RequestPassthrough;
 							}
@@ -289,7 +289,7 @@ namespace UE::WebRemoteControl
 					}
 				}
 
-				if (GetDefault<URemoteControlSettings>()->AllowedIPsForRemotePassphrases.Contains(PeerAddress))
+				if (GetDefault<URemoteControlSettings>()->IsClientAllowed(PeerAddress))
 				{
 					return EPreprocessorResult::RequestPassthrough;
 				}
@@ -431,6 +431,14 @@ namespace UE::WebRemoteControl
 				// Allow requests from localhost
 				if (ClientIP != TEXT("localhost") && ClientIP != TEXT("127.0.0.1"))
 				{
+					if (!GetDefault<URemoteControlSettings>()->IsClientAllowed(ClientIP))
+					{
+						WebRemoteControlInternalUtils::CreateUTF8ErrorMessage(FString::Printf(TEXT("Client IP %s does not respect the allowed IP set in Remote Control Settings."), *ClientIP), Response->Body);
+						Response->Code = EHttpServerResponseCodes::Denied;
+						OnComplete(MoveTemp(Response));
+						return EPreprocessorResult::RequestHandled;
+					}
+
 					if (!WildcardAllowedIP.IsEmpty() && WildcardAllowedIP.IsMatch(ClientIP))
 					{
 						WebRemoteControlInternalUtils::CreateUTF8ErrorMessage(FString::Printf(TEXT("Client IP %s does not respect the allowed IP set in Remote Control Settings."), *ClientIP), Response->Body);
