Fix remote control routes circumventing security features

Remote control internal preprocessors will now be redirected to the same preprocessors that are used for http routes.

diff --git a/Engine/Plugins/VirtualProduction/RemoteControl/Source/WebRemoteControl/Private/RemoteControlDefaultPreprocessors.h b/Engine/Plugins/VirtualProduction/RemoteControl/Source/WebRemoteControl/Private/RemoteControlDefaultPreprocessors.h
--- a/Engine/Plugins/VirtualProduction/RemoteControl/Source/WebRemoteControl/Private/RemoteControlDefaultPreprocessors.h
+++ b/Engine/Plugins/VirtualProduction/RemoteControl/Source/WebRemoteControl/Private/RemoteControlDefaultPreprocessors.h
@@ -28,32 +28,80 @@
 
 namespace UE::WebRemoteControl
 {
+	/** Indicates whether a request has been handled (OnComplete callback called) or if it just passed through the preprocessor. */
 	enum class EPreprocessorResult : uint8
 	{
 		RequestPassthrough,
 		RequestHandled
 	};
+
+	/** Result of a preprocessor. When failed, it will automatically respond to the client request.
+	 * Usage:
+	 *	 Deny a request:
+	 *		return FPreprocessorResult::Deny(TEXT("My preprocessor error message"));
+	 * 
+	 *	or let it through:
+	 *		return FPreprocessorResult::Passthrough();
+	 * 
+	 * 
+	 */
+	struct FPreprocessorResult
+	{
+		/** Let request pass. */
+		static FPreprocessorResult Passthrough()
+		{
+			return FPreprocessorResult();
+		}
+
+		/** Deny request and respond with error message. */
+		static FPreprocessorResult Deny(const FString& ErrorMessage)
+		{
+			UE_LOG(LogRemoteControl, Error, TEXT("%s"), *ErrorMessage);
+			IRemoteControlModule::BroadcastError(ErrorMessage);
+
+			TUniquePtr<FHttpServerResponse> Response = WebRemoteControlInternalUtils::CreateHttpResponse();
+			WebRemoteControlInternalUtils::CreateUTF8ErrorMessage(ErrorMessage, Response->Body);
+			Response->Code = EHttpServerResponseCodes::Denied;
+
+			return FPreprocessorResult{ EPreprocessorResult::RequestHandled, MoveTemp(Response) };
+		}
+
+	/** Holds the preprocessor result. */
+	EPreprocessorResult Result = EPreprocessorResult::RequestPassthrough;
+	/** If denied, holds the response to be sent to the client. */
+	TUniquePtr<FHttpServerResponse> OptionalResponse;
+
+	private:
+		FPreprocessorResult() = default;
+	};
 	
-	using FRCPreprocessorHandler = TFunction<EPreprocessorResult(const FHttpServerRequest& Request, const FHttpResultCallback& OnComplete)>;
+	using FRCPreprocessorHandler = TFunction<FPreprocessorResult(const FHttpServerRequest& Request)>;
 
 	/** Utility function to wrap a preprocessor handler to a http request handler than the HttpRouter can take. */
 	FHttpRequestHandler MakeHttpRequestHandler(FRCPreprocessorHandler Handler)
 	{
 		return [WrappedHandler = MoveTemp(Handler)](const FHttpServerRequest& Request, const FHttpResultCallback& OnComplete)
 		{
-			return WrappedHandler(Request, OnComplete) == EPreprocessorResult::RequestPassthrough ? false : true;
+			FPreprocessorResult Result = WrappedHandler(Request);
+			if (Result.Result == EPreprocessorResult::RequestPassthrough)
+			{
+				return false; // Request not handled.
+			}
+			else
+			{
+				OnComplete(MoveTemp(Result.OptionalResponse));
+				return true; // Request handled.
+			}
 		};
 	}
 
 #if WITH_EDITOR
+	/** Notification prompt used when a remote client attempts to connect to unreal without a passphrase. */
 	TWeakPtr<SNotificationItem> NoPassphrasePrompt;
+	/** Message displayed in the notification prompt when no action was taken. */
 	FText NoActionTakenText = LOCTEXT("NoActionTaken", "No action was taken, further requests from this IP will be denied.");
 
-	const FString& GetRemoteControlConfigPath()
-	{
-		static const FString RemoteControlConfigPath = FPackageName::LongPackageNameToFilename(GetDefault<URemoteControlSettings>()->GetPackage()->GetName(), TEXT(".ini"));
-		return RemoteControlConfigPath;
-	}
+	/** Attempt saving the remote control config, checkouting it if needed.*/
 	void SaveRemoteControlConfig()
 	{
 		FString ConfigFilename = FPaths::ConvertRelativePathToFull(URemoteControlSettings::StaticClass()->GetConfigName());
@@ -102,6 +150,7 @@ namespace UE::WebRemoteControl
 		}
 	}
 
+	/** Add IP to the list of IPs that should be let through by RC. */
 	void AddIPToAllowlist(FString IPAddress)
 	{
 		if (TSharedPtr<SNotificationItem> Notification = NoPassphrasePrompt.Pin())
@@ -116,6 +165,7 @@ namespace UE::WebRemoteControl
 		}
 	}
 
+	/** Prompts the user to create a passphrase.  */
 	void CreatePassphrase()
 	{
 		if (TSharedPtr<SNotificationItem> Notification = NoPassphrasePrompt.Pin())
@@ -226,6 +276,7 @@ namespace UE::WebRemoteControl
 		}
 	}
 
+	/** Disables remote passphrase enforcement by modifying a RC project setting. */
 	void DisableRemotePassphrases()
 	{
 		if (TSharedPtr<SNotificationItem> Notification = NoPassphrasePrompt.Pin())
@@ -252,46 +303,59 @@ namespace UE::WebRemoteControl
 #endif
 
 	// Notifies the editor if a client tries to access the RC server without a passphrase.
-	EPreprocessorResult RemotePassphraseEnforcementPreprocessor(const FHttpServerRequest& Request, const FHttpResultCallback& OnComplete)
+	FPreprocessorResult RemotePassphraseEnforcementPreprocessor(const FHttpServerRequest& Request)
 	{
 		if (GetDefault<URemoteControlSettings>()->bRestrictServerAccess)
 		{
 			if (Request.PeerAddress)
 			{
 				constexpr bool bAppendPort = false;
-				const FString PeerAddress = Request.PeerAddress->ToString(bAppendPort);
+				FString PeerAddress = Request.PeerAddress->ToString(bAppendPort);
 
 				if (!GetDefault<URemoteControlSettings>()->bEnforcePassphraseForRemoteClients)
 				{
-					return EPreprocessorResult::RequestPassthrough;
+					return FPreprocessorResult::Passthrough();
 				}
 
-				if (PeerAddress.Contains(TEXT("127.0.0.1")) || PeerAddress.Contains(TEXT("localhost")))
+				auto IsLocal = [](const FString& Address)
+				{
+					return Address.Contains(TEXT("127.0.0.1")) || Address.Contains(TEXT("localhost"));
+				};
+
+				bool bContainsForwardedAddress = false;
+				if (IsLocal(PeerAddress))
 				{
 					if (const TArray<FString>* ForwardedIP = Request.Headers.Find(WebRemoteControlInternalUtils::ForwardedIPHeader))
 					{
 						if (ForwardedIP->Num())
 						{
-							// We will need to rework this when the IP range change is in.
-							if (GetDefault<URemoteControlSettings>()->IsClientAllowed(ForwardedIP->Last()))
+							PeerAddress = ForwardedIP->Last();
+
+							if (GetDefault<URemoteControlSettings>()->IsClientAllowed(PeerAddress))
 							{
-								return EPreprocessorResult::RequestPassthrough;
+								return FPreprocessorResult::Passthrough();
+							}
+							else
+							{
+								// Let the forwarded IP go through the rest of the validations if it's not 
+								bContainsForwardedAddress = true;
 							}
 						}
 						else
 						{
-							return EPreprocessorResult::RequestPassthrough;
+							return FPreprocessorResult::Passthrough();
 						}
 					}
 					else
 					{
-						return EPreprocessorResult::RequestPassthrough;
+						return FPreprocessorResult::Passthrough();
 					}
 				}
 
-				if (GetDefault<URemoteControlSettings>()->IsClientAllowed(PeerAddress))
+				// We've already attempted to validate a forwarded address so just do it if it's a direct remote address.
+				if (!bContainsForwardedAddress && GetDefault<URemoteControlSettings>()->IsClientAllowed(PeerAddress))
 				{
-					return EPreprocessorResult::RequestPassthrough;
+					return FPreprocessorResult::Passthrough();
 				}
 
 				if (GetDefault<URemoteControlSettings>()->Passphrases.Num())
@@ -299,26 +363,16 @@ namespace UE::WebRemoteControl
 					const TArray<FString>* PassphraseHeader = Request.Headers.Find(WebRemoteControlInternalUtils::PassphraseHeader);
 					if (!PassphraseHeader || PassphraseHeader->Num() == 0)
 					{
-						TUniquePtr<FHttpServerResponse> Response = WebRemoteControlInternalUtils::CreateHttpResponse();
-
-						const FString ErrorMessage = TEXT("Remote passphrase enforcement is enabled but no passphrase was specified!");
-						UE_LOG(LogRemoteControl, Error, TEXT("%s"), *ErrorMessage);
-						IRemoteControlModule::BroadcastError(ErrorMessage);
-						WebRemoteControlInternalUtils::CreateUTF8ErrorMessage(ErrorMessage, Response->Body);
-
-						Response->Code = EHttpServerResponseCodes::Denied;
-						OnComplete(MoveTemp(Response));
-						return EPreprocessorResult::RequestHandled;
+						return FPreprocessorResult::Deny(TEXT("Remote passphrase enforcement is enabled but no passphrase was specified!"));
 					}
 
 					if (!WebRemoteControlInternalUtils::CheckPassphrase(PassphraseHeader->Last()))
 					{
-						OnComplete(WebRemoteControlInternalUtils::CreatedInvalidPassphraseResponse());
-						return EPreprocessorResult::RequestHandled;
+						return FPreprocessorResult::Deny(WebRemoteControlInternalUtils::InvalidPassphraseError);
 					}
 					else
 					{
-						return EPreprocessorResult::RequestPassthrough;
+						return FPreprocessorResult::Passthrough();
 					}
 				}
 				else
@@ -329,38 +383,43 @@ namespace UE::WebRemoteControl
 					 * 2: Create a passphrase, tell user he must enter that on the app he's using or put it in the header
 					 * 3: Disable remote passphrase enforcement. (Warn that this is dangerous)
 					 */
-	#if WITH_EDITOR
+#if WITH_EDITOR
 					if (GEditor)
 					{
-						if (!NoPassphrasePrompt.IsValid() || NoPassphrasePrompt.Pin()->GetCompletionState() == SNotificationItem::ECompletionState::CS_None)
+						if (!NoPassphrasePrompt.IsValid() || NoPassphrasePrompt.Pin()->GetCompletionState() == SNotificationItem::CS_None)
 						{
 							FNotificationInfo Info(LOCTEXT("NoPassphraseNotificationHeader", "Remote control request denied!"));
 							Info.SubText = FText::Format(LOCTEXT("NoPassphraseNotificationSubtext", "A remote control request was made by an external client ({0}) without a passphrase."), FText::FromString(PeerAddress));
 							Info.bFireAndForget = false;
 							Info.FadeInDuration = 0.5f;
-							Info.ButtonDetails.Add(FNotificationButtonInfo(LOCTEXT("AllowListIP", "AllowList IP"), LOCTEXT("AllowListTooltip", "Add this IP to the list of allowed IPs that can make remote control requests without a passphrase."), FSimpleDelegate::CreateStatic(&AddIPToAllowlist, PeerAddress), SNotificationItem::CS_None));
-							Info.ButtonDetails.Add(FNotificationButtonInfo(LOCTEXT("CreatePassphrase", "Create Passphrase"), LOCTEXT("CreatePassPhraseTooltip", "Create a passphrase for this client to use."), FSimpleDelegate::CreateStatic(&CreatePassphrase), SNotificationItem::CS_None));
-							Info.ButtonDetails.Add(FNotificationButtonInfo(LOCTEXT("DisableRemotePassphrase", "Disable Passphrases"), LOCTEXT("DisableRemotePassphrasesTooltip", "Disable the requirement for remote control requests coming from external clients to have a passphrase.\nWarning: This should only be done as a last resort since all clients on the network will be able to access your servers."), FSimpleDelegate::CreateStatic(&DisableRemotePassphrases), SNotificationItem::CS_None));
+							Info.ButtonDetails.Add(FNotificationButtonInfo(LOCTEXT("AllowListIP", "AllowList IP"), LOCTEXT("AllowListTooltip", "Add this IP to the list of allowed IPs that can make remote control requests without a passphrase."), FSimpleDelegate::CreateStatic(&AddIPToAllowlist, PeerAddress)));
+							Info.ButtonDetails.Add(FNotificationButtonInfo(LOCTEXT("CreatePassphrase", "Create Passphrase"), LOCTEXT("CreatePassPhraseTooltip", "Create a passphrase for this client to use."), FSimpleDelegate::CreateStatic(&CreatePassphrase)));
+							Info.ButtonDetails.Add(FNotificationButtonInfo(LOCTEXT("DisableRemotePassphrase", "Disable Passphrases"), LOCTEXT("DisableRemotePassphrasesTooltip", "Disable the requirement for remote control requests coming from external clients to have a passphrase.\nWarning: This should only be done as a last resort since all clients on the network will be able to access your servers."), FSimpleDelegate::CreateStatic(&DisableRemotePassphrases)));
 							Info.WidthOverride = 450.f;
-							NoPassphrasePrompt = FSlateNotificationManager::Get().AddNotification(MoveTemp(Info));	
+							NoPassphrasePrompt = FSlateNotificationManager::Get().AddNotification(MoveTemp(Info));
+							NoPassphrasePrompt.Pin()->SetCompletionState(SNotificationItem::CS_Pending);
+							return FPreprocessorResult::Deny(TEXT("A passphrase for remote control is required. See editor for resolving this."));
+						}
+						else
+						{
+							return FPreprocessorResult::Deny(TEXT("A passphrase for remote control is required. See editor for resolving this."));
 						}
 					}
-	#endif
-					// In -game or packaged, we won't prompt the user so just deny the request.
-					const FString ErrorMessage = TEXT("A passphrase for remote control is required but we can't create one in -game or packaged.");
-					UE_LOG(LogRemoteControl, Error, TEXT("%s"), *ErrorMessage);
-					IRemoteControlModule::BroadcastError(ErrorMessage);
-
-					OnComplete(WebRemoteControlInternalUtils::CreatedInvalidPassphraseResponse());
-					return EPreprocessorResult::RequestHandled;
+					else
+#endif
+					{
+						// In -game or packaged, we won't prompt the user so just deny the request.
+						return FPreprocessorResult::Deny(TEXT("A passphrase for remote control is required but we can't create one in -game or packaged."));
+					}
 				}
 			}
 		}
 
-		return EPreprocessorResult::RequestPassthrough;
+		return FPreprocessorResult::Passthrough();
 	}
 	
-	EPreprocessorResult PassphrasePreprocessor(const FHttpServerRequest& Request, const FHttpResultCallback& OnComplete)
+	/** Checks whether a request has a valid passphrase when passphrases are enabled for this editor. */
+	FPreprocessorResult PassphrasePreprocessor(const FHttpServerRequest& Request)
 	{
 		TUniquePtr<FHttpServerResponse> Response = WebRemoteControlInternalUtils::CreateHttpResponse();
 
@@ -374,14 +433,17 @@ namespace UE::WebRemoteControl
 
 		if (!WebRemoteControlInternalUtils::CheckPassphrase(Passphrase))
 		{
-			OnComplete(WebRemoteControlInternalUtils::CreatedInvalidPassphraseResponse());
-			return EPreprocessorResult::RequestHandled;
+			return FPreprocessorResult::Deny(WebRemoteControlInternalUtils::InvalidPassphraseError);
 		}
 
-		return EPreprocessorResult::RequestPassthrough;
+		return FPreprocessorResult::Passthrough();
 	}
 
-	EPreprocessorResult IPValidationPreprocessor(const FHttpServerRequest& Request, const FHttpResultCallback& OnComplete)
+	/** Checks whether an IP is in a valid range for making remote control requests. 
+	 * Also checks for a valid origin to block malicious requests from web browsers.
+	 * Can be controlled using the Allowed Origin and AllowedIP remote control settings.
+	 */
+	FPreprocessorResult IPValidationPreprocessor(const FHttpServerRequest& Request)
 	{
 		if (GetDefault<URemoteControlSettings>()->bRestrictServerAccess)
 		{
@@ -415,10 +477,7 @@ namespace UE::WebRemoteControl
 			{
 				if (!SimplifiedAllowedOrigin.IsMatch(SimplifiedOrigin))
 				{
-					WebRemoteControlInternalUtils::CreateUTF8ErrorMessage(FString::Printf(TEXT("Client origin %s does not respect the allowed origin set in Remote Control Settings."), *OriginHeader), Response->Body);
-					Response->Code = EHttpServerResponseCodes::Denied;
-					OnComplete(MoveTemp(Response));
-					return EPreprocessorResult::RequestHandled;
+					return FPreprocessorResult::Deny(FString::Printf(TEXT("Client origin %s does not respect the allowed origin set in Remote Control Settings."), *OriginHeader));
 				}
 			}
 
@@ -433,24 +492,18 @@ namespace UE::WebRemoteControl
 				{
 					if (!GetDefault<URemoteControlSettings>()->IsClientAllowed(ClientIP))
 					{
-						WebRemoteControlInternalUtils::CreateUTF8ErrorMessage(FString::Printf(TEXT("Client IP %s does not respect the allowed IP set in Remote Control Settings."), *ClientIP), Response->Body);
-						Response->Code = EHttpServerResponseCodes::Denied;
-						OnComplete(MoveTemp(Response));
-						return EPreprocessorResult::RequestHandled;
+						return FPreprocessorResult::Deny(FString::Printf(TEXT("Client IP %s does not respect the allowed IP set in Remote Control Settings."), *ClientIP));
 					}
 
 					if (!WildcardAllowedIP.IsEmpty() && WildcardAllowedIP.IsMatch(ClientIP))
 					{
-						WebRemoteControlInternalUtils::CreateUTF8ErrorMessage(FString::Printf(TEXT("Client IP %s does not respect the allowed IP set in Remote Control Settings."), *ClientIP), Response->Body);
-						Response->Code = EHttpServerResponseCodes::Denied;
-						OnComplete(MoveTemp(Response));
-						return EPreprocessorResult::RequestHandled;
+						return FPreprocessorResult::Deny(FString::Printf(TEXT("Client IP %s does not respect the allowed IP set in Remote Control Settings."), *ClientIP));
 					}
 				}
 			}
 		}
 
-		return EPreprocessorResult::RequestPassthrough;
+		return FPreprocessorResult::Passthrough();
 	}
 }
 
diff --git a/Engine/Plugins/VirtualProduction/RemoteControl/Source/WebRemoteControl/Private/RemoteControlRequest.h b/Engine/Plugins/VirtualProduction/RemoteControl/Source/WebRemoteControl/Private/RemoteControlRequest.h
--- a/Engine/Plugins/VirtualProduction/RemoteControl/Source/WebRemoteControl/Private/RemoteControlRequest.h
+++ b/Engine/Plugins/VirtualProduction/RemoteControl/Source/WebRemoteControl/Private/RemoteControlRequest.h
@@ -525,6 +525,12 @@ struct FRCWebSocketRequest : public FRCRequest
 	 */
 	UPROPERTY()
 	int32 Id = INDEX_NONE;
+
+	/**
+	 * (Optional) If the request was forwared for a remote client, this will contain the forwarded IP.
+	 */
+	UPROPERTY()
+	FString ForwardedFor;
 };
 
 /**
diff --git a/Engine/Plugins/VirtualProduction/RemoteControl/Source/WebRemoteControl/Private/RemoteControlWebSocketServer.cpp b/Engine/Plugins/VirtualProduction/RemoteControl/Source/WebRemoteControl/Private/RemoteControlWebSocketServer.cpp
--- a/Engine/Plugins/VirtualProduction/RemoteControl/Source/WebRemoteControl/Private/RemoteControlWebSocketServer.cpp
+++ b/Engine/Plugins/VirtualProduction/RemoteControl/Source/WebRemoteControl/Private/RemoteControlWebSocketServer.cpp
@@ -53,7 +53,11 @@ namespace RemoteControlWebSocketServer
 			}
 			if (!Request.Passphrase.IsEmpty())
 			{
-				Message.Header.FindOrAdd(WebRemoteControlInternalUtils::PassphraseHeader) = TArray<FString>({Request.Passphrase});
+				Message.Header.FindOrAdd(WebRemoteControlInternalUtils::PassphraseHeader) = { Request.Passphrase };
+			}
+			if (!Request.ForwardedFor.IsEmpty())
+			{
+				Message.Header.FindOrAdd(WebRemoteControlInternalUtils::ForwardedIPHeader) = { Request.ForwardedFor };
 			}
 			Message.MessageName = MoveTemp(Request.MessageName);
 			ParsedMessage = MoveTemp(Message);
@@ -215,7 +219,7 @@ void FRCWebSocketServer::OnWebSocketClientConnected(INetworkingWebSocket* Socket
 		FWebSocketConnection Connection = FWebSocketConnection{ Socket };
 			
 		FWebSocketPacketReceivedCallBack ReceiveCallBack;
-		ReceiveCallBack.BindRaw(this, &FRCWebSocketServer::ReceivedRawPacket, Connection.Id);
+		ReceiveCallBack.BindRaw(this, &FRCWebSocketServer::ReceivedRawPacket, Connection.Id, Connection.PeerAddress);
 		Socket->SetReceiveCallBack(ReceiveCallBack);
 
 		FWebSocketInfoCallBack CloseCallback;
@@ -227,7 +231,7 @@ void FRCWebSocketServer::OnWebSocketClientConnected(INetworkingWebSocket* Socket
 	}
 }
 
-void FRCWebSocketServer::ReceivedRawPacket(void* Data, int32 Size, FGuid ClientId)
+void FRCWebSocketServer::ReceivedRawPacket(void* Data, int32 Size, FGuid ClientId, TSharedPtr<FInternetAddr> PeerAddress)
 {
 	if (!Router)
 	{
@@ -242,6 +246,8 @@ void FRCWebSocketServer::ReceivedRawPacket(void* Data, int32 Size, FGuid ClientI
 	if (TOptional<FRemoteControlWebSocketMessage> Message = RemoteControlWebSocketServer::ParseWebsocketMessage(Payload))
 	{
 		Message->ClientId = ClientId;
+		Message->PeerAddress = PeerAddress;
+	
 		Router->AttemptDispatch(*Message);
 	}
 }
diff --git a/Engine/Plugins/VirtualProduction/RemoteControl/Source/WebRemoteControl/Private/RemoteControlWebSocketServer.h b/Engine/Plugins/VirtualProduction/RemoteControl/Source/WebRemoteControl/Private/RemoteControlWebSocketServer.h
--- a/Engine/Plugins/VirtualProduction/RemoteControl/Source/WebRemoteControl/Private/RemoteControlWebSocketServer.h
+++ b/Engine/Plugins/VirtualProduction/RemoteControl/Source/WebRemoteControl/Private/RemoteControlWebSocketServer.h
@@ -3,12 +3,14 @@
 #pragma once
 
 #include "WebSocketNetDriver.h"
-#include "UObject/StrongObjectPtr.h"
+
+#include "Containers/Ticker.h"
 #include "INetworkingWebSocket.h"
 #include "IWebRemoteControlModule.h"
 #include "IWebSocketServer.h"
 #include "RemoteControlRoute.h"
-#include "Containers/Ticker.h"
+#include "SocketSubsystem.h"
+#include "UObject/StrongObjectPtr.h"
 
 /**
  * Router used to dispatch messages received by a WebSocketServer.
@@ -32,6 +34,7 @@ public:
 	/**
 	 * Adds a Preprocessor called before the Dispatch of the Route
 	 * @param WebsocketPreprocessor The Preprocessor Function to register 
+	 * @note: If the predispatch returns false, the actual request will not be executed.
 	 */
 	void AddPreDispatch(TFunction<bool(const struct FRemoteControlWebSocketMessage& Message)> WebsocketPreprocessor);
 
@@ -118,7 +121,7 @@ private:
 	void OnWebSocketClientConnected(INetworkingWebSocket* Socket);
 
 	/** Handles sending the received packet to the message router. */
-	void ReceivedRawPacket(void* Data, int32 Size, FGuid ClientId);
+	void ReceivedRawPacket(void* Data, int32 Size, FGuid ClientId, TSharedPtr<FInternetAddr> PeerAddress);
 
 	void OnSocketClose(INetworkingWebSocket* Socket);
 	
@@ -138,12 +141,18 @@ private:
 			: Socket(InSocket)
 			, Id(FGuid::NewGuid())
 		{
+			constexpr bool bAppendPort = false;
+			if (ISocketSubsystem* SocketSubsystem = ISocketSubsystem::Get(PLATFORM_SOCKETSUBSYSTEM))
+			{
+				PeerAddress = SocketSubsystem->GetAddressFromString(InSocket->RemoteEndPoint(bAppendPort));
+			}
 		}
 
 		FWebSocketConnection(FWebSocketConnection&& WebSocketConnection)
 			: Id(WebSocketConnection.Id)
 		{
 			Socket = WebSocketConnection.Socket;
+			PeerAddress = WebSocketConnection.PeerAddress;
 			WebSocketConnection.Socket = nullptr;
 		}
 
@@ -165,6 +174,9 @@ private:
 
 		/** Generated ID for this client. */
 		FGuid Id;
+
+		/** IP Address of the client. */
+		TSharedPtr<FInternetAddr> PeerAddress;
 	};
 
 
diff --git a/Engine/Plugins/VirtualProduction/RemoteControl/Source/WebRemoteControl/Private/WebRemoteControl.cpp b/Engine/Plugins/VirtualProduction/RemoteControl/Source/WebRemoteControl/Private/WebRemoteControl.cpp
--- a/Engine/Plugins/VirtualProduction/RemoteControl/Source/WebRemoteControl/Private/WebRemoteControl.cpp
+++ b/Engine/Plugins/VirtualProduction/RemoteControl/Source/WebRemoteControl/Private/WebRemoteControl.cpp
@@ -357,30 +357,6 @@ void FWebRemoteControlModule::StartupModule()
 	RegisterConsoleCommands();
 	RegisterRoutes();
 
-	WebSocketRouter->AddPreDispatch([this](const struct FRemoteControlWebSocketMessage& Message)
-	{
-		const TArray<FString>* InPassphrase = Message.Header.Find(WebRemoteControlInternalUtils::PassphraseHeader);
-		const FString Passphrase = InPassphrase ? InPassphrase->Last() : FString("");
-		
-		const bool bCanBeDispatched = WebRemoteControlInternalUtils::CheckPassphrase(Passphrase);
-
-		if (!bCanBeDispatched)
-		{
-			TArray<uint8> Response;
-			FRCRequestWrapper Wrapper;
-			Wrapper.RequestId = Message.MessageId;
-			Wrapper.Passphrase = Passphrase;
-			Wrapper.Verb = "401";
-			
-			WebRemoteControlInternalUtils::CreateUTF8ErrorMessage(WebRemoteControlInternalUtils::InvalidPassphraseError, Wrapper.TCHARBody);
-			WebRemoteControlUtils::SerializeMessage(Wrapper, Response);
-
-			WebSocketServer.Send(Message.ClientId, MoveTemp(Response));
-		}
-
-		return bCanBeDispatched;
-	});
-
 	if (GetDefault<URemoteControlSettings>()->bAutoStartWebServer || CVarWebControlStartOnBoot.GetValueOnAnyThread() > 0)
 	{
 		StartHttpServer();
@@ -390,6 +366,9 @@ void FWebRemoteControlModule::StartupModule()
 	{
 		StartWebSocketServer();
 	}
+
+	RegisterDefaultPreprocessors();
+	RegisterExternalPreprocesors();
 }
 
 void FWebRemoteControlModule::ShutdownModule()
@@ -512,9 +491,6 @@ void FWebRemoteControlModule::StartHttpServer()
 		{
 			StartRoute(Route);
 		}
-
-		RegisterDefaultPreprocessors();
-		RegisterExternalPreprocesors();
 		
 		FHttpServerModule::Get().StartAllListeners();
 
@@ -2455,9 +2431,54 @@ void FWebRemoteControlModule::InvokeWrappedRequest(const FRCRequestWrapper& Wrap
 void FWebRemoteControlModule::RegisterDefaultPreprocessors()
 {
 	using namespace UE::WebRemoteControl;
-	auto RegisterInternalPreprocessor = [this](FRCPreprocessorHandler PreprocessorHandler)
+
+	auto MakeWebsocketPreDispatch = [this](FRCPreprocessorHandler PreprocessorHandler)
+	{
+		return [this, PreprocessorHandler](const FRemoteControlWebSocketMessage& Message) -> bool
+		{
+			FHttpServerRequest Request;
+			Request.Headers = Message.Header;
+			Request.PeerAddress = Message.PeerAddress;
+
+			FPreprocessorResult Result = PreprocessorHandler(Request);
+			if (Result.Result == EPreprocessorResult::RequestPassthrough)
+			{
+				return true;
+			}
+			else
+			{
+				TArray<uint8> Response;
+
+				const TArray<FString>* InPassphrase = Message.Header.Find(WebRemoteControlInternalUtils::PassphraseHeader);
+				const FString Passphrase = InPassphrase ? InPassphrase->Last() : FString("");
+
+				FRCRequestWrapper Wrapper;
+				Wrapper.RequestId = Message.MessageId;
+				Wrapper.Passphrase = Passphrase;
+				Wrapper.Verb = "401";
+
+				// Not ideal, we're re-converting to tchar, this should be removed once the WebRC pipeline is cleaned up.
+				WebRemoteControlUtils::ConvertToTCHAR(Result.OptionalResponse->Body, Wrapper.TCHARBody);
+				WebRemoteControlUtils::SerializeMessage(Wrapper, Response);
+
+				WebSocketServer.Send(Message.ClientId, MoveTemp(Response));
+
+				return false;
+			}
+		};
+	};
+
+	auto RegisterInternalPreprocessor = [this, MakeWebsocketPreDispatch](FRCPreprocessorHandler PreprocessorHandler)
 	{
-		AllRegisteredPreprocessorHandlers.Add(HttpRouter->RegisterRequestPreprocessor(MakeHttpRequestHandler(MoveTemp(PreprocessorHandler))));
+		if (HttpRouter)
+		{
+			AllRegisteredPreprocessorHandlers.Add(HttpRouter->RegisterRequestPreprocessor(MakeHttpRequestHandler(PreprocessorHandler)));
+		}
+
+		if (WebSocketRouter)
+		{
+			WebSocketRouter->AddPreDispatch(MakeWebsocketPreDispatch(PreprocessorHandler));
+		}
 	};
 
 	RegisterInternalPreprocessor(&RemotePassphraseEnforcementPreprocessor);
@@ -2478,19 +2499,22 @@ void FWebRemoteControlModule::UnregisterAllPreprocessors()
 
 void FWebRemoteControlModule::RegisterExternalPreprocesors()
 {
-	for (const TPair<FDelegateHandle, FHttpRequestHandler>& Handler : PreprocessorsToRegister)
+	if (HttpRouter)
 	{
-		// Find the pre-processors HTTP-handle from the one we generated.
-		FDelegateHandle& Handle = PreprocessorsHandleMappings.FindChecked(Handler.Key);
-		if (Handle.IsValid())
+		for (const TPair<FDelegateHandle, FHttpRequestHandler>& Handler : PreprocessorsToRegister)
 		{
-			HttpRouter->UnregisterRequestPreprocessor(Handle);
-			AllRegisteredPreprocessorHandlers.RemoveAtSwap(AllRegisteredPreprocessorHandlers.IndexOfByKey(Handle));
-		}
+			// Find the pre-processors HTTP-handle from the one we generated.
+			FDelegateHandle& Handle = PreprocessorsHandleMappings.FindChecked(Handler.Key);
+			if (Handle.IsValid())
+			{
+				HttpRouter->UnregisterRequestPreprocessor(Handle);
+				AllRegisteredPreprocessorHandlers.RemoveAtSwap(AllRegisteredPreprocessorHandlers.IndexOfByKey(Handle));
+			}
 
-		// Update the preprocessor handle mapping.
-		Handle = HttpRouter->RegisterRequestPreprocessor(Handler.Value);
-		AllRegisteredPreprocessorHandlers.Add(Handle);
+			// Update the preprocessor handle mapping.
+			Handle = HttpRouter->RegisterRequestPreprocessor(Handler.Value);
+			AllRegisteredPreprocessorHandlers.Add(Handle);
+		}
 	}
 }
 
@@ -2516,20 +2540,24 @@ void FWebRemoteControlModule::OnSettingsModified(UObject* Settings, FPropertyCha
 	const bool bRestartHttpServer = RCSettings->RemoteControlHttpServerPort != HttpServerPort;
 	const bool bRestartWebSocketServer = RCSettings->RemoteControlWebSocketServerPort != WebSocketServerPort;
 
-	if ((bIsWebServerStarted && bRestartHttpServer)
-		|| (!bIsWebServerStarted && RCSettings->bAutoStartWebServer))
+	if (PropertyChangedEvent.ChangeType != EPropertyChangeType::Interactive)
 	{
-		HttpServerPort = RCSettings->RemoteControlHttpServerPort;
-		StopHttpServer();
-		StartHttpServer();
-	}
 
-	if ((bIsWebSocketServerStarted && bRestartWebSocketServer)
-		|| (!bIsWebSocketServerStarted && RCSettings->bAutoStartWebSocketServer))
-	{
-		WebSocketServerPort = RCSettings->RemoteControlWebSocketServerPort;
-		StopWebSocketServer();
-		StartWebSocketServer();
+		if ((bIsWebServerStarted && bRestartHttpServer)
+			|| (!bIsWebServerStarted && RCSettings->bAutoStartWebServer))
+		{
+			HttpServerPort = RCSettings->RemoteControlHttpServerPort;
+			StopHttpServer();
+			StartHttpServer();
+		}
+
+		if ((bIsWebSocketServerStarted && bRestartWebSocketServer)
+			|| (!bIsWebSocketServerStarted && RCSettings->bAutoStartWebSocketServer))
+		{
+			WebSocketServerPort = RCSettings->RemoteControlWebSocketServerPort;
+			StopWebSocketServer();
+			StartWebSocketServer();
+		}
 	}
 
 	/** Letting the Server know what the current state of the Passphrase Usage is. */
diff --git a/Engine/Plugins/VirtualProduction/RemoteControl/Source/WebRemoteControl/Public/RemoteControlWebsocketRoute.h b/Engine/Plugins/VirtualProduction/RemoteControl/Source/WebRemoteControl/Public/RemoteControlWebsocketRoute.h
--- a/Engine/Plugins/VirtualProduction/RemoteControl/Source/WebRemoteControl/Public/RemoteControlWebsocketRoute.h
+++ b/Engine/Plugins/VirtualProduction/RemoteControl/Source/WebRemoteControl/Public/RemoteControlWebsocketRoute.h
@@ -9,6 +9,7 @@ struct FRemoteControlWebSocketMessage
 	FString MessageName;
 	int32 MessageId = -1;
 	FGuid ClientId;
+	TSharedPtr<FInternetAddr> PeerAddress;
 	TArrayView<uint8> RequestPayload;
 	TMap<FString, TArray<FString>> Header;
 };
