Fix CIS issue with remote control

diff --git a/Engine/Plugins/VirtualProduction/RemoteControl/Source/WebRemoteControl/Public/RemoteControlWebsocketRoute.h b/Engine/Plugins/VirtualProduction/RemoteControl/Source/WebRemoteControl/Public/RemoteControlWebsocketRoute.h
--- a/Engine/Plugins/VirtualProduction/RemoteControl/Source/WebRemoteControl/Public/RemoteControlWebsocketRoute.h
+++ b/Engine/Plugins/VirtualProduction/RemoteControl/Source/WebRemoteControl/Public/RemoteControlWebsocketRoute.h
@@ -9,7 +9,7 @@ struct FRemoteControlWebSocketMessage
 	FString MessageName;
 	int32 MessageId = -1;
 	FGuid ClientId;
-	TSharedPtr<FInternetAddr> PeerAddress;
+	TSharedPtr<class FInternetAddr> PeerAddress;
 	TArrayView<uint8> RequestPayload;
 	TMap<FString, TArray<FString>> Header;
 };
