Fix C++20 compile error. Some types in RemoteControl are can no longer be aggregate initialized in C++20.

diff --git a/Engine/Plugins/VirtualProduction/RemoteControl/Source/RemoteControlCommon/Public/RemoteControlSettings.h b/Engine/Plugins/VirtualProduction/RemoteControl/Source/RemoteControlCommon/Public/RemoteControlSettings.h
--- a/Engine/Plugins/VirtualProduction/RemoteControl/Source/RemoteControlCommon/Public/RemoteControlSettings.h
+++ b/Engine/Plugins/VirtualProduction/RemoteControl/Source/RemoteControlCommon/Public/RemoteControlSettings.h
@@ -16,6 +16,9 @@ struct FRCNetworkAddress
 	GENERATED_BODY()
 
 	FRCNetworkAddress() = default;
+	FRCNetworkAddress(uint8 InClassA, uint8 InClassB, uint8 InClassC, uint8 InClassD)
+		: ClassA(InClassA), ClassB(InClassB), ClassC(InClassC), ClassD(InClassD)
+	{}
 
 	bool operator==(const FRCNetworkAddress& OtherNetworkAddress) const
 	{
@@ -49,19 +52,19 @@ struct FRCNetworkAddress
 
 	/** Denotes the first octet of the IPv4 address (0-255.xxx.xxx.xxx) */
 	UPROPERTY(EditAnywhere, Category="Network Address")
-		uint8 ClassA = 192;
+	uint8 ClassA = 192;
 
 	/** Denotes the second octet of the IPv4 address (xxx.0-255.xxx.xxx) */
 	UPROPERTY(EditAnywhere, Category = "Network Address")
-		uint8 ClassB = 168;
+	uint8 ClassB = 168;
 
 	/** Denotes the third octet of the IPv4 address (xxx.xxx.0-255.xxx) */
 	UPROPERTY(EditAnywhere, Category = "Network Address")
-		uint8 ClassC = 1;
+	uint8 ClassC = 1;
 
 	/** Denotes the fourth octet of the IPv4 address (xxx.xxx.xxx.0-255) */
 	UPROPERTY(EditAnywhere, Category = "Network Address")
-		uint8 ClassD = 1;
+	uint8 ClassD = 1;
 };
 
 /**
@@ -73,6 +76,9 @@ struct FRCNetworkAddressRange
 	GENERATED_BODY()
 
 	FRCNetworkAddressRange() = default;
+	FRCNetworkAddressRange(FRCNetworkAddress InLowerBound, FRCNetworkAddress InUpperBound)
+		: LowerBound(InLowerBound), UpperBound(InUpperBound)
+	{}
 
 	bool operator==(const FRCNetworkAddressRange& OtherNetworkAddressRange) const
 	{
@@ -135,11 +141,11 @@ public:
 
 	/** Denotes the lower bound IPv4 address. */
 	UPROPERTY(EditAnywhere, Category="Network Address Range")
-		FRCNetworkAddress LowerBound = { 192, 168, 1, 1 };
+	FRCNetworkAddress LowerBound = { 192, 168, 1, 1 };
 
 	/** Denotes the upper bound IPv4 address. */
 	UPROPERTY(EditAnywhere, Category = "Network Address Range")
-		FRCNetworkAddress UpperBound = { 192, 168, 255, 255 };
+	FRCNetworkAddress UpperBound = { 192, 168, 255, 255 };
 };
 
 /**
@@ -216,10 +222,10 @@ public:
 
 		if (InClientAddressStr.ParseIntoArray(IndividualClasses, TEXT(".")) == 4)
 		{
-			FRCNetworkAddress LowerAndUpperBounds { FCString::Atoi(*IndividualClasses[0]) // Class A
-				, FCString::Atoi(*IndividualClasses[1]) // Class B
-				, FCString::Atoi(*IndividualClasses[2]) // Class C
-				, FCString::Atoi(*IndividualClasses[3]) // Class D
+			FRCNetworkAddress LowerAndUpperBounds { (uint8)FCString::Atoi(*IndividualClasses[0]) // Class A
+				, (uint8)FCString::Atoi(*IndividualClasses[1]) // Class B
+				, (uint8)FCString::Atoi(*IndividualClasses[2]) // Class C
+				, (uint8)FCString::Atoi(*IndividualClasses[3]) // Class D
 			};
 
 			FRCNetworkAddressRange NewRange = { LowerAndUpperBounds, LowerAndUpperBounds };
diff --git a/Engine/Plugins/VirtualProduction/RemoteControl/Source/WebRemoteControl/Private/RemoteControlDefaultPreprocessors.h b/Engine/Plugins/VirtualProduction/RemoteControl/Source/WebRemoteControl/Private/RemoteControlDefaultPreprocessors.h
--- a/Engine/Plugins/VirtualProduction/RemoteControl/Source/WebRemoteControl/Private/RemoteControlDefaultPreprocessors.h
+++ b/Engine/Plugins/VirtualProduction/RemoteControl/Source/WebRemoteControl/Private/RemoteControlDefaultPreprocessors.h
@@ -66,13 +66,16 @@ namespace UE::WebRemoteControl
 			return FPreprocessorResult{ EPreprocessorResult::RequestHandled, MoveTemp(Response) };
 		}
 
-	/** Holds the preprocessor result. */
-	EPreprocessorResult Result = EPreprocessorResult::RequestPassthrough;
-	/** If denied, holds the response to be sent to the client. */
-	TUniquePtr<FHttpServerResponse> OptionalResponse;
+		/** Holds the preprocessor result. */
+		EPreprocessorResult Result = EPreprocessorResult::RequestPassthrough;
+		/** If denied, holds the response to be sent to the client. */
+		TUniquePtr<FHttpServerResponse> OptionalResponse;
 
 	private:
 		FPreprocessorResult() = default;
+		FPreprocessorResult(EPreprocessorResult InResult, TUniquePtr<FHttpServerResponse> InOptionalResponse)
+			: Result(InResult), OptionalResponse(MoveTemp(InOptionalResponse))
+		{}
 	};
 	
 	using FRCPreprocessorHandler = TFunction<FPreprocessorResult(const FHttpServerRequest& Request)>;
