Remote Control Security
Unify Allowed IP and AllowListed IP settings
Unify Use Passphrase and Allow passphrase

diff --git a/Engine/Plugins/VirtualProduction/RemoteControl/Source/RemoteControl/Private/RemoteControlModule.cpp b/Engine/Plugins/VirtualProduction/RemoteControl/Source/RemoteControl/Private/RemoteControlModule.cpp
--- a/Engine/Plugins/VirtualProduction/RemoteControl/Source/RemoteControl/Private/RemoteControlModule.cpp
+++ b/Engine/Plugins/VirtualProduction/RemoteControl/Source/RemoteControl/Private/RemoteControlModule.cpp
@@ -82,6 +82,8 @@ namespace RemoteControlUtil
 
 	UFunction* FindFunctionByNameOrMetaDataName(UObject* Object, const FString& FunctionName)
 	{
+		TRACE_CPUPROFILER_EVENT_SCOPE(RemoteControl::FindFunctionByNameOrMetaDataName);
+
 		UFunction* Function = Object->FindFunction(FName(*FunctionName));
 #if WITH_EDITOR
 		// if the function wasn't found through the function map, try finding it through its `ScriptName` or `DisplayName` metadata
diff --git a/Engine/Plugins/VirtualProduction/RemoteControl/Source/RemoteControlCommon/Public/RemoteControlSettings.h b/Engine/Plugins/VirtualProduction/RemoteControl/Source/RemoteControlCommon/Public/RemoteControlSettings.h
--- a/Engine/Plugins/VirtualProduction/RemoteControl/Source/RemoteControlCommon/Public/RemoteControlSettings.h
+++ b/Engine/Plugins/VirtualProduction/RemoteControl/Source/RemoteControlCommon/Public/RemoteControlSettings.h
@@ -80,6 +80,14 @@ struct FRCNetworkAddressRange
 		: LowerBound(InLowerBound), UpperBound(InUpperBound)
 	{}
 
+	static FRCNetworkAddressRange AllowAllIPs()
+	{
+		FRCNetworkAddressRange Range;
+		Range.LowerBound = { 0, 0, 0, 0 };
+		Range.UpperBound = { 255, 255, 255, 255 };
+		return Range;
+	}
+
 	bool operator==(const FRCNetworkAddressRange& OtherNetworkAddressRange) const
 	{
 		return LowerBound == OtherNetworkAddressRange.LowerBound
@@ -127,6 +135,7 @@ struct FRCNetworkAddressRange
 
 private:
 
+	UE_DISABLE_OPTIMIZATION
 	bool IsInRange_Internal(uint8 InClassA, uint8 InClassB, uint8 InClassC, uint8 InClassD) const
 	{
 		bool bClassAIsInRange = InClassA >= LowerBound.ClassA && InClassA <= UpperBound.ClassA;
@@ -136,6 +145,7 @@ private:
 		
 		return bClassAIsInRange && bClassBIsInRange && bClassCIsInRange && bClassDIsInRange;
 	}
+	UE_ENABLE_OPTIMIZATION
 
 public:
 
@@ -306,48 +316,36 @@ public:
 	/** Whether to restrict access to a list of hostname/IPs in the AllowedOrigins setting. */
 	UPROPERTY(Config, EditAnywhere, Category = "Remote Control | Security")
 	bool bRestrictServerAccess = false;
-	
-	/** Whether communication with the Web Interface should only be allowed with an Passphrase */
-	UPROPERTY(Config, EditAnywhere, Category = "Remote Control | Security", DisplayName = "Use Passphrase to block Access")
-	bool bUseRemoteControlPassphrase = false;
-
-	/** Whether the User should be warned that Passphrase usage is disabled or now. Initially activated */
-	UPROPERTY(Config, EditAnywhere, Category = "Remote Control | Security", DisplayName = "Warn that Passphrase might be disabled ")
-	bool bShowPassphraseDisabledWarning = true;
 
 	/** Enable remote python execution, enabling this could open you open to vulnerabilities if an outside actor has access to your server. */
-	UPROPERTY(Config, EditAnywhere, Category = "Remote Control | Security")
+	UPROPERTY(Config, EditAnywhere, Category = "Remote Control | Security", meta = (editCondition = bRestrictServerAccess))
 	bool bEnableRemotePythonExecution = false;
 
-	/** List of IP Addresses that are allowed to access the Web API. */
-	UPROPERTY(config, EditAnywhere, Category = "Remote Control | Security", DisplayName = "Range of Allowlisted Clients", Meta = (EditCondition = "bRestrictServerAccess", EditConditionHides))
-	TSet<FRCNetworkAddressRange> AllowlistedClients;
+	/** List of IP Addresses that are allowed to access the Web API without authentication. */
+	UPROPERTY(config, EditAnywhere, Category = "Remote Control | Security", DisplayName = "Range of Allowlisted Clients", Meta = (EditCondition = bRestrictServerAccess, EditConditionHides))
+	TSet<FRCNetworkAddressRange> AllowlistedClients = { FRCNetworkAddressRange::AllowAllIPs() };
 
 	/** 
 	 * Origin that can make requests to the remote control server. Should contain the hostname or IP of the server making requests to remote control. ie. "http://yourdomain.com", or "*" to allow all origins. 
-	 * @Note: This is used to block requests coming from a browser (ie. Coming from a script running on a website), ideally you should use both this setting and AllowedIPs, as a request coming from a websocket client can have an empty Origin.
+	 * @Note: This is used to block requests coming from a browser (ie. Coming from a script running on a website), ideally you should use both this setting and AllowListedClients, as a request coming from a websocket client can have an empty Origin.
 	 * @Note Supports wildcards (ie. *.yourdomain.com)
 	 */
-	UPROPERTY(Config, EditAnywhere, Category = "Remote Control | Security", meta=(EditCondition = bRestrictServerAccess))
+	UPROPERTY(Config, EditAnywhere, Category = "Remote Control | Security", meta=(EditCondition=bRestrictServerAccess))
 	FString AllowedOrigin = TEXT("*");
-	
-	/** 
-	 * What IP is allowed to make request to the RemoteControl and RemoteControl Websocket servers.
-	 * @Note If empty or *.*.*.*,  all IP addresses will be able to make requests to your servers if they are on your network, so consider limiting this to a range of IPs that you own.
-	 * @Note Using this setting without AllowedOrigin can potentially open you up to malicous requests coming from a website, as the request would come from localhost.
-	 * @Note Supports wildcards (ie. 202.120.*.*)
-	 */
-	UPROPERTY(Config, EditAnywhere, Category = "Remote Control | Security", meta = (EditCondition = bRestrictServerAccess))
-	FString AllowedIP = TEXT("127.0.0.1");
 
 	/**
 	 * Controls whether a passphrase should be required when remote control is accessed by a client outside of localhost.
 	 */
-	UPROPERTY(config, EditAnywhere, Category = "Remote Control | Security")
+	UPROPERTY(config, EditAnywhere, Category = "Remote Control | Security", meta = (EditCondition = bRestrictServerAccess))
 	bool bEnforcePassphraseForRemoteClients = false;
-	
-	UPROPERTY(config, EditAnywhere, Category = "Remote Control | Security", DisplayName = "Remote Control Passphrase")
+
+	/** List of passphrases used for accessing remote control outside of localhost. */
+	UPROPERTY(config, EditAnywhere, Category = "Remote Control | Security", DisplayName = "Remote Control Passphrase", meta = (EditCondition=bEnforcePassphraseForRemoteClients, EditConditionHides))
 	TArray<FRCPassphrase> Passphrases = {};
+	
+	/** Whether the User should be warned that Passphrase usage is disabled or now. Initially activated */
+	UPROPERTY(Config, EditAnywhere, Category = "Remote Control | Security", DisplayName = "Warn that Passphrase might be disabled ", meta=(EditCondition=bRestrictServerAccess))
+	bool bShowPassphraseDisabledWarning = true;
 
 private:
 	UPROPERTY(config)
diff --git a/Engine/Plugins/VirtualProduction/RemoteControl/Source/WebRemoteControl/Private/RemoteControlDefaultPreprocessors.h b/Engine/Plugins/VirtualProduction/RemoteControl/Source/WebRemoteControl/Private/RemoteControlDefaultPreprocessors.h
--- a/Engine/Plugins/VirtualProduction/RemoteControl/Source/WebRemoteControl/Private/RemoteControlDefaultPreprocessors.h
+++ b/Engine/Plugins/VirtualProduction/RemoteControl/Source/WebRemoteControl/Private/RemoteControlDefaultPreprocessors.h
@@ -488,7 +488,6 @@ namespace UE::WebRemoteControl
 			{
 				constexpr bool bAppendPort = false;
 				FString ClientIP = Request.PeerAddress->ToString(bAppendPort);
-				const FWildcardString WildcardAllowedIP = SimplifyAddress(GetDefault<URemoteControlSettings>()->AllowedIP);
 				
 				// Allow requests from localhost
 				if (ClientIP != TEXT("localhost") && ClientIP != TEXT("127.0.0.1"))
@@ -497,11 +496,6 @@ namespace UE::WebRemoteControl
 					{
 						return FPreprocessorResult::Deny(FString::Printf(TEXT("Client IP %s does not respect the allowed IP set in Remote Control Settings."), *ClientIP));
 					}
-
-					if (!WildcardAllowedIP.IsEmpty() && WildcardAllowedIP.IsMatch(ClientIP))
-					{
-						return FPreprocessorResult::Deny(FString::Printf(TEXT("Client IP %s does not respect the allowed IP set in Remote Control Settings."), *ClientIP));
-					}
 				}
 			}
 		}
diff --git a/Engine/Plugins/VirtualProduction/RemoteControl/Source/WebRemoteControl/Private/RemoteControlWebSocketServer.cpp b/Engine/Plugins/VirtualProduction/RemoteControl/Source/WebRemoteControl/Private/RemoteControlWebSocketServer.cpp
--- a/Engine/Plugins/VirtualProduction/RemoteControl/Source/WebRemoteControl/Private/RemoteControlWebSocketServer.cpp
+++ b/Engine/Plugins/VirtualProduction/RemoteControl/Source/WebRemoteControl/Private/RemoteControlWebSocketServer.cpp
@@ -297,12 +297,6 @@ EWebsocketConnectionFilterResult FRCWebSocketServer::FilterConnection(FString Or
 				{
 					return EWebsocketConnectionFilterResult::ConnectionRefused;
 				}
-
-				const FWildcardString WildcardAllowedIP = SimplifyAddress(Settings->AllowedIP);
-				if (!WildcardAllowedIP.IsEmpty() && WildcardAllowedIP.IsMatch(ClientIP))
-				{
-					return EWebsocketConnectionFilterResult::ConnectionRefused;
-				}
 			}
 		}
 	}
diff --git a/Engine/Plugins/VirtualProduction/RemoteControl/Source/WebRemoteControl/Private/WebRemoteControl.cpp b/Engine/Plugins/VirtualProduction/RemoteControl/Source/WebRemoteControl/Private/WebRemoteControl.cpp
--- a/Engine/Plugins/VirtualProduction/RemoteControl/Source/WebRemoteControl/Private/WebRemoteControl.cpp
+++ b/Engine/Plugins/VirtualProduction/RemoteControl/Source/WebRemoteControl/Private/WebRemoteControl.cpp
@@ -2563,7 +2563,7 @@ void FWebRemoteControlModule::OnSettingsModified(UObject* Settings, FPropertyCha
 	/** Letting the Server know what the current state of the Passphrase Usage is. */
 	if (bIsWebSocketServerStarted && bIsWebServerStarted)
 	{
-		const bool bIsOpen = RCSettings->bUseRemoteControlPassphrase;
+		const bool bIsOpen = RCSettings->bEnforcePassphraseForRemoteClients;
 		
 		TArray<uint8> Response;
 		const FCheckPassphraseResponse BoolResponse = FCheckPassphraseResponse(!bIsOpen);
diff --git a/Engine/Plugins/VirtualProduction/RemoteControl/Source/WebRemoteControl/Private/WebRemoteControlInternalUtils.cpp b/Engine/Plugins/VirtualProduction/RemoteControl/Source/WebRemoteControl/Private/WebRemoteControlInternalUtils.cpp
--- a/Engine/Plugins/VirtualProduction/RemoteControl/Source/WebRemoteControl/Private/WebRemoteControlInternalUtils.cpp
+++ b/Engine/Plugins/VirtualProduction/RemoteControl/Source/WebRemoteControl/Private/WebRemoteControlInternalUtils.cpp
@@ -603,7 +603,7 @@ bool WebRemoteControlInternalUtils::IsRequestContentType(const FHttpServerReques
 
 bool WebRemoteControlInternalUtils::CheckPassphrase(const FString& HashedPassphrase)
 {
-	bool bOutResult = !(GetDefault<URemoteControlSettings>()->bUseRemoteControlPassphrase);
+	bool bOutResult = !(GetDefault<URemoteControlSettings>()->bEnforcePassphraseForRemoteClients) || !(GetDefault<URemoteControlSettings>()->bRestrictServerAccess);
 
 	if (bOutResult)
 	{
diff --git a/Engine/Plugins/VirtualProduction/RemoteControlWebInterface/Source/RemoteControlWebInterface/Private/RCWebInterfaceCustomizations.cpp b/Engine/Plugins/VirtualProduction/RemoteControlWebInterface/Source/RemoteControlWebInterface/Private/RCWebInterfaceCustomizations.cpp
--- a/Engine/Plugins/VirtualProduction/RemoteControlWebInterface/Source/RemoteControlWebInterface/Private/RCWebInterfaceCustomizations.cpp
+++ b/Engine/Plugins/VirtualProduction/RemoteControlWebInterface/Source/RemoteControlWebInterface/Private/RCWebInterfaceCustomizations.cpp
@@ -214,7 +214,7 @@ void FRCWebInterfaceCustomizations::GeneratePanelExtensions(TArray<TSharedRef<SW
 		.Visibility_Lambda([]()
 		{
 			const URemoteControlSettings* RCSettings = GetDefault<URemoteControlSettings>();
-			bool bShouldWarn = RCSettings->bShowPassphraseDisabledWarning && !RCSettings->bUseRemoteControlPassphrase;
+			bool bShouldWarn = RCSettings->bShowPassphraseDisabledWarning && !RCSettings->bEnforcePassphraseForRemoteClients;
 			return bShouldWarn ? EVisibility::Visible : EVisibility::Collapsed;
 		})
 		.OnClicked_Lambda([]()
