Fix CIS error caused by submission of UE_DISABLE_OPTIMIZATION

diff --git a/Engine/Plugins/VirtualProduction/RemoteControl/Source/RemoteControlCommon/Public/RemoteControlSettings.h b/Engine/Plugins/VirtualProduction/RemoteControl/Source/RemoteControlCommon/Public/RemoteControlSettings.h
--- a/Engine/Plugins/VirtualProduction/RemoteControl/Source/RemoteControlCommon/Public/RemoteControlSettings.h
+++ b/Engine/Plugins/VirtualProduction/RemoteControl/Source/RemoteControlCommon/Public/RemoteControlSettings.h
@@ -135,7 +135,6 @@ struct FRCNetworkAddressRange
 
 private:
 
-	UE_DISABLE_OPTIMIZATION
 	bool IsInRange_Internal(uint8 InClassA, uint8 InClassB, uint8 InClassC, uint8 InClassD) const
 	{
 		bool bClassAIsInRange = InClassA >= LowerBound.ClassA && InClassA <= UpperBound.ClassA;
@@ -145,7 +144,6 @@ private:
 		
 		return bClassAIsInRange && bClassBIsInRange && bClassCIsInRange && bClassDIsInRange;
 	}
-	UE_ENABLE_OPTIMIZATION
 
 public:
 
