Remote Control - Right clicking exposed properties header column shows empty entries

diff --git a/Engine/Plugins/VirtualProduction/RemoteControl/Source/RemoteControlUI/Private/UI/SRCHeaderRow.cpp b/Engine/Plugins/VirtualProduction/RemoteControl/Source/RemoteControlUI/Private/UI/SRCHeaderRow.cpp
new file mode 100644
--- /dev/null
+++ b/Engine/Plugins/VirtualProduction/RemoteControl/Source/RemoteControlUI/Private/UI/SRCHeaderRow.cpp
@@ -0,0 +1,126 @@
+// Copyright Epic Games, Inc. All Rights Reserved.
+
+#include "SRCHeaderRow.h"
+#include "Framework/Application/SlateApplication.h"
+#include "Framework/MultiBox/MultiBoxBuilder.h"
+#include "Layout/WidgetPath.h"
+
+void SRCHeaderRow::Construct(const FArguments& InArgs)
+{
+	bOverrideCanSelectGeneratedColumn = InArgs._CanSelectGeneratedColumn;
+
+	TArray<SHeaderRow::FColumn*> Slots;
+	for (FColumn* const Column : InArgs.Slots)
+	{
+		Columns.Add(Column);
+		SHeaderRow::FColumn* HeaderRowSlot = CreateHeaderRowColumn(*Column);
+		Slots.Add(HeaderRowSlot);
+	}
+
+	SHeaderRow::FArguments HeaderRowArgs;
+	HeaderRowArgs._Style = InArgs._Style;
+	HeaderRowArgs.Slots = Slots;
+	HeaderRowArgs._CanSelectGeneratedColumn = InArgs._CanSelectGeneratedColumn;
+	SHeaderRow::Construct(HeaderRowArgs);
+}
+
+FReply SRCHeaderRow::OnMouseButtonUp(const FGeometry& MyGeometry, const FPointerEvent& MouseEvent)
+{
+	if (bOverrideCanSelectGeneratedColumn && MouseEvent.GetEffectingButton() == EKeys::RightMouseButton)
+	{
+		FVector2f SummonLocation = MouseEvent.GetScreenSpacePosition();
+		FWidgetPath WidgetPath = MouseEvent.GetEventPath() != nullptr ? *MouseEvent.GetEventPath() : FWidgetPath();
+
+		const bool CloseAfterSelection = true;
+		FMenuBuilder MenuBuilder(CloseAfterSelection, nullptr);
+		OnOverrideGenerateSelectColumnsSubMenu(MenuBuilder);
+
+		FSlateApplication::Get().CloseToolTip();
+		FSlateApplication::Get().PushMenu(AsShared(), WidgetPath, MenuBuilder.MakeWidget(), SummonLocation, FPopupTransitionEffect(FPopupTransitionEffect::ContextMenu));
+		return FReply::Handled();
+	}
+
+	return FReply::Unhandled();
+}
+
+SHeaderRow::FColumn* SRCHeaderRow::CreateHeaderRowColumn(const FColumn InColumn)
+{
+	SHeaderRow::FColumn::FArguments SlotsArgument;
+	SlotsArgument._ColumnId = InColumn.ColumnId;
+	SlotsArgument._DefaultLabel = InColumn.DefaultText;
+	SlotsArgument._HAlignHeader = InColumn.HeaderHAlignment;
+	SlotsArgument._VAlignHeader = InColumn.HeaderVAlignment;
+	SlotsArgument._HAlignCell = InColumn.CellHAlignment;
+	SlotsArgument._VAlignCell = InColumn.CellVAlignment;
+	SlotsArgument._HeaderContentPadding = InColumn.HeaderContentPadding;
+	SlotsArgument._ShouldGenerateWidget = InColumn.ShouldGenerateWidget;
+
+	SHeaderRow::FColumn* HeaderRowSlot = new SHeaderRow::FColumn(SlotsArgument);
+
+	HeaderRowSlot->Width = InColumn.DefaultWidth;
+	HeaderRowSlot->DefaultWidth = InColumn.DefaultWidth;
+	HeaderRowSlot->SizeRule = InColumn.SizeRule;
+	HeaderRowSlot->bIsVisible = InColumn.bIsVisible;
+	return HeaderRowSlot;
+}
+
+void SRCHeaderRow::RegenerateOverriddenWidgets()
+{
+	ClearColumns();
+
+	for (const FColumn& Column : Columns)
+	{
+		SHeaderRow::FColumn* HeaderRowSlot = CreateHeaderRowColumn(Column);
+		AddColumn(*HeaderRowSlot);
+	}
+
+	RefreshColumns();
+}
+
+void SRCHeaderRow::OnOverrideGenerateSelectColumnsSubMenu(FMenuBuilder& InSubMenuBuilder)
+{
+	for (const SRCHeaderRow::FColumn& SomeColumn : Columns)
+	{
+		if (SomeColumn.ShouldGenerateSubMenuEntry.IsSet() && !SomeColumn.ShouldGenerateSubMenuEntry.Get())
+		{
+			continue;
+		}
+
+		const bool bCanExecuteAction = !SomeColumn.ShouldGenerateWidget.IsSet();
+		const FName ColumnId = SomeColumn.ColumnId;
+
+		InSubMenuBuilder.AddMenuEntry(
+			SomeColumn.DefaultText,
+			FText::GetEmpty(),
+			FSlateIcon(),
+			FUIAction(
+				FExecuteAction::CreateSP(this, &SRCHeaderRow::ToggleOverriddenGeneratedColumn, ColumnId),
+				FCanExecuteAction::CreateLambda([bCanExecuteAction]() { return bCanExecuteAction; }),
+				FGetActionCheckState::CreateSP(this, &SRCHeaderRow::GetOverriddenGeneratedColumnCheckedState, ColumnId)
+			),
+			NAME_None,
+			EUserInterfaceActionType::ToggleButton);
+	}
+}
+
+void SRCHeaderRow::ToggleOverriddenGeneratedColumn(FName ColumnId)
+{
+	for (FColumn& SomeColumn : Columns)
+	{
+		if (SomeColumn.ColumnId == ColumnId)
+		{
+			if (!SomeColumn.ShouldGenerateWidget.IsSet())
+			{
+				SomeColumn.bIsVisible = !SomeColumn.bIsVisible;
+				RegenerateOverriddenWidgets();
+			}
+
+			break;
+		}
+	}
+}
+
+ECheckBoxState SRCHeaderRow::GetOverriddenGeneratedColumnCheckedState(FName ColumnId) const
+{
+	return IsColumnGenerated(ColumnId) ? ECheckBoxState::Checked : ECheckBoxState::Unchecked;
+}
diff --git a/Engine/Plugins/VirtualProduction/RemoteControl/Source/RemoteControlUI/Private/UI/SRCHeaderRow.h b/Engine/Plugins/VirtualProduction/RemoteControl/Source/RemoteControlUI/Private/UI/SRCHeaderRow.h
new file mode 100644
--- /dev/null
+++ b/Engine/Plugins/VirtualProduction/RemoteControl/Source/RemoteControlUI/Private/UI/SRCHeaderRow.h
@@ -0,0 +1,160 @@
+// Copyright Epic Games, Inc. All Rights Reserved.
+
+#pragma once
+
+#include "CoreMinimal.h"
+#include "Widgets/Views/SHeaderRow.h"
+
+class FMenuBuilder;
+enum class ECheckBoxState : uint8;
+
+/**
+ * A custom Header Row to hide entry in the SubMenu to Hide/Show the Header.
+ */
+class SRCHeaderRow : public SHeaderRow
+{
+public:
+
+	class FColumn
+	{
+	public:
+
+		SLATE_BEGIN_ARGS(FColumn)
+			: _ColumnId()
+			, _DefaultLabel()
+			, _FillWidth(1.f)
+			, _FixedWidth()
+			, _HAlignHeader(HAlign_Fill)
+			, _VAlignHeader(VAlign_Fill)
+			, _HeaderContentPadding()
+			, _HAlignCell(HAlign_Fill)
+			, _VAlignCell(VAlign_Fill)
+		{}
+
+			SLATE_ARGUMENT(FName, ColumnId)
+			SLATE_ATTRIBUTE(FText, DefaultLabel)
+			SLATE_ATTRIBUTE(float, FillWidth)
+			SLATE_ARGUMENT(TOptional<float>, FixedWidth)
+			SLATE_ARGUMENT(EHorizontalAlignment, HAlignHeader)
+			SLATE_ARGUMENT(EVerticalAlignment, VAlignHeader)
+			SLATE_ARGUMENT(TOptional<FMargin>, HeaderContentPadding)
+			SLATE_ARGUMENT(EHorizontalAlignment, HAlignCell)
+			SLATE_ARGUMENT(EVerticalAlignment, VAlignCell)
+			SLATE_ATTRIBUTE(bool, ShouldGenerateWidget)
+
+			/* If true the Entry is showed in the SubMenu to make the Column Visible or Hide it	*/
+			SLATE_ATTRIBUTE(bool, ShouldGenerateSubMenuEntry)
+
+		SLATE_END_ARGS()
+
+		FColumn(const FArguments& InArgs)
+			: ColumnId(InArgs._ColumnId)
+			, DefaultText(InArgs._DefaultLabel)
+			, Width(1.0f)
+			, DefaultWidth(1.0f)
+			, SizeRule(EColumnSizeMode::Fill)
+			, HeaderHAlignment(InArgs._HAlignHeader)
+			, HeaderVAlignment(InArgs._VAlignHeader)
+			, HeaderContentPadding(InArgs._HeaderContentPadding)
+			, CellHAlignment(InArgs._HAlignCell)
+			, CellVAlignment(InArgs._VAlignCell)
+			, ShouldGenerateWidget(InArgs._ShouldGenerateWidget)
+			, ShouldGenerateSubMenuEntry(InArgs._ShouldGenerateSubMenuEntry)
+			, bIsVisible(true)
+		{
+			if (InArgs._FixedWidth.IsSet())
+			{
+				Width = InArgs._FixedWidth.GetValue();
+				SizeRule = EColumnSizeMode::Fixed;
+			}
+			else
+			{
+				Width = InArgs._FillWidth;
+				SizeRule = EColumnSizeMode::Fill;
+			}
+
+			DefaultWidth = Width.Get();
+		}
+
+	public:
+
+		FName ColumnId;
+		TAttribute<FText> DefaultText;
+		TAttribute<float> Width;
+		float DefaultWidth;
+		EColumnSizeMode::Type SizeRule;
+		EHorizontalAlignment HeaderHAlignment;
+		EVerticalAlignment HeaderVAlignment;
+		TOptional<FMargin> HeaderContentPadding;
+		EHorizontalAlignment CellHAlignment;
+		EVerticalAlignment CellVAlignment;
+		TAttribute<bool> ShouldGenerateWidget;
+		TAttribute<bool> ShouldGenerateSubMenuEntry;
+		bool bIsVisible;
+	};
+
+	static FColumn::FArguments Column(const FName& InColumnId)
+	{
+		FColumn::FArguments NewArgs;
+		NewArgs.ColumnId(InColumnId);
+		return NewArgs;
+	}
+
+	SLATE_BEGIN_ARGS(SRCHeaderRow)
+		: _CanSelectGeneratedColumn(false)
+	{}
+
+		SLATE_STYLE_ARGUMENT(FHeaderRowStyle, Style)
+		SLATE_SUPPORTS_SLOT_WITH_ARGS(FColumn)
+		SLATE_ARGUMENT(bool, CanSelectGeneratedColumn)
+
+	SLATE_END_ARGS()
+
+	/**
+	 * Construct the widget with given InArgs.
+	 */
+	void Construct(const FArguments& InArgs);
+
+	//~ Begin SWidget interface
+	virtual FReply OnMouseButtonUp(const FGeometry& MyGeometry, const FPointerEvent& MouseEvent) override;
+	//~ End SWidget interface
+
+private:
+
+	/**
+	 * Creates and returns a customized column object.
+	 */
+	SHeaderRow::FColumn* CreateHeaderRowColumn(const FColumn InColumn);
+	
+	/**
+	 * Overridden method that is used to regenerate the custom column(s).
+	 */
+	void RegenerateOverriddenWidgets();
+
+	/**
+	 * Overridden method that is called when the user perform right click on custom column(s).
+	 */
+	void OnOverrideGenerateSelectColumnsSubMenu(FMenuBuilder& InSubMenuBuilder);
+
+	/**
+	 * Toggles the visibility of the overridden custom column(s).
+	 */
+	void ToggleOverriddenGeneratedColumn(FName ColumnId);
+
+	/**
+	 * Rretrieves the check box state of the overridden custom column(s).
+	 */
+	ECheckBoxState GetOverriddenGeneratedColumnCheckedState(FName ColumnId) const;
+
+private:
+
+	/**
+	 * Holds the array of overridden custom column(s).
+	 */
+	TIndirectArray<FColumn> Columns;
+	
+	/**
+	 * Determines whether the column can be shown in the context menu or not.
+	 */
+	bool bOverrideCanSelectGeneratedColumn;
+};
diff --git a/Engine/Plugins/VirtualProduction/RemoteControl/Source/RemoteControlUI/Private/UI/SRCPanelExposedEntitiesList.cpp b/Engine/Plugins/VirtualProduction/RemoteControl/Source/RemoteControlUI/Private/UI/SRCPanelExposedEntitiesList.cpp
--- a/Engine/Plugins/VirtualProduction/RemoteControl/Source/RemoteControlUI/Private/UI/SRCPanelExposedEntitiesList.cpp
+++ b/Engine/Plugins/VirtualProduction/RemoteControl/Source/RemoteControlUI/Private/UI/SRCPanelExposedEntitiesList.cpp
@@ -5,23 +5,24 @@
 #include "Algo/ForEach.h"
 #include "Commands/RemoteControlCommands.h"
 #include "Editor.h"
+#include "Editor/EditorEngine.h"
 #include "EditorFontGlyphs.h"
-#include "ISettingsModule.h"
-#include "RCPanelWidgetRegistry.h"
+#include "Engine/Selection.h"
 #include "GameFramework/Actor.h"
+#include "ISettingsModule.h"
 #include "Input/DragAndDrop.h"
 #include "IRemoteControlProtocolModule.h"
 #include "IRemoteControlProtocolWidgetsModule.h"
-#include "Engine/Selection.h"
-#include "Editor/EditorEngine.h"
 #include "Misc/Attribute.h"
 #include "Misc/Guid.h"
 #include "Misc/MessageDialog.h"
 #include "PropertyPath.h"
+#include "RCPanelWidgetRegistry.h"
 #include "RemoteControlPanelStyle.h"
 #include "RemoteControlPreset.h"
 #include "RemoteControlUIModule.h"
 #include "ScopedTransaction.h"
+#include "SRCHeaderRow.h"
 #include "SRCModeSwitcher.h"
 #include "SRCPanelFieldGroup.h"
 #include "SRCPanelExposedField.h"
@@ -281,45 +282,47 @@ void SRCPanelExposedEntitiesList::Construct(const FArguments& InArgs, URemoteCon
 		.TreeViewStyle(&RCPanelStyle->TableViewStyle)
 		.OnGetChildren(this, &SRCPanelExposedEntitiesList::OnGetNodeChildren)
 		.HeaderRow(
-			SAssignNew(FieldsHeaderRow, SHeaderRow)
+			SAssignNew(FieldsHeaderRow, SRCHeaderRow)
 			.Style(&RCPanelStyle->HeaderRowStyle)
 			.CanSelectGeneratedColumn(true) //To show/hide columns
 
-			+ SHeaderRow::Column(RemoteControlPresetColumns::DragDropHandle)
+			+ SRCHeaderRow::Column(RemoteControlPresetColumns::DragDropHandle)
 			.DefaultLabel(LOCTEXT("RCPresetDragDropHandleColumnHeader", ""))
 			.FixedWidth(25.f)
 			.HeaderContentPadding(RCPanelStyle->HeaderRowPadding)
 			.ShouldGenerateWidget(true)
+			.ShouldGenerateSubMenuEntry(false)
 
-			+ SHeaderRow::Column(RemoteControlPresetColumns::OwnerName)
+			+ SRCHeaderRow::Column(RemoteControlPresetColumns::OwnerName)
 			.DefaultLabel(LOCTEXT("RCPresetOwnerNameColumnHeader", "Owner Name"))
 			.HAlignHeader(HAlign_Center)
 			.FillWidth(0.15f)
 			.HeaderContentPadding(RCPanelStyle->HeaderRowPadding)
 
-			+ SHeaderRow::Column(RemoteControlPresetColumns::SubobjectPath)
+			+ SRCHeaderRow::Column(RemoteControlPresetColumns::SubobjectPath)
 			.DefaultLabel(LOCTEXT("RCPresetSubobjectPathColumnHeader", "Subobject Path"))
 			.HAlignHeader(HAlign_Center)
 			.FillWidth(0.15f)
 			.HeaderContentPadding(RCPanelStyle->HeaderRowPadding)
 
-			+ SHeaderRow::Column(RemoteControlPresetColumns::Description)
+			+ SRCHeaderRow::Column(RemoteControlPresetColumns::Description)
 			.DefaultLabel(LOCTEXT("RCPresetDescColumnHeader", "Description"))
 			.HAlignHeader(HAlign_Center)
 			.FillWidth(0.35f)
 			.HeaderContentPadding(RCPanelStyle->HeaderRowPadding)
 
-			+ SHeaderRow::Column(RemoteControlPresetColumns::Value)
+			+ SRCHeaderRow::Column(RemoteControlPresetColumns::Value)
 			.DefaultLabel(LOCTEXT("RCPresetValueColumnHeader", "Value"))
 			.HAlignHeader(HAlign_Center)
 			.FillWidth(0.4f)
 			.HeaderContentPadding(RCPanelStyle->HeaderRowPadding)
 
-			+ SHeaderRow::Column(RemoteControlPresetColumns::Reset)
+			+ SRCHeaderRow::Column(RemoteControlPresetColumns::Reset)
 			.DefaultLabel(LOCTEXT("RCPresetResetButtonColumnHeader", ""))
 			.FixedWidth(48.f)
 			.HeaderContentPadding(RCPanelStyle->HeaderRowPadding)
 			.ShouldGenerateWidget(true)
+			.ShouldGenerateSubMenuEntry(false)
 		);
 
 	// Exposed Entities Dock Panel
diff --git a/Engine/Plugins/VirtualProduction/RemoteControl/Source/RemoteControlUI/Private/UI/SRCPanelExposedEntitiesList.h b/Engine/Plugins/VirtualProduction/RemoteControl/Source/RemoteControlUI/Private/UI/SRCPanelExposedEntitiesList.h
--- a/Engine/Plugins/VirtualProduction/RemoteControl/Source/RemoteControlUI/Private/UI/SRCPanelExposedEntitiesList.h
+++ b/Engine/Plugins/VirtualProduction/RemoteControl/Source/RemoteControlUI/Private/UI/SRCPanelExposedEntitiesList.h
@@ -20,7 +20,7 @@ struct FRemoteControlProperty;
 struct FRemoteControlPresetGroup;
 struct FRemoteControlFunction;
 class ITableRow;
-class SHeaderRow;
+class SRCHeaderRow;
 class SRCPanelGroup;
 struct SRCPanelTreeNode;
 class SRemoteControlTarget;
@@ -229,7 +229,7 @@ private:
 	/** Panel Style reference. */
 	const FRCPanelStyle* RCPanelStyle;
 	/** Holds the header row of entities list. */
-	TSharedPtr<SHeaderRow> FieldsHeaderRow;
+	TSharedPtr<SRCHeaderRow> FieldsHeaderRow;
 	/** Holds the active list mode. */
 	EEntitiesListMode ActiveListMode;
 	/** Holds the active protocol enabled by the mode switcher. */
