Remote Control does not inherit material defaults

diff --git a/Engine/Plugins/VirtualProduction/RemoteControl/Source/RemoteControl/Private/Factories/RCDefaultValueFactories.cpp b/Engine/Plugins/VirtualProduction/RemoteControl/Source/RemoteControl/Private/Factories/RCDefaultValueFactories.cpp
--- a/Engine/Plugins/VirtualProduction/RemoteControl/Source/RemoteControl/Private/Factories/RCDefaultValueFactories.cpp
+++ b/Engine/Plugins/VirtualProduction/RemoteControl/Source/RemoteControl/Private/Factories/RCDefaultValueFactories.cpp
@@ -2,13 +2,16 @@
 
 #include "Factories/RCDefaultValueFactories.h"
 #include "Components/LightComponent.h"
+#include "Components/MeshComponent.h"
+#include "Components/SkinnedMeshComponent.h"
+#include "Components/StaticMeshComponent.h"
 
 TSharedRef<IRCDefaultValueFactory> FLightIntensityDefaultValueFactory::MakeInstance()
 {
 	return MakeShared<FLightIntensityDefaultValueFactory>();
 }
 
-bool FLightIntensityDefaultValueFactory::CanResetToDefaultValue(UObject* InObject) const
+bool FLightIntensityDefaultValueFactory::CanResetToDefaultValue(UObject* InObject, const FRCResetToDefaultArgs& InArgs) const
 {
 	if (const ULightComponent* LightComponent = Cast<ULightComponent>(InObject))
 	{
@@ -21,11 +24,13 @@ bool FLightIntensityDefaultValueFactory::CanResetToDefaultValue(UObject* InObjec
 	return false;
 }
 
-void FLightIntensityDefaultValueFactory::ResetToDefaultValue(UObject* InObject, FProperty* InProperty)
+void FLightIntensityDefaultValueFactory::ResetToDefaultValue(UObject* InObject, FRCResetToDefaultArgs& InArgs)
 {
 #if WITH_EDITOR
 	
-	InObject->PreEditChange(InProperty);
+	bool bReset = false;
+
+	InObject->PreEditChange(InArgs.Property);
 	InObject->Modify();
 
 	if (ULightComponent* LightComponent = Cast<ULightComponent>(InObject))
@@ -33,11 +38,15 @@ void FLightIntensityDefaultValueFactory::ResetToDefaultValue(UObject* InObject,
 		if (const ULightComponent* ArchetypeComponent = Cast<ULightComponent>(LightComponent->GetArchetype()))
 		{
 			LightComponent->SetLightBrightness(ArchetypeComponent->ComputeLightBrightness());
+			bReset = true;
 		}
 	}
 
-	FPropertyChangedEvent ChangeEvent(InProperty, EPropertyChangeType::ValueSet);
-	InObject->PostEditChangeProperty(ChangeEvent);
+	if (bReset)
+	{
+		FPropertyChangedEvent ChangeEvent(InArgs.Property, EPropertyChangeType::ValueSet);
+		InObject->PostEditChangeProperty(ChangeEvent);
+	}
 
 #endif // WITH_EDITOR
 }
@@ -61,3 +70,78 @@ bool FLightIntensityDefaultValueFactory::SupportsProperty(const FProperty* InPro
 
 	return InProperty->GetFName() == GET_MEMBER_NAME_CHECKED(ULightComponentBase, Intensity);
 }
+
+TSharedRef<IRCDefaultValueFactory> FOverrideMaterialsDefaultValueFactory::MakeInstance()
+{
+	return MakeShared<FOverrideMaterialsDefaultValueFactory>();
+}
+
+bool FOverrideMaterialsDefaultValueFactory::CanResetToDefaultValue(UObject* InObject, const FRCResetToDefaultArgs& InArgs) const
+{
+	return true;
+}
+
+void FOverrideMaterialsDefaultValueFactory::ResetToDefaultValue(UObject* InObject, FRCResetToDefaultArgs& InArgs)
+{
+#if WITH_EDITOR
+
+	bool bReset = false;
+	
+	InObject->PreEditChange(InArgs.Property);
+	InObject->Modify();
+
+	if (UMeshComponent* MeshComponent = Cast<UMeshComponent>(InObject))
+	{
+		if (UStaticMeshComponent* StaticMeshComponent = Cast<UStaticMeshComponent>(MeshComponent))
+		{
+			if (StaticMeshComponent->GetMaterial(InArgs.ArrayIndex))
+			{
+				StaticMeshComponent->SetMaterial(InArgs.ArrayIndex, NULL);
+				MeshComponent->SetMaterial(InArgs.ArrayIndex, StaticMeshComponent->GetMaterial(InArgs.ArrayIndex));
+				bReset = true;
+			}
+		}
+		else if (USkinnedMeshComponent* SkinnedMeshComponent = Cast<USkinnedMeshComponent>(MeshComponent))
+		{
+			if (SkinnedMeshComponent->GetMaterial(InArgs.ArrayIndex))
+			{
+				SkinnedMeshComponent->SetMaterial(InArgs.ArrayIndex, NULL);
+				MeshComponent->SetMaterial(InArgs.ArrayIndex, SkinnedMeshComponent->GetMaterial(InArgs.ArrayIndex));
+				bReset = true;
+			}
+		}
+	}
+
+	if (bReset)
+	{
+		FPropertyChangedEvent ChangeEvent(InArgs.Property, EPropertyChangeType::ValueSet);
+		
+		TArray<TMap<FString, int32>, TInlineAllocator<1>> ArrayIndexPerObject;
+		ArrayIndexPerObject.Add({ { InArgs.Property->GetName(), InArgs.ArrayIndex } });
+		
+		ChangeEvent.SetArrayIndexPerObject(ArrayIndexPerObject);
+		InObject->PostEditChangeProperty(ChangeEvent);
+	}
+
+#endif // WITH_EDITOR
+}
+
+bool FOverrideMaterialsDefaultValueFactory::SupportsClass(const UClass* InObjectClass) const
+{
+	if (!InObjectClass)
+	{
+		return false;
+	}
+
+	return InObjectClass->IsChildOf<UMeshComponent>();
+}
+
+bool FOverrideMaterialsDefaultValueFactory::SupportsProperty(const FProperty* InProperty) const
+{
+	if (!InProperty)
+	{
+		return false;
+	}
+
+	return InProperty->GetFName() == GET_MEMBER_NAME_CHECKED(UMeshComponent, OverrideMaterials);
+}
\ No newline at end of file
diff --git a/Engine/Plugins/VirtualProduction/RemoteControl/Source/RemoteControl/Private/Factories/RCDefaultValueFactories.h b/Engine/Plugins/VirtualProduction/RemoteControl/Source/RemoteControl/Private/Factories/RCDefaultValueFactories.h
--- a/Engine/Plugins/VirtualProduction/RemoteControl/Source/RemoteControl/Private/Factories/RCDefaultValueFactories.h
+++ b/Engine/Plugins/VirtualProduction/RemoteControl/Source/RemoteControl/Private/Factories/RCDefaultValueFactories.h
@@ -14,9 +14,22 @@ public:
 	static TSharedRef<IRCDefaultValueFactory> MakeInstance();
 
 	//~ Begin IRCDefaultValueFactory interface
-	virtual bool CanResetToDefaultValue(UObject* InObject) const override;
-	virtual void ResetToDefaultValue(UObject* InObject, FProperty* InProperty) override;
+	virtual bool CanResetToDefaultValue(UObject* InObject, const FRCResetToDefaultArgs& InArgs) const override;
+	virtual void ResetToDefaultValue(UObject* InObject, FRCResetToDefaultArgs& InArgs) override;
 	virtual bool SupportsClass(const UClass* InObjectClass) const override;
 	virtual bool SupportsProperty(const FProperty* InProperty) const override;
 	//~ End IRCDefaultValueFactory interface
 };
+
+class FOverrideMaterialsDefaultValueFactory : public IRCDefaultValueFactory
+{
+public:
+	static TSharedRef<IRCDefaultValueFactory> MakeInstance();
+
+	//~ Begin IRCDefaultValueFactory interface
+	virtual bool CanResetToDefaultValue(UObject* InObject, const FRCResetToDefaultArgs& InArgs) const override;
+	virtual void ResetToDefaultValue(UObject* InObject, FRCResetToDefaultArgs& InArgs) override;
+	virtual bool SupportsClass(const UClass* InObjectClass) const override;
+	virtual bool SupportsProperty(const FProperty* InProperty) const override;
+	//~ End IRCDefaultValueFactory interface
+};
\ No newline at end of file
diff --git a/Engine/Plugins/VirtualProduction/RemoteControl/Source/RemoteControl/Private/RemoteControlModule.cpp b/Engine/Plugins/VirtualProduction/RemoteControl/Source/RemoteControl/Private/RemoteControlModule.cpp
--- a/Engine/Plugins/VirtualProduction/RemoteControl/Source/RemoteControl/Private/RemoteControlModule.cpp
+++ b/Engine/Plugins/VirtualProduction/RemoteControl/Source/RemoteControl/Private/RemoteControlModule.cpp
@@ -7,6 +7,7 @@
 #include "Backends/CborStructSerializerBackend.h"
 #include "Components/ActorComponent.h"
 #include "Components/LightComponent.h"
+#include "Components/MeshComponent.h"
 #include "Factories/RCDefaultValueFactories.h"
 #include "Factories/RemoteControlMaskingFactories.h"
 #include "Features/IModularFeatures.h"
@@ -752,21 +753,21 @@ void FRemoteControlModule::UnregisterEmbeddedPreset(URemoteControlPreset* Preset
 	EmbeddedPresets.Remove(PresetName);
 }
 
-bool FRemoteControlModule::CanResetToDefaultValue(UObject* InObject, const FProperty* InProperty) const
+bool FRemoteControlModule::CanResetToDefaultValue(UObject* InObject, const FRCResetToDefaultArgs& InArgs) const
 {
-	if (!InObject || !InProperty)
+	if (!InObject || !InArgs.Property)
 	{
 		return false;
 	}
 
 	// NOTE : Get the owning class to which the property actually belongs.
-	if (UClass* OwningClass = InProperty->GetOwnerClass())
+	if (UClass* OwningClass = InArgs.Property->GetOwnerClass())
 	{
-		const FString DefaultValueKey = OwningClass->GetName() + TEXT(".") + InProperty->GetName();
+		const FString DefaultValueKey = OwningClass->GetName() + TEXT(".") + InArgs.Property->GetName();
 		
 		if (const TSharedPtr<IRCDefaultValueFactory>* DefaultValueFactory = DefaultValueFactories.Find(*DefaultValueKey))
 		{
-			return (*DefaultValueFactory)->CanResetToDefaultValue(InObject);
+			return (*DefaultValueFactory)->CanResetToDefaultValue(InObject, InArgs);
 		}
 	}
 
@@ -794,21 +795,39 @@ bool FRemoteControlModule::HasDefaultValueCustomization(const UObject* InObject,
 	return false;
 }
 
-void FRemoteControlModule::ResetToDefaultValue(UObject* InObject, FProperty* InProperty)
+void FRemoteControlModule::ResetToDefaultValue(UObject* InObject, FRCResetToDefaultArgs& InArgs)
 {
-	if (!InObject || !InProperty)
+	if (!InObject || !InArgs.Property)
 	{
 		return;
 	}
 
 	// NOTE : Get the owning class to which the property actually belongs.
-	if (UClass* OwningClass = InProperty->GetOwnerClass())
+	if (UClass* OwningClass = InArgs.Property->GetOwnerClass())
 	{
-		const FString DefaultValueKey = OwningClass->GetName() + TEXT(".") + InProperty->GetName();
+		const FString DefaultValueKey = OwningClass->GetName() + TEXT(".") + InArgs.Property->GetName();
 
 		if (const TSharedPtr<IRCDefaultValueFactory>* DefaultValueFactory = DefaultValueFactories.Find(*DefaultValueKey))
 		{
-			(*DefaultValueFactory)->ResetToDefaultValue(InObject, InProperty);
+#if WITH_EDITOR
+			
+			if (InArgs.bCreateTransaction && GEditor)
+			{
+				GEditor->BeginTransaction(LOCTEXT("RemoteControlModule","Remote Control Custom Reset To Default"));
+			}
+
+#endif // WITH_EDITOR
+
+			(*DefaultValueFactory)->ResetToDefaultValue(InObject, InArgs);
+
+#if WITH_EDITOR
+
+			if (InArgs.bCreateTransaction && GEditor)
+			{
+				GEditor->EndTransaction();
+			}
+
+#endif // WITH_EDITOR
 		}
 	}
 }
@@ -1707,7 +1726,12 @@ bool FRemoteControlModule::ResetObjectProperties(const FRCObjectReference& Objec
 
 		if(HasDefaultValueCustomization(Object, ObjectAccess.Property.Get()))
 		{
-			ResetToDefaultValue(Object, ObjectAccess.Property.Get());
+			FRCResetToDefaultArgs Args;
+			Args.Property = ObjectAccess.Property.Get();
+			Args.Path = ObjectAccess.PropertyPathInfo.ToPathPropertyString();
+			Args.ArrayIndex = ObjectAccess.PropertyPathInfo.Segments[0].ArrayIndex;
+			Args.bCreateTransaction = false;
+			ResetToDefaultValue(Object, Args);
 		}
 		else
 #endif
@@ -2386,6 +2410,7 @@ void FRemoteControlModule::RegisterDefaultValueFactories()
 {
 	// NOTE : Use the class to which the property actually belongs.
 	RegisterDefaultValueFactoryForType(ULightComponentBase::StaticClass(), GET_MEMBER_NAME_CHECKED(ULightComponentBase, Intensity), FLightIntensityDefaultValueFactory::MakeInstance());
+	RegisterDefaultValueFactoryForType(UMeshComponent::StaticClass(), GET_MEMBER_NAME_CHECKED(UMeshComponent, OverrideMaterials), FOverrideMaterialsDefaultValueFactory::MakeInstance());
 }
 
 void FRemoteControlModule::RegisterMaskingFactories()
diff --git a/Engine/Plugins/VirtualProduction/RemoteControl/Source/RemoteControl/Private/RemoteControlModule.h b/Engine/Plugins/VirtualProduction/RemoteControl/Source/RemoteControl/Private/RemoteControlModule.h
--- a/Engine/Plugins/VirtualProduction/RemoteControl/Source/RemoteControl/Private/RemoteControlModule.h
+++ b/Engine/Plugins/VirtualProduction/RemoteControl/Source/RemoteControl/Private/RemoteControlModule.h
@@ -31,9 +31,9 @@ public:
 	virtual bool RegisterEmbeddedPreset(URemoteControlPreset* Preset, bool bReplaceExisting) override;
 	virtual void UnregisterEmbeddedPreset(FName Name) override;
 	virtual void UnregisterEmbeddedPreset(URemoteControlPreset* Preset) override;
-	virtual bool CanResetToDefaultValue(UObject* InObject, const FProperty* InProperty) const override;
+	virtual bool CanResetToDefaultValue(UObject* InObject, const FRCResetToDefaultArgs& InArgs) const override;
 	virtual bool HasDefaultValueCustomization(const UObject* InObject, const FProperty* InProperty) const override;
-	virtual void ResetToDefaultValue(UObject* InObject, FProperty* InProperty) override;
+	virtual void ResetToDefaultValue(UObject* InObject, FRCResetToDefaultArgs& InArgs) override;
 	virtual void PerformMasking(const TSharedRef<FRCMaskingOperation>& InMaskingOperation) override;
 	virtual void RegisterMaskingFactoryForType(UScriptStruct* RemoteControlPropertyType, const TSharedPtr<IRemoteControlMaskingFactory>& InMaskingFactory) override;
 	virtual void UnregisterMaskingFactoryForType(UScriptStruct* RemoteControlPropertyType) override;
diff --git a/Engine/Plugins/VirtualProduction/RemoteControl/Source/RemoteControl/Private/RemoteControlPreset.cpp b/Engine/Plugins/VirtualProduction/RemoteControl/Source/RemoteControl/Private/RemoteControlPreset.cpp
--- a/Engine/Plugins/VirtualProduction/RemoteControl/Source/RemoteControl/Private/RemoteControlPreset.cpp
+++ b/Engine/Plugins/VirtualProduction/RemoteControl/Source/RemoteControl/Private/RemoteControlPreset.cpp
@@ -36,6 +36,7 @@
 #include "Editor.h"
 #include "EngineAnalytics.h"
 #include "Engine/Blueprint.h"
+#include "Factories/IRCDefaultValueFactory.h"
 #include "TimerManager.h"
 #include "UObject/PackageReload.h"
 #endif
@@ -1777,6 +1778,28 @@ void URemoteControlPreset::OnObjectPropertyChanged(UObject* Object, struct FProp
 				}
 			}
 		}
+		for (auto Iter = PreMaterialModifiedCache.CreateIterator(); Iter; ++Iter)
+		{
+			FGuid& PropertyId = Iter.Key();
+			FPreMaterialModifiedCache& CacheEntry = Iter.Value();
+			if (TSharedPtr<FRemoteControlProperty> Property = Registry->GetExposedEntity<FRemoteControlProperty>(PropertyId))
+			{
+				if (UMeshComponent* MeshComponent = Cast<UMeshComponent>(Object))
+				{
+					if (CacheEntry.bHadValue && MeshComponent->OverrideMaterials[CacheEntry.ArrayIndex] == NULL)
+					{
+						FRCResetToDefaultArgs Args;
+						Args.Property = Event.Property;
+						Args.Path = Property->FieldPathInfo.ToPathPropertyString();
+						Args.ArrayIndex = Property->FieldPathInfo.Segments[0].ArrayIndex;
+						Args.bCreateTransaction = false;
+						IRemoteControlModule& RemoteControlModule = IRemoteControlModule::Get();
+						RemoteControlModule.ResetToDefaultValue(Object, Args);
+					}
+				}
+			}
+			Iter.RemoveCurrent();
+		}
 	}
  
 	for (auto Iter = PreObjectsModifiedActorCache.CreateIterator(); Iter; ++Iter)
@@ -1869,7 +1892,17 @@ void URemoteControlPreset::OnPreObjectPropertyChanged(UObject* Object, const cla
 						if (ExposedProperty == Current->GetValue())
 						{
 							bHasFound = true;
-
+							
+							if (UMeshComponent* MeshComponent = Cast<UMeshComponent>(Object))
+							{
+								if (PropertyChain.GetActiveNode()->GetValue()->GetFName() == GET_MEMBER_NAME_CHECKED(UMeshComponent, OverrideMaterials))
+								{
+									FPreMaterialModifiedCache& NewEntry = PreMaterialModifiedCache.FindOrAdd(RCProperty->GetId());
+									NewEntry.ArrayIndex = RCProperty->FieldPathInfo.Segments[0].ArrayIndex;
+									NewEntry.bHadValue = MeshComponent->OverrideMaterials[NewEntry.ArrayIndex] != NULL;
+								}
+							}
+							
 							FPreObjectsModifiedCache& NewEntry = PreObjectsModifiedCache.FindOrAdd(RCProperty->GetId());
 							NewEntry.Objects.AddUnique(Object);
 							NewEntry.Property = PropertyChain.GetActiveNode()->GetValue();
diff --git a/Engine/Plugins/VirtualProduction/RemoteControl/Source/RemoteControl/Public/Factories/IRCDefaultValueFactory.h b/Engine/Plugins/VirtualProduction/RemoteControl/Source/RemoteControl/Public/Factories/IRCDefaultValueFactory.h
--- a/Engine/Plugins/VirtualProduction/RemoteControl/Source/RemoteControl/Public/Factories/IRCDefaultValueFactory.h
+++ b/Engine/Plugins/VirtualProduction/RemoteControl/Source/RemoteControl/Public/Factories/IRCDefaultValueFactory.h
@@ -4,6 +4,36 @@
 
 #include "CoreMinimal.h"
 
+/**
+ * Holds the arguments required by the Reset Factories.
+ */
+struct FRCResetToDefaultArgs
+{
+	/** Holds a reference to the actual property. */
+	FProperty* Property;
+
+	/** Holds the property path of the actual property. */
+	FString Path;
+
+	/** Holds the index of the property being modified incase it is embeded in an array. */
+	int32 ArrayIndex;
+
+	/** Determines whether should we notify the transaction manager or not. */
+	bool bCreateTransaction;
+
+public:
+
+	FRCResetToDefaultArgs()
+		: Property(nullptr)
+		, Path(TEXT(""))
+		, ArrayIndex(0)
+		, bCreateTransaction(false)
+	{}
+	
+	FRCResetToDefaultArgs(const FRCResetToDefaultArgs&);
+	FRCResetToDefaultArgs& operator=(const FRCResetToDefaultArgs&);
+};
+
 /**
  * Factory which resets the FRemoteControlProperty to its default values.
  */
@@ -17,15 +47,16 @@ public:
 	/**
 	 * Returns true when the given object can be reset to its default value, false otherwise.
 	 * @param InObject Reference to the exposed object.
+	 * @param InArgs Arguments to be passed.
 	 */
-	virtual bool CanResetToDefaultValue(UObject* InObject) const = 0;
+	virtual bool CanResetToDefaultValue(UObject* InObject, const FRCResetToDefaultArgs& InArgs) const = 0;
 
 	/**
 	 * Performs actual data reset on the given remote object.
 	 * @param InObject Reference to the exposed object.
-	 * @param InProperty Reference to the exposed property.
+	 * @param InArgs Arguments to be passed.
 	 */
-	virtual void ResetToDefaultValue(UObject* InObject, FProperty* InProperty) = 0;
+	virtual void ResetToDefaultValue(UObject* InObject, FRCResetToDefaultArgs& InArgs) = 0;
 	
 	/**
 	 * Whether the factory support exposed entity.
diff --git a/Engine/Plugins/VirtualProduction/RemoteControl/Source/RemoteControl/Public/IRemoteControlModule.h b/Engine/Plugins/VirtualProduction/RemoteControl/Source/RemoteControl/Public/IRemoteControlModule.h
--- a/Engine/Plugins/VirtualProduction/RemoteControl/Source/RemoteControl/Public/IRemoteControlModule.h
+++ b/Engine/Plugins/VirtualProduction/RemoteControl/Source/RemoteControl/Public/IRemoteControlModule.h
@@ -22,6 +22,7 @@ PRAGMA_DISABLE_DEPRECATION_WARNINGS
 struct FExposedProperty;
 PRAGMA_ENABLE_DEPRECATION_WARNINGS
 struct FRCMaskingOperation;
+struct FRCResetToDefaultArgs;
 struct FRemoteControlProperty;
 class IRemoteControlPropertyFactory;
 
@@ -282,9 +283,9 @@ public:
 	/**
 	 * Returns true when the given object can be reset to its default value, false otherwise.
 	 * @param InObject Reference to the exposed object.
-	 * @param InProperty Reference to the exposed property.
+	 * @param InArgs Arguments to be passed.
 	 */
-	virtual bool CanResetToDefaultValue(UObject* InObject, const FProperty* InProperty) const = 0;
+	virtual bool CanResetToDefaultValue(UObject* InObject, const FRCResetToDefaultArgs& InArgs) const = 0;
 
 	/**
 	 * Returns true when the given object and property has custom default value, false otherwise.
@@ -296,9 +297,9 @@ public:
 	/**
 	 * Performs actual data reset on the given object.
 	 * @param InObject Reference to the exposed object.
-	 * @param InProperty Reference to the exposed property.
+	 * @param InArgs Arguments to be passed.
 	 */
-	virtual void ResetToDefaultValue(UObject* InObject, FProperty* InProperty) = 0;
+	virtual void ResetToDefaultValue(UObject* InObject, FRCResetToDefaultArgs& InArgs) = 0;
 
 	/**
 	 * Performs the given masking operation.
diff --git a/Engine/Plugins/VirtualProduction/RemoteControl/Source/RemoteControl/Public/RemoteControlPreset.h b/Engine/Plugins/VirtualProduction/RemoteControl/Source/RemoteControl/Public/RemoteControlPreset.h
--- a/Engine/Plugins/VirtualProduction/RemoteControl/Source/RemoteControl/Public/RemoteControlPreset.h
+++ b/Engine/Plugins/VirtualProduction/RemoteControl/Source/RemoteControl/Public/RemoteControlPreset.h
@@ -847,6 +847,16 @@ private:
 	TMap<FGuid, FPreObjectsModifiedCache> PreObjectsModifiedCache;
 	TMap<FGuid, FPreObjectsModifiedCache> PreObjectsModifiedActorCache;
 
+	/** A struct representing a material object modified that is housed by a container (eg. array, set, map). */
+	struct FPreMaterialModifiedCache
+	{
+		int32 ArrayIndex;
+		bool bHadValue;
+	};
+
+	/** Caches material modifications during a frame. */
+	TMap<FGuid, FPreMaterialModifiedCache> PreMaterialModifiedCache;
+
 	/** Cache properties that were modified during a frame. */
 	TSet<FGuid> PerFrameModifiedProperties;
 	
diff --git a/Engine/Plugins/VirtualProduction/RemoteControl/Source/RemoteControlUI/Private/UI/SRCPanelExposedField.cpp b/Engine/Plugins/VirtualProduction/RemoteControl/Source/RemoteControlUI/Private/UI/SRCPanelExposedField.cpp
--- a/Engine/Plugins/VirtualProduction/RemoteControl/Source/RemoteControlUI/Private/UI/SRCPanelExposedField.cpp
+++ b/Engine/Plugins/VirtualProduction/RemoteControl/Source/RemoteControlUI/Private/UI/SRCPanelExposedField.cpp
@@ -3,8 +3,9 @@
 #include "SRCPanelExposedField.h"
 
 #include "Algo/Transform.h"
-#include "Components/LocalLightComponent.h"
+#include "Components/MeshComponent.h"
 #include "EditorFontGlyphs.h"
+#include "Factories/IRCDefaultValueFactory.h"
 #include "IDetailTreeNode.h"
 #include "IRCProtocolBindingList.h"
 #include "IRemoteControlModule.h"
@@ -13,6 +14,7 @@
 #include "IRemoteControlUIModule.h"
 #include "Layout/Visibility.h"
 #include "Materials/MaterialInstanceDynamic.h"
+#include "PropertyEditor/Private/PropertyNode.h"
 #include "RCPanelWidgetRegistry.h"
 #include "RemoteControlField.h"
 #include "RemoteControlPanelStyle.h"
@@ -580,7 +582,17 @@ TSharedRef<SWidget> SRCPanelExposedField::ConstructResetToDefaultWidget(UObject*
 		
 			if (RemoteControlModule.HasDefaultValueCustomization(InObject, InPropertyHandle->GetProperty()))
 			{
-				return RemoteControlModule.CanResetToDefaultValue(InObject, InPropertyHandle->GetProperty());
+				FRCResetToDefaultArgs Args;
+				Args.Property = InPropertyHandle->GetProperty();
+				
+				if (TSharedPtr<FPropertyNode> PropertyNode = InPropertyHandle->GetPropertyNode())
+				{
+					Args.Path = PropertyNode->GetPropertyPath();
+					Args.ArrayIndex = PropertyNode->GetArrayIndex();
+				}
+
+				Args.bCreateTransaction = false;
+				return RemoteControlModule.CanResetToDefaultValue(InObject, Args);
 			}
 
 			return InPropertyHandle->DiffersFromDefault();
@@ -593,7 +605,17 @@ TSharedRef<SWidget> SRCPanelExposedField::ConstructResetToDefaultWidget(UObject*
 
 			if (RemoteControlModule.HasDefaultValueCustomization(InObject, InPropertyHandle->GetProperty()))
 			{
-				RemoteControlModule.ResetToDefaultValue(InObject, InPropertyHandle->GetProperty());
+				FRCResetToDefaultArgs Args;
+				Args.Property = InPropertyHandle->GetProperty();
+				
+				if (TSharedPtr<FPropertyNode> PropertyNode = InPropertyHandle->GetPropertyNode())
+				{
+					Args.Path = PropertyNode->GetPropertyPath();
+					Args.ArrayIndex = PropertyNode->GetArrayIndex();
+				}
+				
+				Args.bCreateTransaction = true;
+				RemoteControlModule.ResetToDefaultValue(InObject, Args);
 			
 				return;
 			}
@@ -603,6 +625,28 @@ TSharedRef<SWidget> SRCPanelExposedField::ConstructResetToDefaultWidget(UObject*
 	);
 	
 	const FResetToDefaultOverride ResetValueOverride = FResetToDefaultOverride::Create(CanResetToDefaultValue, PerformValueReset);
+
+	if (UMeshComponent* MeshComponent = Cast<UMeshComponent>(InObject))
+	{
+		if (InPropertyHandle->GetProperty()->GetFName() == GET_MEMBER_NAME_CHECKED(UMeshComponent, OverrideMaterials) &&
+			MeshComponent->OverrideMaterials[InPropertyHandle->GetPropertyNode()->GetArrayIndex()] == NULL)
+		{
+			IRemoteControlModule& RemoteControlModule = IRemoteControlModule::Get();
+
+			FRCResetToDefaultArgs Args;
+			Args.Property = InPropertyHandle->GetProperty();
+			
+			if (TSharedPtr<FPropertyNode> PropertyNode = InPropertyHandle->GetPropertyNode())
+			{
+				Args.Path = PropertyNode->GetPropertyPath();
+				Args.ArrayIndex = PropertyNode->GetArrayIndex();
+			}
+
+			Args.bCreateTransaction = false;
+			RemoteControlModule.ResetToDefaultValue(InObject, Args);
+		}
+		
+	}
 	
 	ResetDefaultProperty.Reset();
 	
