CIS fix for remote control related error

diff --git a/Engine/Plugins/VirtualProduction/RemoteControl/Source/RemoteControlUI/Private/UI/SRCPanelExposedField.cpp b/Engine/Plugins/VirtualProduction/RemoteControl/Source/RemoteControlUI/Private/UI/SRCPanelExposedField.cpp
--- a/Engine/Plugins/VirtualProduction/RemoteControl/Source/RemoteControlUI/Private/UI/SRCPanelExposedField.cpp
+++ b/Engine/Plugins/VirtualProduction/RemoteControl/Source/RemoteControlUI/Private/UI/SRCPanelExposedField.cpp
@@ -14,7 +14,7 @@
 #include "IRemoteControlUIModule.h"
 #include "Layout/Visibility.h"
 #include "Materials/MaterialInstanceDynamic.h"
-#include "PropertyEditor/Private/PropertyNode.h"
+#include "PropertyNode.h"
 #include "RCPanelWidgetRegistry.h"
 #include "RemoteControlField.h"
 #include "RemoteControlPanelStyle.h"
diff --git a/Engine/Plugins/VirtualProduction/RemoteControl/Source/RemoteControlUI/RemoteControlUI.Build.cs b/Engine/Plugins/VirtualProduction/RemoteControl/Source/RemoteControlUI/RemoteControlUI.Build.cs
--- a/Engine/Plugins/VirtualProduction/RemoteControl/Source/RemoteControlUI/RemoteControlUI.Build.cs
+++ b/Engine/Plugins/VirtualProduction/RemoteControl/Source/RemoteControlUI/RemoteControlUI.Build.cs
@@ -1,11 +1,16 @@
 // Copyright Epic Games, Inc. All Rights Reserved.
 
+using System.IO;
 using UnrealBuildTool;
 
 public class RemoteControlUI : ModuleRules
 {
 	public RemoteControlUI(ReadOnlyTargetRules Target) : base(Target)
 	{
+		PrivateIncludePaths.AddRange(new string[] {
+			Path.Combine(GetModuleDirectory("PropertyEditor"), "Private"),
+		});
+
 		PublicDependencyModuleNames.AddRange(
 			new string[] {}
 		);
