Remote API - Packaging with MIDI properties exposed fails

diff --git a/Engine/Plugins/VirtualProduction/RemoteControl/Source/RemoteControlProtocol/Private/RemoteControlProtocolModule.cpp b/Engine/Plugins/VirtualProduction/RemoteControl/Source/RemoteControlProtocol/Private/RemoteControlProtocolModule.cpp
--- a/Engine/Plugins/VirtualProduction/RemoteControl/Source/RemoteControlProtocol/Private/RemoteControlProtocolModule.cpp
+++ b/Engine/Plugins/VirtualProduction/RemoteControl/Source/RemoteControlProtocol/Private/RemoteControlProtocolModule.cpp
@@ -2,6 +2,7 @@
 
 #include "RemoteControlProtocolModule.h"
 
+#include "CoreGlobals.h"
 #include "RemoteControlPreset.h"
 #include "RemoteControlProtocol.h"
 #include "Misc/CommandLine.h"
@@ -104,12 +105,12 @@ void FRemoteControlProtocolModule::EmptyProtocols()
 
 void FRemoteControlProtocolModule::OnPostLoadRemoteControlPreset(URemoteControlPreset* InPreset) const
 {
-	if (!ensure(InPreset))
+	// Avoid protocol binding if the preset is not valid or while cooking.
+	if (!ensure(InPreset) || GIsCookerLoadingPackage)
 	{
 		return;
 	}
 	
-	
 	for(TWeakPtr<FRemoteControlProperty> ExposedPropertyWeakPtr : InPreset->GetExposedEntities<FRemoteControlProperty>())
 	{
 		if (TSharedPtr<FRemoteControlProperty> ExposedPropertyPtr = ExposedPropertyWeakPtr.Pin())
