Remote Control:
- Add PostEditChangeChainProperty call on finalize changes

diff --git a/Engine/Plugins/VirtualProduction/RemoteControl/Source/RemoteControl/Private/RemoteControlModule.cpp b/Engine/Plugins/VirtualProduction/RemoteControl/Source/RemoteControl/Private/RemoteControlModule.cpp
--- a/Engine/Plugins/VirtualProduction/RemoteControl/Source/RemoteControl/Private/RemoteControlModule.cpp
+++ b/Engine/Plugins/VirtualProduction/RemoteControl/Source/RemoteControl/Private/RemoteControlModule.cpp
@@ -2610,11 +2610,41 @@ void FRemoteControlModule::TestOrFinalizeOngoingChange(bool bForceEndChange)
 
 			if (OngoingModification->Reference.IsType<FRCObjectReference>())
 			{
-				FPropertyChangedEvent PropertyEvent = OngoingModification->Reference.Get<FRCObjectReference>().PropertyPathInfo.ToPropertyChangedEvent();
-				PropertyEvent.ChangeType = EPropertyChangeType::ValueSet;
-				if (UObject* Object = OngoingModification->Reference.Get<FRCObjectReference>().Object.Get())
+				const FRCObjectReference& ObjectReference  = OngoingModification->Reference.Get<FRCObjectReference>();
+
+				if (UObject* const Object = ObjectReference.Object.Get())
 				{
-					Object->PostEditChangeProperty(PropertyEvent);
+					FPropertyChangedEvent PropertyEvent = ObjectReference.PropertyPathInfo.ToPropertyChangedEvent(EPropertyChangeType::ValueSet);
+
+					FEditPropertyChain EditPropertyChain;
+					ObjectReference.PropertyPathInfo.ToEditPropertyChain(EditPropertyChain);
+
+					if (EditPropertyChain.IsEmpty())
+					{
+						Object->PostEditChangeProperty(PropertyEvent);
+					}
+					else
+					{
+						FPropertyChangedChainEvent ChainEvent(EditPropertyChain, PropertyEvent);
+
+						TArray<TMap<FString, int32>> ArrayIndicesPerObject;
+						{
+							TMap<FString, int32> ArrayIndices;
+							ArrayIndices.Reserve(ObjectReference.PropertyPathInfo.Segments.Num());
+
+							for (const FRCFieldPathSegment& Segment : ObjectReference.PropertyPathInfo.Segments)
+							{
+								ArrayIndices.Add(Segment.Name.ToString(), Segment.ArrayIndex);
+							}
+
+							ArrayIndicesPerObject.Add(MoveTemp(ArrayIndices));
+						}
+
+						ChainEvent.ObjectIteratorIndex = 0;
+						ChainEvent.SetArrayIndexPerObject(ArrayIndicesPerObject);
+
+						Object->PostEditChangeChainProperty(ChainEvent);
+					}
 				}
 			}
 
