Remote Control - Crash when changing levels

diff --git a/Engine/Plugins/VirtualProduction/RemoteControl/Source/WebRemoteControl/Private/WebSocketMessageHandler.cpp b/Engine/Plugins/VirtualProduction/RemoteControl/Source/WebRemoteControl/Private/WebSocketMessageHandler.cpp
--- a/Engine/Plugins/VirtualProduction/RemoteControl/Source/WebRemoteControl/Private/WebSocketMessageHandler.cpp
+++ b/Engine/Plugins/VirtualProduction/RemoteControl/Source/WebRemoteControl/Private/WebSocketMessageHandler.cpp
@@ -1799,18 +1799,21 @@ bool FWebSocketMessageHandler::TrySendMultipleBoolProperties(URemoteControlPrese
 	{
 		if (TSharedPtr<FRemoteControlProperty> RCProperty = InPreset->GetExposedEntity<FRemoteControlProperty>(InModifiedPropertyIds.Array()[0]).Pin())
 		{
-			if(RCProperty->GetProperty()->IsA<FBoolProperty>())
+			if (const FProperty* Property = RCProperty->GetProperty())
 			{
-				bFound = true;
-				for(FGuid ModifiedPropertyId : InModifiedPropertyIds)
+				if (Property->IsA<FBoolProperty>())
 				{
-					TArray<uint8> BoolsWorkingBuffer;
-					if(WritePropertyChangeEventPayload(InPreset, { ModifiedPropertyId }, InSequenceNumber, BoolsWorkingBuffer))
+					bFound = true;
+					for (FGuid ModifiedPropertyId : InModifiedPropertyIds)
 					{
-						TArray<uint8> Payload;
-						WebRemoteControlUtils::ConvertToUTF8(BoolsWorkingBuffer, Payload);
-						Server->Send(InTargetClientId, Payload);
-						++NumberSent;
+						TArray<uint8> BoolsWorkingBuffer;
+						if (WritePropertyChangeEventPayload(InPreset, { ModifiedPropertyId }, InSequenceNumber, BoolsWorkingBuffer))
+						{
+							TArray<uint8> Payload;
+							WebRemoteControlUtils::ConvertToUTF8(BoolsWorkingBuffer, Payload);
+							Server->Send(InTargetClientId, Payload);
+							++NumberSent;
+						}
 					}
 				}
 			}
