Remote Control not releasing Midi Controller properly

diff --git a/Engine/Plugins/VirtualProduction/RemoteControl/Source/RemoteControlProtocolWidgets/Private/ViewModels/ProtocolEntityViewModel.cpp b/Engine/Plugins/VirtualProduction/RemoteControl/Source/RemoteControlProtocolWidgets/Private/ViewModels/ProtocolEntityViewModel.cpp
--- a/Engine/Plugins/VirtualProduction/RemoteControl/Source/RemoteControlProtocolWidgets/Private/ViewModels/ProtocolEntityViewModel.cpp
+++ b/Engine/Plugins/VirtualProduction/RemoteControl/Source/RemoteControlProtocolWidgets/Private/ViewModels/ProtocolEntityViewModel.cpp
@@ -82,6 +82,12 @@ namespace Commands
 			const FRemoteControlProtocolBinding* FoundElement = RCProperty->ProtocolBindings.FindByHash(GetTypeHash(InArgs.BindingId), InArgs.BindingId);
 			if (FoundElement)
 			{
+				// Unbind the protocol from the property explicitly while removing.
+				if (const TSharedPtr<IRemoteControlProtocol> Protocol = IRemoteControlProtocolModule::Get().GetProtocolByName(FoundElement->GetProtocolName()))
+				{
+					Protocol->Unbind(FoundElement->GetRemoteControlProtocolEntityPtr());
+				}
+
 				RCProperty->ProtocolBindings.RemoveByHash(GetTypeHash(InArgs.BindingId), InArgs.BindingId);
 
 				const FName RemovedProtocolName = FoundElement->GetProtocolName();
