Allow changing RC web interface bind address
Enable remote passphrase enforcement by default when Restrict Server access is enabled

diff --git a/Engine/Plugins/VirtualProduction/RemoteControl/Source/RemoteControlCommon/Public/RemoteControlSettings.h b/Engine/Plugins/VirtualProduction/RemoteControl/Source/RemoteControlCommon/Public/RemoteControlSettings.h
--- a/Engine/Plugins/VirtualProduction/RemoteControl/Source/RemoteControlCommon/Public/RemoteControlSettings.h
+++ b/Engine/Plugins/VirtualProduction/RemoteControl/Source/RemoteControlCommon/Public/RemoteControlSettings.h
@@ -263,6 +263,10 @@ public:
 	UPROPERTY(config, EditAnywhere, Category = "Remote Control")
 	bool bProtocolsGenerateTransactions = true;
 
+	/** The remote control web app bind address. */
+	UPROPERTY(config, EditAnywhere, Category = "Remote Control Web Interface", DisplayName = "Remote Control Web Interface bind address")
+	FString RemoteControlWebInterfaceBindAddress = TEXT("0.0.0.0");
+
 	/** The remote control web app http port. */
 	UPROPERTY(config, EditAnywhere, Category = "Remote Control Web Interface", DisplayName = "Remote Control Web Interface http Port")
 	uint32 RemoteControlWebInterfacePort = 30000;
@@ -339,7 +343,7 @@ public:
 	 * Controls whether a passphrase should be required when remote control is accessed by a client outside of localhost.
 	 */
 	UPROPERTY(config, EditAnywhere, Category = "Remote Control | Security", meta = (EditCondition = bRestrictServerAccess))
-	bool bEnforcePassphraseForRemoteClients = false;
+	bool bEnforcePassphraseForRemoteClients = true;
 
 	/** List of passphrases used for accessing remote control outside of localhost. */
 	UPROPERTY(config, EditAnywhere, Category = "Remote Control | Security", DisplayName = "Remote Control Passphrase", meta = (EditCondition=bEnforcePassphraseForRemoteClients, EditConditionHides))
diff --git a/Engine/Plugins/VirtualProduction/RemoteControlWebInterface/Source/RemoteControlWebInterface/Private/RCWebInterface.cpp b/Engine/Plugins/VirtualProduction/RemoteControlWebInterface/Source/RemoteControlWebInterface/Private/RCWebInterface.cpp
--- a/Engine/Plugins/VirtualProduction/RemoteControlWebInterface/Source/RemoteControlWebInterface/Private/RCWebInterface.cpp
+++ b/Engine/Plugins/VirtualProduction/RemoteControlWebInterface/Source/RemoteControlWebInterface/Private/RCWebInterface.cpp
@@ -155,7 +155,9 @@ void FRemoteControlWebInterfaceModule::OnSettingsModified(UObject* Settings, FPr
 
 	if (PropertyChangedEvent.ChangeType != EPropertyChangeType::Interactive)
 	{
-		if (PropertyName == GET_MEMBER_NAME_CHECKED(URemoteControlSettings, RemoteControlWebInterfacePort) || PropertyName == GET_MEMBER_NAME_CHECKED(URemoteControlSettings, bForceWebAppBuildAtStartup))
+		if (PropertyName == GET_MEMBER_NAME_CHECKED(URemoteControlSettings, RemoteControlWebInterfacePort)
+			|| PropertyName == GET_MEMBER_NAME_CHECKED(URemoteControlSettings, bForceWebAppBuildAtStartup)
+			|| PropertyName == GET_MEMBER_NAME_CHECKED(URemoteControlSettings, RemoteControlWebInterfaceBindAddress))
 		{
 			WebApp->Shutdown();
 			WebApp->Start();
diff --git a/Engine/Plugins/VirtualProduction/RemoteControlWebInterface/Source/RemoteControlWebInterface/Private/RCWebInterfaceProcess.cpp b/Engine/Plugins/VirtualProduction/RemoteControlWebInterface/Source/RemoteControlWebInterface/Private/RCWebInterfaceProcess.cpp
--- a/Engine/Plugins/VirtualProduction/RemoteControlWebInterface/Source/RemoteControlWebInterface/Private/RCWebInterfaceProcess.cpp
+++ b/Engine/Plugins/VirtualProduction/RemoteControlWebInterface/Source/RemoteControlWebInterface/Private/RCWebInterfaceProcess.cpp
@@ -129,10 +129,11 @@ uint32 FRemoteControlWebInterfaceProcess::Run()
 
 	const URemoteControlSettings* RCSettings = GetDefault<URemoteControlSettings>();
 
-	FString Args = FString::Printf(TEXT("--port %d --uews %d --uehttp %d --monitor "),
+	FString Args = FString::Printf(TEXT("--port %d --uews %d --uehttp %d --monitor --hostname %s "),
 									RCSettings->RemoteControlWebInterfacePort,
 									RCSettings->RemoteControlWebSocketServerPort,
-									RCSettings->RemoteControlHttpServerPort);
+									RCSettings->RemoteControlHttpServerPort,
+									*RCSettings->RemoteControlWebInterfaceBindAddress);
 
 	if (RCSettings->bForceWebAppBuildAtStartup)
 	{
diff --git a/Engine/Plugins/VirtualProduction/RemoteControlWebInterface/WebApp/Server/src/App.ts b/Engine/Plugins/VirtualProduction/RemoteControlWebInterface/WebApp/Server/src/App.ts
--- a/Engine/Plugins/VirtualProduction/RemoteControlWebInterface/WebApp/Server/src/App.ts
+++ b/Engine/Plugins/VirtualProduction/RemoteControlWebInterface/WebApp/Server/src/App.ts
@@ -32,7 +32,12 @@ export class App extends Server {
                     .catch(() => {});
 
       await this.startServer();
-      console.log('DONE: WebApp started, port:', Program.port);
+
+      var hostname = Program.hostname;
+      if (hostname == '0.0.0.0'){
+        hostname = '127.0.0.1';
+      }
+      console.log(`DONE: WebApp started, ${hostname}:${Program.port}`, );
 
     } catch (error) {
       console.error('ERROR:', error.message);
@@ -42,7 +47,7 @@ export class App extends Server {
 
   private startServer(): Promise<void> {
     return new Promise((resolve, reject) => {
-      const server = this.app.listen(Program.port, async () => {
+      const server = this.app.listen(Program.port, Program.hostname,  async () => {
         await Notify.initialize(server);
         await UnrealEngine.initialize();
         await LogServer.initialize();
diff --git a/Engine/Plugins/VirtualProduction/RemoteControlWebInterface/WebApp/Server/src/Program.ts b/Engine/Plugins/VirtualProduction/RemoteControlWebInterface/WebApp/Server/src/Program.ts
--- a/Engine/Plugins/VirtualProduction/RemoteControlWebInterface/WebApp/Server/src/Program.ts
+++ b/Engine/Plugins/VirtualProduction/RemoteControlWebInterface/WebApp/Server/src/Program.ts
@@ -6,6 +6,7 @@ class Program {
   readonly dev: boolean;
   readonly monitor: boolean;
   readonly port: number;
+  readonly hostname: string;
   readonly ueHttpPort: number;
   readonly ueWebSocketPort: number;
   readonly rootFolder: string;
@@ -15,6 +16,7 @@ class Program {
     this.dev = false;
     this.monitor = false;
     this.port = 7000;
+    this.hostname = '0.0.0.0';
     this.ueHttpPort = 30010;
     this.ueWebSocketPort = 30020;
     this.logger = false;
@@ -48,6 +50,9 @@ class Program {
         case '--log':
           this.logger = true;
           break;
+
+        case '--hostname':
+          this.hostname = process.argv[i + 1] || '0.0.0.0';
       }
     }
 
diff --git a/Engine/Plugins/VirtualProduction/RemoteControlWebInterface/WebApp/_package.json b/Engine/Plugins/VirtualProduction/RemoteControlWebInterface/WebApp/_package.json
--- a/Engine/Plugins/VirtualProduction/RemoteControlWebInterface/WebApp/_package.json
+++ b/Engine/Plugins/VirtualProduction/RemoteControlWebInterface/WebApp/_package.json
@@ -1,6 +1,6 @@
 {
   "name": "conductor",
-  "version": "0.3.89",
+  "version": "0.3.90",
   "scripts": {
     "postinstall": "concurrently \"cd ./Client && npm install && cd ..\" \"cd ./Server && npm install && cd ..\"",
     "build": "node ./Scripts/build.js",
