Remove control: Fix game mode issues.

diff --git a/Engine/Plugins/VirtualProduction/RemoteControl/Source/RemoteControl/Private/RemoteControlModule.cpp b/Engine/Plugins/VirtualProduction/RemoteControl/Source/RemoteControl/Private/RemoteControlModule.cpp
--- a/Engine/Plugins/VirtualProduction/RemoteControl/Source/RemoteControl/Private/RemoteControlModule.cpp
+++ b/Engine/Plugins/VirtualProduction/RemoteControl/Source/RemoteControl/Private/RemoteControlModule.cpp
@@ -56,7 +56,7 @@ struct FRCInterceptionPayload
 
 #if WITH_EDITOR
 static TAutoConsoleVariable<int32> CVarRemoteControlEnableOngoingChangeOptimization(TEXT("RemoteControl.EnableOngoingChangeOptimization"), 1, TEXT("Enable an optimization that keeps track of the ongoing remote control change in order to improve performance."));
-#endif
+#endif // WITH_EDITOR
 
 #define REMOTE_CONTROL_LOG_ONCE(Verbosity, Format, ...) (LogOnce(ELogVerbosity::Verbosity, *FString::Printf(Format, ##__VA_ARGS__), FString(__FILE__), __LINE__))
 
@@ -112,6 +112,8 @@ namespace RemoteControlUtil
 	{
 #if WITH_EDITOR
 		URemoteControlSettings* RCSettings = GetMutableDefault<URemoteControlSettings>();
+		// Override this flag to false if we are running in Editor Mode with -game flag.
+		bObjectInGamePackage = false;
 #endif
 		// The property is allowed to be accessed if it exists...
 		return InProperty &&
@@ -606,7 +608,7 @@ void FRemoteControlModule::StartupModule()
 #if WITH_EDITOR
 	FCoreDelegates::OnPostEngineInit.AddRaw(this, &FRemoteControlModule::HandleEnginePostInit);
 	FCoreUObjectDelegates::PreLoadMap.AddRaw(this, &FRemoteControlModule::HandleMapPreLoad);
-#endif
+#endif // WITH_EDITOR
 	
 	// Register Default Value Factories
 	RegisterDefaultValueFactories();
@@ -2676,6 +2678,16 @@ void FRemoteControlModule::RegisterEditorDelegates()
 	{
 		GEditor->GetTimerManager()->SetTimer(OngoingChangeTimer, FTimerDelegate::CreateRaw(this, &FRemoteControlModule::TestOrFinalizeOngoingChange, false), SecondsBetweenOngoingChangeCheck, true);
 	}
+	else
+	{
+		FallbackOngoingChangeTimer = FTSTicker::GetCoreTicker().AddTicker(FTickerDelegate::CreateLambda([&](float DeltaTime)
+			{
+				TestOrFinalizeOngoingChange(false);
+
+				return true;
+			}
+		), SecondsBetweenOngoingChangeCheck);
+	}
 }
 	
 void FRemoteControlModule::UnregisterEditorDelegates()
@@ -2684,6 +2696,10 @@ void FRemoteControlModule::UnregisterEditorDelegates()
 	{ 
 		GEditor->GetTimerManager()->ClearTimer(OngoingChangeTimer);
 	}
+	else
+	{
+		FTSTicker::GetCoreTicker().RemoveTicker(FallbackOngoingChangeTimer);
+	}
 }
 
 FRemoteControlModule::FOngoingChange::FOngoingChange(FRCObjectReference InReference)
diff --git a/Engine/Plugins/VirtualProduction/RemoteControl/Source/RemoteControl/Private/RemoteControlModule.h b/Engine/Plugins/VirtualProduction/RemoteControl/Source/RemoteControl/Private/RemoteControlModule.h
--- a/Engine/Plugins/VirtualProduction/RemoteControl/Source/RemoteControl/Private/RemoteControlModule.h
+++ b/Engine/Plugins/VirtualProduction/RemoteControl/Source/RemoteControl/Private/RemoteControlModule.h
@@ -2,6 +2,7 @@
 #pragma once
 
 #include "AssetRegistry/AssetData.h"
+#include "Containers/Ticker.h"
 #include "CoreMinimal.h"
 #include "Engine/EngineTypes.h"
 #include "Factories/IRCDefaultValueFactory.h"
@@ -256,6 +257,9 @@ private:
 
 	/** Handle to the timer that that ends the ongoing change in regards to PostEditChange and transactions. */
 	FTimerHandle OngoingChangeTimer;
+	
+	/** Handle to the timer that that ends the ongoing change in regards to PostEditChange and transactions. */
+	FTSTicker::FDelegateHandle FallbackOngoingChangeTimer;
 
 	/** Delay before we check if a modification is no longer ongoing. */
 	static constexpr float SecondsBetweenOngoingChangeCheck = 0.2f;
