[RemoteControl] Respect ClampMin and ClampMax property metadata when applying delta operations

diff --git a/Engine/Plugins/VirtualProduction/RemoteControl/Source/RemoteControl/Private/RemoteControlModule.cpp b/Engine/Plugins/VirtualProduction/RemoteControl/Source/RemoteControl/Private/RemoteControlModule.cpp
--- a/Engine/Plugins/VirtualProduction/RemoteControl/Source/RemoteControl/Private/RemoteControlModule.cpp
+++ b/Engine/Plugins/VirtualProduction/RemoteControl/Source/RemoteControl/Private/RemoteControlModule.cpp
@@ -237,18 +237,20 @@ namespace RemoteControlUtil
 		const T BaseValue = Getter(BasePropertyData);
 		const T DeltaValue = Getter(DeltaPropertyData);
 
+		T NewValue;
+
 		switch (Operation)
 		{
 		case ERCModifyOperation::ADD:
-			Setter(DeltaPropertyData, BaseValue + DeltaValue);
+			NewValue = BaseValue + DeltaValue;
 			break;
 
 		case ERCModifyOperation::SUBTRACT:
-			Setter(DeltaPropertyData, BaseValue - DeltaValue);
+			NewValue = BaseValue - DeltaValue;
 			break;
 
 		case ERCModifyOperation::MULTIPLY:
-			Setter(DeltaPropertyData, BaseValue * DeltaValue);
+			NewValue = BaseValue * DeltaValue;
 			break;
 
 		case ERCModifyOperation::DIVIDE:
@@ -256,7 +258,7 @@ namespace RemoteControlUtil
 			{
 				return false;
 			}
-			Setter(DeltaPropertyData, BaseValue / DeltaValue);
+			NewValue = BaseValue / DeltaValue;
 			break;
 
 		default:
@@ -264,6 +266,26 @@ namespace RemoteControlUtil
 			return false;
 		}
 
+#if WITH_EDITORONLY_DATA
+		// Respect property's clamp settings if they exist
+		const FString& ClampMinString = Property->GetMetaData(TEXT("ClampMin"));
+		if (!ClampMinString.IsEmpty())
+		{
+			T ClampMin;
+			TTypeFromString<T>::FromString(ClampMin, *ClampMinString);
+			NewValue = FMath::Max(NewValue, ClampMin);
+		}
+
+		const FString& ClampMaxString = Property->GetMetaData(TEXT("ClampMax"));
+		if (!ClampMaxString.IsEmpty())
+		{
+			T ClampMax;
+			TTypeFromString<T>::FromString(ClampMax, *ClampMaxString);
+			NewValue = FMath::Min(NewValue, ClampMax);
+		}
+#endif
+
+		Setter(DeltaPropertyData,  NewValue);
 		return true;
 	}
 }
