[Remote Control] - Adding Field Ids and collapsed controllers list mode

- Adding Field Ids:
	- Add Field Id to URCVirtualPropertyBase
	- Add Field Id based getters to URCVirtualPropertyContainer and to URemoteControlPreset
		e.g. TArray<URCVirtualPropertyBase*> GetControllersByFieldId(const FName InFieldId)
	- Add RefreshUI() method to URemoteControlPreset, so that it's possible to refresh the UI when adding/removing controllers from code
		(this also incidentally allows to refresh the UI just once, e.g. when multiple Controllers are added or removed as a single action

- Updating Controllers List UI / adding "collapsed" mode
	- Add Field Id column
	- Add Value Type column (useful when controlling multiple controllers from a single one in "collapsed" mode)
	- Add "collapsed" mode button (allowing to just show one controller per Field Id)
	- Add FRCTypeTranslator class and FRCTypeHandler classes to facilitate the conversion of values from one Controller type to another, used in "collapsed" mode
	- Add basic spec unit for FRCTypeTranslator, for testing conversions.
	- Add OnRemoteControlPresetOpened delegate to IRemoteControlUIModule, which can be used when entities dealing with a Preset need to refresh themselves or the preset once a Preset is open. For example, a Control Object may change its Field Id when the Preset is not open, so it will need to know when the Preset is available for it to be edited.

diff --git a/Engine/Plugins/VirtualProduction/RemoteControl/Source/RemoteControl/Private/RCVirtualProperty.cpp b/Engine/Plugins/VirtualProduction/RemoteControl/Source/RemoteControl/Private/RCVirtualProperty.cpp
--- a/Engine/Plugins/VirtualProduction/RemoteControl/Source/RemoteControl/Private/RCVirtualProperty.cpp
+++ b/Engine/Plugins/VirtualProduction/RemoteControl/Source/RemoteControl/Private/RCVirtualProperty.cpp
@@ -759,6 +759,7 @@ bool URCVirtualPropertySelfContainer::DuplicatePropertyWithCopy(URCVirtualProper
 void URCVirtualPropertySelfContainer::Reset()
 {
 	PropertyName = NAME_None;
+	FieldId = NAME_None;
 	Id = FGuid();
 	Bag.Reset();
 }
@@ -771,4 +772,4 @@ const FInstancedPropertyBag* URCVirtualPropertySelfContainer::GetPropertyBagInst
 TSharedPtr<FStructOnScope> URCVirtualPropertySelfContainer::CreateStructOnScope()
 {
 	return MakeShared<FStructOnScope>(Bag.GetPropertyBagStruct(), Bag.GetMutableValue().GetMemory());
-}
\ No newline at end of file
+}
diff --git a/Engine/Plugins/VirtualProduction/RemoteControl/Source/RemoteControl/Private/RCVirtualPropertyContainer.cpp b/Engine/Plugins/VirtualProduction/RemoteControl/Source/RemoteControl/Private/RCVirtualPropertyContainer.cpp
--- a/Engine/Plugins/VirtualProduction/RemoteControl/Source/RemoteControl/Private/RCVirtualPropertyContainer.cpp
+++ b/Engine/Plugins/VirtualProduction/RemoteControl/Source/RemoteControl/Private/RCVirtualPropertyContainer.cpp
@@ -192,6 +192,49 @@ URCVirtualPropertyBase* URCVirtualPropertyContainerBase::GetVirtualPropertyByDis
 	return nullptr;
 }
 
+URCVirtualPropertyBase* URCVirtualPropertyContainerBase::GetVirtualPropertyByFieldId(const FName InFieldId) const
+{
+	for (URCVirtualPropertyBase* VirtualProperty : VirtualProperties)
+	{
+		if (VirtualProperty->FieldId == InFieldId)
+		{
+			return VirtualProperty;
+		}
+	}
+
+	return nullptr;
+}
+
+URCVirtualPropertyBase* URCVirtualPropertyContainerBase::GetVirtualPropertyByFieldId(const FName InFieldId, const EPropertyBagPropertyType InType) const
+{
+	const TArray<URCVirtualPropertyBase*>& VirtualPropertiesByFieldId = GetVirtualPropertiesByFieldId(InFieldId);
+
+	for (URCVirtualPropertyBase* VirtualProperty : VirtualPropertiesByFieldId)
+	{
+		if (VirtualProperty->GetValueType() == InType)
+		{
+			return VirtualProperty;
+		}
+	}
+
+	return nullptr;
+}
+
+TArray<URCVirtualPropertyBase*> URCVirtualPropertyContainerBase::GetVirtualPropertiesByFieldId(const FName InFieldId) const
+{
+	TArray<URCVirtualPropertyBase*> VirtualPropertiesByFieldId;
+	
+	for (URCVirtualPropertyBase* VirtualProperty : VirtualProperties)
+	{
+		if (VirtualProperty->FieldId == InFieldId)
+		{
+			VirtualPropertiesByFieldId.Add(VirtualProperty);
+		}
+	}
+
+	return VirtualPropertiesByFieldId;
+}
+
 int32 URCVirtualPropertyContainerBase::GetNumVirtualProperties() const
 {
 	const int32 NumPropertiesInBag = Bag.GetNumPropertiesInBag();
@@ -266,4 +309,4 @@ void URCVirtualPropertyContainerBase::OnModifyPropertyValue(const FPropertyChang
 }
 #endif
 
-#undef LOCTEXT_NAMESPACE
\ No newline at end of file
+#undef LOCTEXT_NAMESPACE
diff --git a/Engine/Plugins/VirtualProduction/RemoteControl/Source/RemoteControl/Private/RemoteControlPreset.cpp b/Engine/Plugins/VirtualProduction/RemoteControl/Source/RemoteControl/Private/RemoteControlPreset.cpp
--- a/Engine/Plugins/VirtualProduction/RemoteControl/Source/RemoteControl/Private/RemoteControlPreset.cpp
+++ b/Engine/Plugins/VirtualProduction/RemoteControl/Source/RemoteControl/Private/RemoteControlPreset.cpp
@@ -645,6 +645,51 @@ URCVirtualPropertyBase* URemoteControlPreset::GetControllerByDisplayName(const F
 	return ControllerContainer->GetVirtualPropertyByDisplayName(InDisplayName);
 }
 
+URCVirtualPropertyBase* URemoteControlPreset::GetControllerByFieldId(const FName InFieldId) const
+{
+	if (!ensure(ControllerContainer))
+	{
+		return nullptr;
+	}
+
+	return ControllerContainer->GetVirtualPropertyByFieldId(InFieldId);
+}
+
+URCVirtualPropertyBase* URemoteControlPreset::GetControllerByFieldId(const FName InFieldId,	const EPropertyBagPropertyType InType) const
+{
+	if (!ensure(ControllerContainer))
+	{
+		return nullptr;
+	}
+
+	return ControllerContainer->GetVirtualPropertyByFieldId(InFieldId, InType);
+}
+
+TArray<URCVirtualPropertyBase*> URemoteControlPreset::GetControllersByFieldId(const FName InFieldId) const
+{
+	if (!ensure(ControllerContainer))
+	{
+		return {};
+	}
+
+	return ControllerContainer->GetVirtualPropertiesByFieldId(InFieldId);
+}
+
+TArray<EPropertyBagPropertyType> URemoteControlPreset::GetControllersTypesByFieldId(const FName InFieldId) const
+{	
+	TArray<EPropertyBagPropertyType> ValueTypes;
+	
+	for (const URCVirtualPropertyBase* Controller : GetControllersByFieldId(InFieldId))
+	{
+		if (Controller)
+		{
+			ValueTypes.Add(Controller->GetValueType());
+		}
+	}
+
+	return ValueTypes;
+}
+
 URCVirtualPropertyInContainer* URemoteControlPreset::AddController(TSubclassOf<URCVirtualPropertyInContainer> InPropertyClass, const EPropertyBagPropertyType InValueType, UObject* InValueTypeObject /*= nullptr*/, const FName InPropertyName /*= NAME_None*/)
 {
 	if (!ensure(ControllerContainer))
diff --git a/Engine/Plugins/VirtualProduction/RemoteControl/Source/RemoteControl/Public/RCVirtualProperty.h b/Engine/Plugins/VirtualProduction/RemoteControl/Source/RemoteControl/Public/RCVirtualProperty.h
--- a/Engine/Plugins/VirtualProduction/RemoteControl/Source/RemoteControl/Public/RCVirtualProperty.h
+++ b/Engine/Plugins/VirtualProduction/RemoteControl/Source/RemoteControl/Public/RCVirtualProperty.h
@@ -302,6 +302,10 @@ public:
 	UPROPERTY()
 	FGuid Id;
 
+	/** Property Field Id */
+	UPROPERTY()
+	FName FieldId;
+
 	/** Pointer to Remote Control Preset */
 	UPROPERTY()
 	TWeakObjectPtr<URemoteControlPreset> PresetWeakPtr;
@@ -403,4 +407,4 @@ private:
 	/** Instanced property bag for store a bag of properties. */
 	UPROPERTY()
 	FInstancedPropertyBag Bag;
-};
\ No newline at end of file
+};
diff --git a/Engine/Plugins/VirtualProduction/RemoteControl/Source/RemoteControl/Public/RCVirtualPropertyContainer.h b/Engine/Plugins/VirtualProduction/RemoteControl/Source/RemoteControl/Public/RCVirtualPropertyContainer.h
--- a/Engine/Plugins/VirtualProduction/RemoteControl/Source/RemoteControl/Public/RCVirtualPropertyContainer.h
+++ b/Engine/Plugins/VirtualProduction/RemoteControl/Source/RemoteControl/Public/RCVirtualPropertyContainer.h
@@ -97,6 +97,21 @@ public:
 	 */
 	virtual URCVirtualPropertyBase* GetVirtualPropertyByDisplayName(const FName InDisplayName) const;
 
+	/**
+	 * Returns the first found virtual property matching the specified Field Id
+	 */
+	virtual URCVirtualPropertyBase* GetVirtualPropertyByFieldId(const FName InFieldId) const;
+	
+	/**
+	 * Returns the first virtual property matching the specified Field Id and ValueType
+	 */
+	virtual URCVirtualPropertyBase* GetVirtualPropertyByFieldId(const FName InFieldId, const EPropertyBagPropertyType InType) const;
+
+	/**
+	 * Returns all virtual properties with the specified Field Id
+	 */
+	virtual TArray<URCVirtualPropertyBase*> GetVirtualPropertiesByFieldId(const FName InFieldId) const;
+
 	/**
 	 * Returns number of virtual properties.
 	 */ 
diff --git a/Engine/Plugins/VirtualProduction/RemoteControl/Source/RemoteControl/Public/RemoteControlPreset.h b/Engine/Plugins/VirtualProduction/RemoteControl/Source/RemoteControl/Public/RemoteControlPreset.h
--- a/Engine/Plugins/VirtualProduction/RemoteControl/Source/RemoteControl/Public/RemoteControlPreset.h
+++ b/Engine/Plugins/VirtualProduction/RemoteControl/Source/RemoteControl/Public/RemoteControlPreset.h
@@ -646,6 +646,18 @@ public:
 	/** Fetches a virtual property by specified name. */
 	URCVirtualPropertyBase* GetControllerByDisplayName(const FName InDisplayName) const;
 
+	/** Fetches the first found Controller matching the specified Field Id */
+	URCVirtualPropertyBase* GetControllerByFieldId(const FName InFieldId) const;
+
+	/** Fetches the first Controller matching the specified Field Id and ValueType */
+	URCVirtualPropertyBase* GetControllerByFieldId(const FName InFieldId, const EPropertyBagPropertyType InType) const;
+
+	/** Fetches an array of all Controllers with the specified Field Id. */
+	TArray<URCVirtualPropertyBase*> GetControllersByFieldId(const FName InFieldId) const;
+
+	/** Fetches an array of all Controllers Types with the specified Field Id. */
+	TArray<EPropertyBagPropertyType> GetControllersTypesByFieldId(const FName InFieldId) const;
+
 	/** Adds a Virtual Property (Controller) to the Remote Control Preset */
 	URCVirtualPropertyInContainer* AddController(TSubclassOf<URCVirtualPropertyInContainer> InPropertyClass, const EPropertyBagPropertyType InValueType, UObject* InValueTypeObject = nullptr, const FName InPropertyName = NAME_None);
 
diff --git a/Engine/Plugins/VirtualProduction/RemoteControl/Source/RemoteControlUI/Private/RemoteControlUIModule.cpp b/Engine/Plugins/VirtualProduction/RemoteControl/Source/RemoteControlUI/Private/RemoteControlUIModule.cpp
--- a/Engine/Plugins/VirtualProduction/RemoteControl/Source/RemoteControlUI/Private/RemoteControlUIModule.cpp
+++ b/Engine/Plugins/VirtualProduction/RemoteControl/Source/RemoteControlUI/Private/RemoteControlUIModule.cpp
@@ -364,6 +364,8 @@ TSharedRef<SRemoteControlPanel> FRemoteControlUIModule::CreateRemoteControlPanel
 		}
 	}
 
+	OnRemoteControlPresetOpened().Broadcast(Preset);
+	
 	return PanelRef;
 }
 
diff --git a/Engine/Plugins/VirtualProduction/RemoteControl/Source/RemoteControlUI/Private/RemoteControlUIModule.h b/Engine/Plugins/VirtualProduction/RemoteControl/Source/RemoteControlUI/Private/RemoteControlUIModule.h
--- a/Engine/Plugins/VirtualProduction/RemoteControl/Source/RemoteControlUI/Private/RemoteControlUIModule.h
+++ b/Engine/Plugins/VirtualProduction/RemoteControl/Source/RemoteControlUI/Private/RemoteControlUIModule.h
@@ -47,6 +47,7 @@ public:
 
 	//~ Begin IRemoteControlUIModule interface
 	virtual FOnGenerateExtensions& GetExtensionGenerators() override { return ExtensionsGenerator; }
+	virtual FOnRemoteControlPresetOpened& OnRemoteControlPresetOpened() override { return RemoteControlPresetOpenedDelegate; }
 	virtual FDelegateHandle AddPropertyFilter(FOnDisplayExposeIcon OnDisplayExposeIcon) override;
 	virtual void RemovePropertyFilter(const FDelegateHandle& FilterDelegateHandle) override;
 	virtual void RegisterMetadataCustomization(FName MetadataKey, FOnCustomizeMetadataEntry OnCustomizeCallback) override;
@@ -207,6 +208,9 @@ private:
 	/** Delegate called to gather extensions added externally to the panel. */
 	FOnGenerateExtensions ExtensionsGenerator;
 
+	/** Delegate called when a Remote Control Preset panel is created for the specified Remote Control Preset */
+	FOnRemoteControlPresetOpened RemoteControlPresetOpenedDelegate;
+
 	/** Filters added by other plugins queried to determine if a property should display an expose icon. */
 	TMap<FDelegateHandle, FOnDisplayExposeIcon> ExternalFilterDelegates;
 
diff --git a/Engine/Plugins/VirtualProduction/RemoteControl/Source/RemoteControlUI/Private/UI/Controller/RCControllerModel.cpp b/Engine/Plugins/VirtualProduction/RemoteControl/Source/RemoteControlUI/Private/UI/Controller/RCControllerModel.cpp
--- a/Engine/Plugins/VirtualProduction/RemoteControl/Source/RemoteControlUI/Private/UI/Controller/RCControllerModel.cpp
+++ b/Engine/Plugins/VirtualProduction/RemoteControl/Source/RemoteControlUI/Private/UI/Controller/RCControllerModel.cpp
@@ -3,9 +3,11 @@
 #include "RCControllerModel.h"
 
 #include "IDetailTreeNode.h"
-#include "RemoteControlPreset.h"
 #include "RCVirtualProperty.h"
+#include "RemoteControlPreset.h"
+#include "TypeTranslator/RCTypeTranslator.h"
 #include "UI/SRemoteControlPanel.h"
+#include "Widgets/Input/STextComboBox.h"
 #include "Widgets/Text/SInlineEditableTextBlock.h"
 
 #define LOCTEXT_NAMESPACE "FRCControllerModel"
@@ -15,14 +17,25 @@ FRCControllerModel::FRCControllerModel(URCVirtualPropertyBase* InVirtualProperty
 	, VirtualPropertyWeakPtr(InVirtualProperty)
 	, DetailTreeNodeWeakPtr(InTreeNode)
 {
-	if(ensure(InVirtualProperty))
+	if (ensure(InVirtualProperty))
 	{
-		if(InVirtualProperty->DisplayName.IsNone())
+		if (InVirtualProperty->DisplayName.IsNone())
+		{
 			InVirtualProperty->DisplayName = InVirtualProperty->PropertyName;
-
+		}
+		
 		SAssignNew(ControllerNameTextBox, SInlineEditableTextBlock)
 			.Text(FText::FromName(InVirtualProperty->DisplayName))
 			.OnTextCommitted_Raw(this, &FRCControllerModel::OnControllerNameCommitted);
+
+		SAssignNew(ControllerFieldIdTextBox, SInlineEditableTextBlock)
+			.Text(FText::FromName(InVirtualProperty->FieldId))
+			.OnTextCommitted_Raw(this, &FRCControllerModel::OnControllerFieldIdCommitted);
+
+		if (const TSharedPtr<IPropertyHandle>& PropertyHandle = InTreeNode->CreatePropertyHandle())
+		{
+			PropertyHandle->SetOnPropertyValueChangedWithData(TDelegate<void(const FPropertyChangedEvent&)>::CreateRaw(this, &FRCControllerModel::OnPropertyValueChanged));
+		}
 	}
 
 	Id = FGuid::NewGuid();
@@ -62,6 +75,46 @@ TSharedRef<SWidget> FRCControllerModel::GetNameWidget() const
 		];
 }
 
+TSharedRef<SWidget> FRCControllerModel::GetFieldIdWidget() const
+{
+	return SNew(SBox).Padding(10.f, 2.f)
+		[
+			ControllerFieldIdTextBox.ToSharedRef()
+		];
+}
+
+TSharedRef<SWidget> FRCControllerModel::GetTypeSelectionWidget()
+{
+	if (const URCVirtualPropertyBase* Controller = GetVirtualProperty())
+	{
+		const FName& FieldId = Controller->FieldId;
+		const TArray<URCVirtualPropertyBase*>& Controllers = GetPreset()->GetControllersByFieldId(FieldId);
+		
+		if (Controllers.Num() > 1 && bIsMultiController)
+		{
+			const EPropertyBagPropertyType OptimalValueType = FRCTypeTranslator::GetOptimalValueType(GetPreset()->GetControllersTypesByFieldId(FieldId));
+			CurrentControlValueType = OptimalValueType;
+
+			return SNew(SBox).Padding(10.f, 2.f)
+			[
+				SNew(STextComboBox)
+				.OptionsSource(&ControlledTypesAsStrings)
+				.OnSelectionChanged(this, &FRCControllerModel::OnTextControlValueTypeChanged)
+				.InitiallySelectedItem(ControlledTypesAsStrings[CurrentControlValueTypeIndex])
+			];
+		}
+	}
+
+	return SNullWidget::NullWidget;
+}
+
+void FRCControllerModel::SetMultiController(bool bInIsMultiController)
+{
+	bIsMultiController = bInIsMultiController;
+	
+	InitControlledTypes();
+}
+
 URCVirtualPropertyBase* FRCControllerModel::GetVirtualProperty() const
 {
 	return VirtualPropertyWeakPtr.Get();
@@ -91,9 +144,9 @@ void FRCControllerModel::OnControllerNameCommitted(const FText& InNewControllerN
 {
 	if (URemoteControlPreset* Preset = GetPreset())
 	{
-		if(URCVirtualPropertyBase* Controller = GetVirtualProperty())
+		if (URCVirtualPropertyBase* Controller = GetVirtualProperty())
 		{
-			FName OldName = Controller->DisplayName;
+			const FName& OldName = Controller->DisplayName;
 			Controller->DisplayName = *InNewControllerName.ToString();
 			ControllerNameTextBox->SetText(InNewControllerName);
 			Preset->OnControllerRenamed().Broadcast(Preset, OldName, *InNewControllerName.ToString());
@@ -101,8 +154,88 @@ void FRCControllerModel::OnControllerNameCommitted(const FText& InNewControllerN
 	}
 }
 
-void FRCControllerModel::EnterRenameMode()
+void FRCControllerModel::OnControllerFieldIdCommitted(const FText& InNewControllerFieldId, ETextCommit::Type InCommitInfo)
 {
+	if (URemoteControlPreset* Preset = GetPreset())
+	{
+		if (URCVirtualPropertyBase* Controller = GetVirtualProperty())
+		{
+			// todo: think about how to setup this Field Id change feature
+			// We might need some delegate firing when a Field Id update occurs e.g
+			// OnControllerFieldIdChangedFromPanel(FName OldFieldId, FName NewFieldId)
+		}
+	}
+}
+
+void FRCControllerModel::OnTextControlValueTypeChanged(TSharedPtr<FString, ESPMode::ThreadSafe> InControlValueTypeString, ESelectInfo::Type Arg)
+{
+	if (ControlledTypesAsStrings.Contains(InControlValueTypeString))
+	{
+		CurrentControlValueTypeIndex = ControlledTypesAsStrings.IndexOfByKey(InControlValueTypeString);
+	}
+	else
+	{
+		CurrentControlValueTypeIndex = 0;
+	}
+
+	const FString& Value = *InControlValueTypeString.Get();
+	CurrentControlValueType = static_cast<EPropertyBagPropertyType>(StaticEnum<EPropertyBagPropertyType>()->GetValueByNameString(Value));
+	
+	if (OnValueTypeChanged.IsBound())
+	{
+		if (URCVirtualPropertyBase* Controller = GetVirtualProperty())
+		{
+			if (OnValueTypeChanged.IsBound())
+			{
+				OnValueTypeChanged.Broadcast(Controller, CurrentControlValueType);
+			}
+		}
+	}
+}
+
+void FRCControllerModel::OnPropertyValueChanged(const FPropertyChangedEvent& InPropertyChangedEvent)
+{	
+	if (URCVirtualPropertyBase* ControllerProperty = GetVirtualProperty())
+	{
+		if (OnValueChanged.IsBound())
+		{
+			OnValueChanged.Broadcast(ControllerProperty);
+		}
+	}
+}
+
+void FRCControllerModel::InitControlledTypes()
+{
+	ControlledTypesAsStrings.Empty();
+	
+	if (const URCVirtualPropertyBase* Controller = GetVirtualProperty())
+	{
+		const TArray<EPropertyBagPropertyType>& ValueTypes = GetPreset()->GetControllersTypesByFieldId(Controller->FieldId);
+
+		int32 Count = 0;
+		TArray<FString> Types;
+		for (const EPropertyBagPropertyType ValueType : ValueTypes)
+		{
+			const FText& TypeName = UEnum::GetDisplayValueAsText(ValueType);
+			Types.AddUnique(TypeName.ToString());
+
+			if (Controller->GetValueType() == ValueType)
+			{
+				CurrentControlValueTypeIndex = Count;
+			}
+
+			Count++;
+		}
+
+		for (const FString& TypeName : Types)
+		{
+			ControlledTypesAsStrings.Add(MakeShared<FString>(TypeName));
+		}
+	}
+}
+
+void FRCControllerModel::EnterRenameMode()
+{	
 	ControllerNameTextBox->EnterEditingMode();
 }
 
@@ -116,4 +249,4 @@ FName FRCControllerModel::GetControllerDisplayName()
 	return NAME_None;
 }
 
-#undef LOCTEXT_NAMESPACE
\ No newline at end of file
+#undef LOCTEXT_NAMESPACE
diff --git a/Engine/Plugins/VirtualProduction/RemoteControl/Source/RemoteControlUI/Private/UI/Controller/RCControllerModel.h b/Engine/Plugins/VirtualProduction/RemoteControl/Source/RemoteControlUI/Private/UI/Controller/RCControllerModel.h
--- a/Engine/Plugins/VirtualProduction/RemoteControl/Source/RemoteControlUI/Private/UI/Controller/RCControllerModel.h
+++ b/Engine/Plugins/VirtualProduction/RemoteControl/Source/RemoteControlUI/Private/UI/Controller/RCControllerModel.h
@@ -7,6 +7,7 @@
 #include "Widgets/DeclarativeSyntaxSupport.h"
 #include "Widgets/SWidget.h"
 
+enum class EPropertyBagPropertyType : uint8;
 class FRCBehaviourModel;
 class IDetailTreeNode;
 class SInlineEditableTextBlock;
@@ -21,20 +22,38 @@ class URCVirtualPropertyBase;
 * ~ FRCControllerModel ~
 *
 * UI model for representing a Controller
-* Contains a row widget with Controller Name and a Value Widget
+* Contains a row widget with Controller Name, Field Id and a Value Widget.
+* If this is setup to be a MultiController, a Value Type selection Widget will be available as well.
 */
 class FRCControllerModel : public FRCLogicModeBase
 {
 public:
+	DECLARE_MULTICAST_DELEGATE_TwoParams(FOnValueTypeChanged, URCVirtualPropertyBase* /* InController */, EPropertyBagPropertyType /* InValueType */);
+	DECLARE_MULTICAST_DELEGATE_OneParam(FOnValueChanged, URCVirtualPropertyBase* /* InController */);
+	
 	FRCControllerModel(URCVirtualPropertyBase* InVirtualProperty, const TSharedRef<IDetailTreeNode>& InTreeNode, const TSharedPtr<SRemoteControlPanel> InRemoteControlPanel);
 	
-	/** The widget to be rendered for this Controller
-	* Used to represent a single row when added to the Controllers Panel List
-	*/
+	/**
+	 * The widget to be rendered for this Controller
+	 * Used to represent a single row when added to the Controllers Panel List
+	 */
 	virtual TSharedRef<SWidget> GetWidget() const override;
 
+	/** The widget showing the Name of this Controller */
 	TSharedRef<SWidget> GetNameWidget() const;
 
+	/** The widget showing the Field Id of this Controller */
+	TSharedRef<SWidget> GetFieldIdWidget() const;
+
+	/** The widget showing the Value Type selection dropdown for this Controller (used for MultiControllers) */
+	TSharedRef<SWidget> GetTypeSelectionWidget();
+
+	/**
+	 * Specify if this is a MultiController handling other hidden controllers.
+	 * MultiControllers will setup and show a Type Selection Widget.
+	 */
+	void SetMultiController(bool bInIsMultiController);
+
 	/** Returns the Controller (Data Model) associated with us*/
 	virtual URCVirtualPropertyBase* GetVirtualProperty() const;
 	
@@ -56,10 +75,28 @@ public:
 	/**Fetches the unique Id for this UI item */
 	FGuid GetId() const { return Id; }
 
+	/** Triggered when the Value Type of this changes (for MultiControllers) */
+	FOnValueTypeChanged OnValueTypeChanged;
+
+	/** Triggered when the Value this controller changes */
+	FOnValueChanged OnValueChanged;
+
 private:
-	/* Text commit event for Controller Name text box*/
+	/** Text commit event for Controller Name text box */
 	void OnControllerNameCommitted(const FText& InNewControllerName, ETextCommit::Type InCommitInfo);
 
+	/** Text commit event for Controller FieldId text box */
+	void OnControllerFieldIdCommitted(const FText& InNewControllerFieldId, ETextCommit::Type InCommitInfo);
+
+	/** Value Type change event, in case this is a MultiController */
+	void OnTextControlValueTypeChanged(TSharedPtr<FString, ESPMode::ThreadSafe> InControlValueTypeString, ESelectInfo::Type Arg);
+
+	/** Value change event */
+	void OnPropertyValueChanged(const FPropertyChangedEvent& InPropertyChangedEvent);
+
+	/** Initializes the list of controlled types. Used by MultiControllers */
+	void InitControlledTypes();
+
 	/** Virtual Property (Data Model) represented by us*/
 	TWeakObjectPtr<URCVirtualPropertyBase> VirtualPropertyWeakPtr;
 	
@@ -72,6 +109,21 @@ private:
 	/** Controller name - editable text box */
 	TSharedPtr<SInlineEditableTextBlock> ControllerNameTextBox;
 
+	/** Controller Field Id - editable text box */
+	TSharedPtr<SInlineEditableTextBlock> ControllerFieldIdTextBox;
+
+	/** Holds a list of controlled types, represented as strings. Used to show the Value Type selection dropdown */
+	TArray<TSharedPtr<FString>> ControlledTypesAsStrings;
+
+	/** Index representing the currently selected Value Type within the ControlledTypesAsStrings array */
+	int32 CurrentControlValueTypeIndex = 0;
+
+	/** Current control value type */
+	EPropertyBagPropertyType CurrentControlValueType;
+
+	/** Is this a multi controller?*/
+	bool bIsMultiController = false;
+
 	/**Unique Id for this UI item*/
 	FGuid Id;
 };
diff --git a/Engine/Plugins/VirtualProduction/RemoteControl/Source/RemoteControlUI/Private/UI/Controller/RCMultiController.cpp b/Engine/Plugins/VirtualProduction/RemoteControl/Source/RemoteControlUI/Private/UI/Controller/RCMultiController.cpp
new file mode 100644
--- /dev/null
+++ b/Engine/Plugins/VirtualProduction/RemoteControl/Source/RemoteControlUI/Private/UI/Controller/RCMultiController.cpp
@@ -0,0 +1,215 @@
+﻿// Copyright Epic Games, Inc. All Rights Reserved.
+
+#include "RCMultiController.h"
+#include "Controller/RCController.h"
+#include "PropertyBag.h"
+#include "TypeTranslator/RCTypeTranslator.h"
+#include "RCVirtualProperty.h"
+#include "RemoteControlPreset.h"
+
+FRCMultiController::FRCMultiController(const TWeakObjectPtr<URCVirtualPropertyBase>& InController)
+{
+	if (InController.IsValid())
+	{
+		MainControllerWeak = InController;
+		ValueType = InController->GetValueType();
+	}
+}
+
+void FRCMultiController::AddHandledController(URCVirtualPropertyBase* InController)
+{
+	if (MainControllerWeak.IsValid())
+	{
+		if (URCController* Controller = Cast<URCController>(InController))
+		{
+			HandledControllers.AddUnique(InController);
+		}
+	}
+}
+
+void FRCMultiController::SetMainController(URCVirtualPropertyBase* InController)
+{
+	MainControllerWeak = InController;
+	ValueType = MainControllerWeak->GetValueType();
+}
+
+void FRCMultiController::UpdateHandledControllersValue()
+{
+	if (MainControllerWeak.IsValid())
+	{
+		if (const URCVirtualPropertyBase* MultiController = MainControllerWeak.Get())
+		{
+			FRCTypeTranslator::Get()->Translate(MultiController, HandledControllers);
+		}
+	}
+}
+
+FName FRCMultiController::GetFieldId() const
+{
+	if (MainControllerWeak.IsValid())
+	{
+		return MainControllerWeak->FieldId;
+	}
+
+	return NAME_None;	
+}
+
+FGuid FRCMultiController::GetMainControllerId() const
+{
+	if (MainControllerWeak.IsValid())
+	{
+		return MainControllerWeak->Id;
+	}
+
+	return FGuid();
+}
+
+bool FRCMultiControllersState::TryToAddAsMultiController(URCVirtualPropertyBase* InController)
+{
+	if (!InController)
+	{
+		return false;
+	}
+
+	if (!InController->PresetWeakPtr.IsValid())
+	{
+		return false;
+	}
+	
+	const URemoteControlPreset* const Preset = InController->PresetWeakPtr.Get();
+	const FName& FieldId = InController->FieldId;
+	
+	TArray<URCVirtualPropertyBase*> ControllersWithSameFieldId = Preset->GetControllersByFieldId(FieldId);
+	if (ControllersWithSameFieldId.Num() > 1)
+	{
+		if (MultiControllersMap.Contains(FieldId))
+		{
+			// This Field Id has already been handled, no need to add a new MultiController for it
+			return false;
+		}
+		
+		// Retrieve currently selected value type for this Field Id
+		const EPropertyBagPropertyType* FieldIdTypePtr = SelectedValueTypes.Find(FieldId);
+		EPropertyBagPropertyType FieldIdType;
+	
+		if (!FieldIdTypePtr)
+		{
+			// If type is not set for this Field Id, let's set it to the optimal one
+			FieldIdType = UE::MultiControllers::Public::GetOptimalValueTypeForFieldId(Preset, FieldId);
+			SelectedValueTypes.Emplace(FieldId, FieldIdType);
+		}
+		else
+		{
+			// Otherwise, retrieve the currently selected value type
+			FieldIdType = *FieldIdTypePtr;
+		}
+
+		// Retrieve value type of controller we are trying to add as MultiController
+		const EPropertyBagPropertyType ControllerValueType = InController->GetValueType();
+		const TArray<EPropertyBagPropertyType>& FieldIdControlledTypes = UE::MultiControllers::Public::GetFieldIdValueTypes(Preset, FieldId);
+		if (!FieldIdControlledTypes.Contains(FieldIdType))
+		{
+			// If SelectedValueTypes contains a type which is not actually handled by any controller with this Field Id,
+			// then it's outdated - ignore and remove misleading info
+			SelectedValueTypes.Remove(FieldId);
+		}
+		else if (ControllerValueType != FieldIdType)
+		{
+			// This controller has a value type different from the selected one
+			return false;
+		}
+
+		// We got here, so it means we can add a new MultiController
+		FRCMultiController NewMultiController(InController);
+		
+		for (URCVirtualPropertyBase* ControllerToHandle : ControllersWithSameFieldId)
+		{
+			// Skip self (will have same type)
+			if (ControllerToHandle->GetValueType() != InController->GetValueType())
+			{
+				NewMultiController.AddHandledController(ControllerToHandle);
+			}
+		}
+
+		MultiControllersMap.Add(FieldId, NewMultiController);
+
+		return true;
+	}
+
+	return false;
+}
+
+void FRCMultiControllersState::UpdateFieldIdValueType(const FName& InFieldId, EPropertyBagPropertyType InValueType)
+{
+	SelectedValueTypes.Emplace(InFieldId, InValueType);
+	RemoveMultiController(InFieldId);
+}
+
+EPropertyBagPropertyType FRCMultiControllersState::GetCurrentType(const FName& InFieldId)
+{
+	if (const EPropertyBagPropertyType* const ValueTypePtr = SelectedValueTypes.Find(InFieldId))
+	{
+		return *ValueTypePtr;
+	}
+
+	return EPropertyBagPropertyType::None;
+}
+
+FRCMultiController FRCMultiControllersState::GetMultiController(const FName& InFieldId)
+{
+	if (const FRCMultiController* const MultiController = MultiControllersMap.Find(InFieldId))
+	{
+		return *MultiController;
+	}
+
+	return FRCMultiController(nullptr);
+}
+
+void FRCMultiControllersState::ResetMultiControllers()
+{
+	MultiControllersMap.Empty();
+}
+
+void FRCMultiControllersState::ResetSelectedValueTypes()
+{
+	SelectedValueTypes.Empty();
+}
+
+void FRCMultiControllersState::RemoveMultiController(const FName& InFieldId)
+{
+	if (MultiControllersMap.Contains(InFieldId))
+	{		
+		MultiControllersMap.Remove(InFieldId);
+	}
+}
+
+EPropertyBagPropertyType UE::MultiControllers::Public::GetOptimalValueTypeForFieldId(const URemoteControlPreset* InPreset, const FName& InFieldId)
+{
+	EPropertyBagPropertyType OptimalType = EPropertyBagPropertyType::Bool;
+	
+	if (InPreset)
+	{
+		const TArray<EPropertyBagPropertyType>& ValueTypes = InPreset->GetControllersTypesByFieldId(InFieldId);
+		OptimalType = FRCTypeTranslator::Get()->GetOptimalValueType(ValueTypes);
+	}
+
+	return OptimalType;
+}
+
+TArray<EPropertyBagPropertyType> UE::MultiControllers::Public::GetFieldIdValueTypes(const URemoteControlPreset* InPreset, const FName& InFieldId)
+{
+	if (!InPreset)
+	{
+		return {};
+	}
+	
+	TArray<EPropertyBagPropertyType> ValueTypes;
+	const TArray<URCVirtualPropertyBase*>& Controllers = InPreset->GetControllersByFieldId(InFieldId);
+	
+	for (const URCVirtualPropertyBase* Controller : Controllers)
+	{
+		ValueTypes.Add(Controller->GetValueType());
+	}
+
+	return ValueTypes;
+}
diff --git a/Engine/Plugins/VirtualProduction/RemoteControl/Source/RemoteControlUI/Private/UI/Controller/RCMultiController.h b/Engine/Plugins/VirtualProduction/RemoteControl/Source/RemoteControlUI/Private/UI/Controller/RCMultiController.h
new file mode 100644
--- /dev/null
+++ b/Engine/Plugins/VirtualProduction/RemoteControl/Source/RemoteControlUI/Private/UI/Controller/RCMultiController.h
@@ -0,0 +1,88 @@
+﻿// Copyright Epic Games, Inc. All Rights Reserved.
+
+#pragma once
+#include "PropertyBag.h"
+
+class URCBehaviour;
+class URCVirtualPropertyBase;
+class URemoteControlPreset;
+
+/**
+ * A MultiController keeps a reference to a Main Controller and to a set of Handled Controllers with the same Field Id.
+ * A MultiController can try to use its MainController value to update the value of its Handled Controllers.
+ * This is useful when Control Panel List is in MultiController Mode, and Controllers with duplicate Field Ids are hidden.
+ */
+struct FRCMultiController
+{
+public:
+	FRCMultiController(const TWeakObjectPtr<URCVirtualPropertyBase>& InController);
+
+	/** Sets the provided controller as the MainController, which will drive other Handled Controllers */
+	void SetMainController(URCVirtualPropertyBase* InController);
+
+	/** Adds a controller to the list of Handled Controllers managed by the MainController of this MultiController setup */
+	void AddHandledController(URCVirtualPropertyBase* InController);
+
+	/** Goes through each Handled Controller, and tries to update its value based on the current value of the MainController*/
+	void UpdateHandledControllersValue();
+
+	/** MultiController is considered valid when its MainController is valid */
+	bool IsValid() const { return MainControllerWeak.IsValid(); }
+
+	/** Returns the Field Id this MultiController is associated with */
+	FName GetFieldId() const;
+
+	/** Return the Id of the MainController */
+	FGuid GetMainControllerId() const;
+
+	/** Returns the value type of the MainController */
+	EPropertyBagPropertyType GetValueType() const { return ValueType; }
+
+private:
+	TWeakObjectPtr<URCVirtualPropertyBase> MainControllerWeak;
+	EPropertyBagPropertyType ValueType;
+	TArray<TObjectPtr<URCVirtualPropertyBase>> HandledControllers;
+};
+
+/**
+ * Holds and handles a Map of MultiControllers.
+ * Also keeps reference of currently selected Value Type per each Field Id.
+ */
+struct FRCMultiControllersState
+{
+	bool TryToAddAsMultiController(URCVirtualPropertyBase* InController);
+	
+	EPropertyBagPropertyType GetCurrentType(const FName& InFieldId);
+	FRCMultiController GetMultiController(const FName& InFieldId);
+
+	/**
+	 * Calling this function will overwrite the currently selected type for the specified field id.
+	 * FRCMultiControllersState will then be in an invalid state, so it will require the hosting ControllerPanelList to call Reset()
+	 */
+	void UpdateFieldIdValueType(const FName& InFieldId, EPropertyBagPropertyType InValueType);
+
+	/** Empties the internal map referencing the handled MultiControllers */
+	void ResetMultiControllers();
+
+	/** Empties the internal map used to store the currently selected value type for each Field Id */
+	void ResetSelectedValueTypes();
+
+	/** Removes the MultiController associated with the specified Field Id */	
+	void RemoveMultiController(const FName& InFieldId);
+
+private:
+	/** Map of MultiControllers, indexed by their Field Id*/
+	TMap<FName, FRCMultiController> MultiControllersMap;
+
+	/** Used to store the currently chosen value type. This is stored outside of FRCMultiController, so that the info is kept between UI Refreshes */
+	TMap<FName, EPropertyBagPropertyType> SelectedValueTypes;
+};
+
+namespace UE::MultiControllers::Public
+{
+	/** Returns the optimal value type for the specified Field Id */
+	EPropertyBagPropertyType GetOptimalValueTypeForFieldId(const URemoteControlPreset* InPreset, const FName& InFieldId);
+
+	/** Returns a list of all the Value Types currently associated to the specified Field Id */
+	TArray<EPropertyBagPropertyType> GetFieldIdValueTypes(const URemoteControlPreset* InPreset, const FName& InFieldId);
+}
diff --git a/Engine/Plugins/VirtualProduction/RemoteControl/Source/RemoteControlUI/Private/UI/Controller/SRCControllerPanel.cpp b/Engine/Plugins/VirtualProduction/RemoteControl/Source/RemoteControlUI/Private/UI/Controller/SRCControllerPanel.cpp
--- a/Engine/Plugins/VirtualProduction/RemoteControl/Source/RemoteControlUI/Private/UI/Controller/SRCControllerPanel.cpp
+++ b/Engine/Plugins/VirtualProduction/RemoteControl/Source/RemoteControlUI/Private/UI/Controller/SRCControllerPanel.cpp
@@ -95,6 +95,7 @@ void SRCControllerPanel::Construct(const FArguments& InArgs, const TSharedRef<SR
 
 	ControllerDockPanel->AddHeaderToolbarItem(EToolbar::Left, AddNewControllerButton);
 	ControllerDockPanel->AddFooterToolbarItem(EToolbar::Right, DeleteSelectedControllerButton);
+	ControllerDockPanel->AddHeaderToolbarItem(EToolbar::Right, GetMultiControllerSwitchWidget());
 
 	ChildSlot
 		.Padding(RCPanelStyle->PanelPadding)
@@ -263,6 +264,47 @@ FReply SRCControllerPanel::RequestDeleteAllItems()
 	return FReply::Handled();
 }
 
+void SRCControllerPanel::OnToggleMultiControllersMode(ECheckBoxState CheckBoxState)
+{
+	switch (CheckBoxState)
+	{
+		case ECheckBoxState::Unchecked:
+			bIsMultiControllerMode = false;
+			break;
+
+		case ECheckBoxState::Checked:
+			bIsMultiControllerMode = true;
+			break;
+
+		case ECheckBoxState::Undetermined:
+			bIsMultiControllerMode = false;
+			break;
+	}
+
+	if (ControllerPanelList.IsValid())
+	{
+		ControllerPanelList->SetMultiControllerMode(bIsMultiControllerMode);
+	}
+}
+
+EVisibility SRCControllerPanel::GetMultiControllerSwitchVisibility() const
+{
+	if (const URemoteControlPreset* const Preset = GetPreset())
+	{
+		const int32 NoFieldIdControllersNum = Preset->GetControllersByFieldId(NAME_None).Num();
+
+		// if no controller has a Field Id, let's hide the MultiController Switch
+		if (NoFieldIdControllersNum == Preset->GetNumControllers())
+		{
+			return EVisibility::Collapsed;
+			//todo: this could be smarter, e.g. showing the switch only if there are controllers with the same Field Id
+		}
+	}
+
+	return EVisibility::Visible;
+}
+
+
 TSharedRef<SWidget> SRCControllerPanel::GetControllerMenuContentWidget() const
 {
 	constexpr  bool bShouldCloseWindowAfterMenuSelection = true;
@@ -336,6 +378,28 @@ TSharedRef<SWidget> SRCControllerPanel::GetControllerMenuContentWidget() const
 	return MenuBuilder.MakeWidget();
 }
 
+TSharedRef<SWidget> SRCControllerPanel::GetMultiControllerSwitchWidget()
+{
+	const FText& SwitchTooltipText = LOCTEXT("MultiControllerModeSwitchTooltip", "Enable/Disable MultiController Mode.\nWhen enabled, the Controllers list will only show one Controller per Field Id.");
+
+	return
+		SNew(SCheckBox)
+		.Visibility(this, &SRCControllerPanel::GetMultiControllerSwitchVisibility)
+		.AddMetaData<FTagMetaData>(FTagMetaData(TEXT("Toggle MultiController Mode")))
+		.ToolTipText(SwitchTooltipText)
+		.HAlign(HAlign_Center)
+		.Style(&RCPanelStyle->ToggleButtonStyle)
+		.ForegroundColor(FSlateColor::UseForeground())
+		.IsChecked_Lambda([this]() { return bIsMultiControllerMode ? ECheckBoxState::Checked : ECheckBoxState::Unchecked; })
+		.OnCheckStateChanged(this, &SRCControllerPanel::OnToggleMultiControllersMode)
+		[
+			SNew(STextBlock)
+			.Text(LOCTEXT("MultiControllersModeSwitchLabel", "Multi"))
+			.Justification(ETextJustify::Center)
+			.TextStyle(&RCPanelStyle->PanelTextStyle)
+		];
+}
+
 void SRCControllerPanel::OnAddControllerClicked(const EPropertyBagPropertyType InValueType, UObject* InValueTypeObject) const
 {
 	// Add to the asset
@@ -384,4 +448,4 @@ FReply SRCControllerPanel::OnClickEmptyButton()
 	return FReply::Handled();
 }
 
-#undef LOCTEXT_NAMESPACE
\ No newline at end of file
+#undef LOCTEXT_NAMESPACE
diff --git a/Engine/Plugins/VirtualProduction/RemoteControl/Source/RemoteControlUI/Private/UI/Controller/SRCControllerPanel.h b/Engine/Plugins/VirtualProduction/RemoteControl/Source/RemoteControlUI/Private/UI/Controller/SRCControllerPanel.h
--- a/Engine/Plugins/VirtualProduction/RemoteControl/Source/RemoteControlUI/Private/UI/Controller/SRCControllerPanel.h
+++ b/Engine/Plugins/VirtualProduction/RemoteControl/Source/RemoteControlUI/Private/UI/Controller/SRCControllerPanel.h
@@ -5,11 +5,12 @@
 #include "PropertyBag.h"
 #include "UI/BaseLogicUI/SRCLogicPanelBase.h"
 
-enum class EPropertyBagPropertyType : uint8;
 class FRCControllerModel;
-struct FRCPanelStyle;
 class SRCControllerPanel;
 class URCController;
+struct FRCPanelStyle;
+enum class ECheckBoxState : uint8;
+enum class EPropertyBagPropertyType : uint8;
 
 /*
 * ~ SRCControllerPanel ~
@@ -67,9 +68,17 @@ private:
 	/** Builds a menu containing the list of all possible Controllers*/
 	TSharedRef<SWidget> GetControllerMenuContentWidget() const;
 
+	TSharedRef<SWidget> GetMultiControllerSwitchWidget();
+
 	/** Handles click event for Add Controller button*/
 	void OnAddControllerClicked(const EPropertyBagPropertyType InValueType, UObject* InValueTypeObject = nullptr) const;
 
+	/** Handles checkbox changes event for the MultiController mode selection widget */
+	void OnToggleMultiControllersMode(ECheckBoxState CheckBoxState);
+
+	/** Whether to show the MultiController switch or not */
+	EVisibility GetMultiControllerSwitchVisibility() const;
+
 	/** Handles click event for Empty button; clears all controllers from the panel*/
 	FReply OnClickEmptyButton();
 
@@ -83,4 +92,7 @@ private:
 
 	/** Panel Style reference. */
 	const FRCPanelStyle* RCPanelStyle;
+
+	/** Is the Panel showing a list with MultiControllers handling (and hiding) controllers with duplicate Field Ids?*/
+	bool bIsMultiControllerMode = false;
 };
diff --git a/Engine/Plugins/VirtualProduction/RemoteControl/Source/RemoteControlUI/Private/UI/Controller/SRCControllerPanelList.cpp b/Engine/Plugins/VirtualProduction/RemoteControl/Source/RemoteControlUI/Private/UI/Controller/SRCControllerPanelList.cpp
--- a/Engine/Plugins/VirtualProduction/RemoteControl/Source/RemoteControlUI/Private/UI/Controller/SRCControllerPanelList.cpp
+++ b/Engine/Plugins/VirtualProduction/RemoteControl/Source/RemoteControlUI/Private/UI/Controller/SRCControllerPanelList.cpp
@@ -12,6 +12,7 @@
 #include "IPropertyRowGenerator.h"
 #include "IRemoteControlModule.h"
 #include "RCControllerModel.h"
+#include "RCMultiController.h"
 #include "RCVirtualProperty.h"
 #include "RCVirtualPropertyContainer.h"
 #include "RemoteControlField.h"
@@ -50,6 +51,8 @@ namespace UE::RCControllerPanelList
 		const FName Name = TEXT("Controller Name");
 		const FName Value = TEXT("Controller Value");
 		const FName DragHandle = TEXT("Drag Handle");
+		const FName FieldId = TEXT("Controller Field Id");
+		const FName ValueTypeSelection = TEXT("Value Type Selection");
 	}
 
 	class SControllerItemListRow : public SMultiColumnTableRow<TSharedRef<FRCControllerModel>>
@@ -58,7 +61,7 @@ namespace UE::RCControllerPanelList
 		void Construct(const FTableRowArgs& InArgs, const TSharedRef<STableViewBase>& OwnerTableView, TSharedRef<FRCControllerModel> InControllerItem, TSharedRef<SRCControllerPanelList> InControllerPanelList)
 		{
 			ControllerItem = InControllerItem;
-			ControllerPanelList = InControllerPanelList;
+			ControllerPanelListWeakPtr = InControllerPanelList.ToWeakPtr();
 			FSuperRowType::Construct(InArgs, OwnerTableView);
 		}
 
@@ -67,14 +70,24 @@ namespace UE::RCControllerPanelList
 			if (!ensure(ControllerItem.IsValid()))
 				return SNullWidget::NullWidget;
 
-
 			if (ColumnName == UE::RCControllerPanelList::Columns::TypeColor)
 			{
 				if (URCController* Controller = Cast< URCController>(ControllerItem->GetVirtualProperty()))
 				{
-					return UE::RCUIHelpers::GetTypeColorWidget(Controller->GetProperty());
+					if (const FProperty* Property = Controller->GetProperty())
+					{
+						return UE::RCUIHelpers::GetTypeColorWidget(Property);
+					}
 				}
 			}
+			else if (ColumnName == UE::RCControllerPanelList::Columns::FieldId)
+			{
+				return WrapWithDropTarget(ControllerItem->GetFieldIdWidget());
+			}
+			else if (ColumnName == UE::RCControllerPanelList::Columns::ValueTypeSelection)
+			{
+				return WrapWithDropTarget(ControllerItem->GetTypeSelectionWidget());
+			}
 			else if (ColumnName == UE::RCControllerPanelList::Columns::Name)
 			{
 				return WrapWithDropTarget(ControllerItem->GetNameWidget());
@@ -136,6 +149,9 @@ namespace UE::RCControllerPanelList
 
 		FReply OnControllerItemDragDrop(TSharedPtr<FDragDropOperation> DragDropOperation)
 		{
+			check(ControllerPanelListWeakPtr.IsValid());
+			const TSharedPtr<SRCControllerPanelList> ControllerPanelList = ControllerPanelListWeakPtr.Pin();
+			
 			bIsDragActive = false;
 			ControllerPanelList->bIsAnyControllerItemEligibleForDragDrop = false;
 
@@ -186,10 +202,12 @@ namespace UE::RCControllerPanelList
 
 		bool OnAllowDrop(TSharedPtr<FDragDropOperation> DragDropOperation, const EDragDropSupportedModes SupportedMode)
 		{
-			if (!ensure(ControllerPanelList))
+			if (!ensure(ControllerPanelListWeakPtr.IsValid()))
 			{
 				return false;
 			}
+			
+			const TSharedPtr<SRCControllerPanelList> ControllerPanelList = ControllerPanelListWeakPtr.Pin();
 
 			if (DragDropOperation && ControllerItem)
 			{
@@ -253,7 +271,7 @@ namespace UE::RCControllerPanelList
 
 private:
 		TSharedPtr<FRCControllerModel> ControllerItem;
-		TSharedPtr<SRCControllerPanelList> ControllerPanelList;
+		TWeakPtr<SRCControllerPanelList> ControllerPanelListWeakPtr;
 		bool bIsDragActive = false;
 	};
 } 
@@ -273,7 +291,7 @@ void SRCControllerPanelList::Construct(const FArguments& InArgs, const TSharedRe
 		.SelectionMode(ESelectionMode::Single) // Current setup supports only single selection (and related display) of a Controller in the list
 		.OnContextMenuOpening(this, &SRCLogicPanelListBase::GetContextMenuWidget)
 		.HeaderRow(
-			SNew(SHeaderRow)
+			SAssignNew(ControllersHeaderRow, SHeaderRow)
 			.Style(&RCPanelStyle->HeaderRowStyle)
 
 			+ SHeaderRow::Column(UE::RCControllerPanelList::Columns::TypeColor)
@@ -293,7 +311,7 @@ void SRCControllerPanelList::Construct(const FArguments& InArgs, const TSharedRe
 
 			+ SHeaderRow::Column(UE::RCControllerPanelList::Columns::Value)
 			.DefaultLabel(LOCTEXT("ControllerValueColumnName", "Input"))
-			.FillWidth(0.75f)
+			.FillWidth(0.45f)
 			.HeaderContentPadding(RCPanelStyle->HeaderRowPadding)
 		);
 
@@ -347,6 +365,14 @@ int32 SRCControllerPanelList::NumSelectedLogicItems() const
 
 void SRCControllerPanelList::Reset()
 {
+	for (const TSharedPtr<FRCControllerModel>& ControllerModel : ControllerItems)
+	{
+		if (ControllerModel)
+		{
+			ControllerModel->OnValueTypeChanged.RemoveAll(this);
+		}
+	}
+	
 	ControllerItems.Empty();
 
 	check(ControllerPanelWeakPtr.IsValid());
@@ -363,6 +389,10 @@ void SRCControllerPanelList::Reset()
 	// Generator should be moved to separate class
 	TArray<TSharedRef<IDetailTreeNode>> RootTreeNodes = PropertyRowGenerator->GetRootTreeNodes();
 
+	MultiControllers.ResetMultiControllers();
+
+	bool bShowFieldIdsColumn = false;
+	
 	for (const TSharedRef<IDetailTreeNode>& CategoryNode : RootTreeNodes)
 	{
 		TArray<TSharedRef<IDetailTreeNode>> Children;
@@ -373,21 +403,75 @@ void SRCControllerPanelList::Reset()
 		for (TSharedRef<IDetailTreeNode>& Child : Children)
 		{
 			FProperty* Property = Child->CreatePropertyHandle()->GetProperty();
- 			check(Property);
+			check(Property);
 
-			if(Property->IsA<FStrProperty>() || Property->IsA<FTextProperty>())
+			if (Property->IsA<FStrProperty>() || Property->IsA<FTextProperty>())
 			{
 				Property->AppendMetaData({{TEXT("multiline"), TEXT("true")}});
 			}
 
 			if (URCVirtualPropertyBase* Controller = Preset->GetController(Property->GetFName()))
 			{
-				if(ensureAlways(ControllerItems.IsValidIndex(Controller->DisplayIndex)))
-					ControllerItems[Controller->DisplayIndex] = MakeShared<FRCControllerModel>(Controller, Child, RemoteControlPanel);
+				bool bIsVisible = true;
+				bool bIsMultiController = false;
+
+				const FName& FieldId = Controller->FieldId;
+
+				if (FieldId != NAME_None)
+				{
+					// there's at least one Field Id set, let's show their column
+					bShowFieldIdsColumn = true;
+				}
+				
+				// MultiController Mode: only showing one Controller per Field Id
+				if (bIsInMultiControllerMode && Preset->GetControllersByFieldId(FieldId).Num() > 1)
+				{					
+					bIsMultiController = MultiControllers.TryToAddAsMultiController(Controller);
+					bIsVisible = bIsMultiController;
+				}
+
+				if (bIsVisible)
+				{
+					if (ensureAlways(ControllerItems.IsValidIndex(Controller->DisplayIndex)))
+					{
+						const TSharedRef<FRCControllerModel> ControllerModel = MakeShared<FRCControllerModel>(Controller, Child, RemoteControlPanel);
+						ControllerItems[Controller->DisplayIndex] = ControllerModel;
+
+						if (bIsMultiController)
+						{
+							ControllerModel->SetMultiController(bIsMultiController);
+							ControllerModel->OnValueTypeChanged.AddSP(this, &SRCControllerPanelList::OnControllerValueTypeChanged);
+							ControllerModel->OnValueChanged.AddSP(this, &SRCControllerPanelList::OnControllerValueChanged);
+						}
+					}
+				}
 			}
 		}
 	}
 
+	// sort by Field Id
+	if (bIsInMultiControllerMode)
+	{
+		Algo::Sort(ControllerItems, [](const TSharedPtr<FRCControllerModel>& A, const TSharedPtr<FRCControllerModel>& B)
+		{
+			if (A.IsValid() && B.IsValid())
+			{
+				const URCVirtualPropertyBase* ControllerA = A->GetVirtualProperty();
+				const URCVirtualPropertyBase* ControllerB = B->GetVirtualProperty();
+
+				if (ControllerA && ControllerB)
+				{
+					return ControllerA->FieldId.FastLess(ControllerB->FieldId);
+				}
+			}
+
+			return false;
+		});
+	}
+
+	ShowFieldIdHeaderColumn(bShowFieldIdsColumn);
+	ShowValueTypeHeaderColumn(bIsInMultiControllerMode);
+	
 	ListView->RebuildList();
 }
 
@@ -453,6 +537,30 @@ void SRCControllerPanelList::OnFinishedChangingProperties(const FPropertyChanged
 	}
 }
 
+void SRCControllerPanelList::OnControllerValueTypeChanged(URCVirtualPropertyBase* InController, EPropertyBagPropertyType InValueType)
+{	
+	if (InController)
+	{		
+		MultiControllers.UpdateFieldIdValueType(InController->FieldId, InValueType);
+		Reset();
+
+		// todo: do we also want to refresh controllers values after type change?
+	}
+}
+
+void SRCControllerPanelList::OnControllerValueChanged(URCVirtualPropertyBase* InController)
+{	
+	const FName& FieldId = InController->FieldId;
+
+	FRCMultiController MultiController = MultiControllers.GetMultiController(FieldId);
+
+	if (MultiController.IsValid())
+	{
+		MultiController.UpdateHandledControllersValue();
+	}	
+}
+
+
 void SRCControllerPanelList::OnEmptyControllers()
 {
 	if (TSharedPtr< SRCControllerPanel> ControllerPanel = ControllerPanelWeakPtr.Pin())
@@ -780,4 +888,60 @@ bool SRCControllerPanelList::IsListViewHovered()
 	return ListView->IsDirectlyHovered();
 }
 
-#undef LOCTEXT_NAMESPACE
\ No newline at end of file
+void SRCControllerPanelList::ShowValueTypeHeaderColumn(bool bInShowColumn)
+{
+	if (ControllersHeaderRow.IsValid())
+	{
+		const bool bColumnIsGenerated = ControllersHeaderRow->IsColumnGenerated(UE::RCControllerPanelList::Columns::ValueTypeSelection);
+		if (bInShowColumn)
+		{
+			if (!bColumnIsGenerated)
+			{
+				ControllersHeaderRow->AddColumn(
+					SHeaderRow::FColumn::FArguments()
+					.ColumnId(UE::RCControllerPanelList::Columns::ValueTypeSelection)
+					.DefaultLabel(LOCTEXT("ControllerValueTypeColumnName", "Value Type"))
+					.FillWidth(0.2f)
+					.HeaderContentPadding(RCPanelStyle->HeaderRowPadding)
+				);
+			}
+		}
+		else if (bColumnIsGenerated)
+		{
+			ControllersHeaderRow->RemoveColumn(UE::RCControllerPanelList::Columns::ValueTypeSelection);
+		}
+	}
+}
+
+void SRCControllerPanelList::ShowFieldIdHeaderColumn(bool bInShowColumn)
+{	
+	if (ControllersHeaderRow.IsValid())
+	{
+		const bool bColumnIsGenerated = ControllersHeaderRow->IsColumnGenerated(UE::RCControllerPanelList::Columns::FieldId);
+		if (bInShowColumn)
+		{
+			if (!bColumnIsGenerated)
+			{
+				ControllersHeaderRow->InsertColumn(
+					SHeaderRow::FColumn::FArguments()
+					.ColumnId(UE::RCControllerPanelList::Columns::FieldId)
+					.DefaultLabel(LOCTEXT("ControllerNameColumnFieldId", "Field Id"))
+					.FillWidth(0.2f)
+					.HeaderContentPadding(RCPanelStyle->HeaderRowPadding),
+					2);
+			}
+		}
+		else if (bColumnIsGenerated)
+		{
+			ControllersHeaderRow->RemoveColumn(UE::RCControllerPanelList::Columns::FieldId);
+		}
+	}
+}
+
+void SRCControllerPanelList::SetMultiControllerMode(bool bIsUniqueModeOn)
+{
+	bIsInMultiControllerMode = bIsUniqueModeOn;
+	Reset();
+}
+
+#undef LOCTEXT_NAMESPACE
diff --git a/Engine/Plugins/VirtualProduction/RemoteControl/Source/RemoteControlUI/Private/UI/Controller/SRCControllerPanelList.h b/Engine/Plugins/VirtualProduction/RemoteControl/Source/RemoteControlUI/Private/UI/Controller/SRCControllerPanelList.h
--- a/Engine/Plugins/VirtualProduction/RemoteControl/Source/RemoteControlUI/Private/UI/Controller/SRCControllerPanelList.h
+++ b/Engine/Plugins/VirtualProduction/RemoteControl/Source/RemoteControlUI/Private/UI/Controller/SRCControllerPanelList.h
@@ -3,6 +3,7 @@
 #pragma once
 #include "DragAndDrop/DecoratedDragDropOp.h"
 #include "Misc/NotifyHook.h"
+#include "RCMultiController.h"
 #include "UI/BaseLogicUI/SRCLogicPanelListBase.h"
 #include "UI/RemoteControlPanelStyle.h"
 #include "Widgets/Layout/SBorder.h"
@@ -138,6 +139,9 @@ public:
 	/** Whether the user's cursor is directly hovered over the List View*/
 	bool IsListViewHovered();
 
+	/** Enable/disable MultiController Mode */
+	void SetMultiControllerMode(bool bIsUniqueModeOn);
+
 	/** Flag that facilitates usage of two mutually exclusive drag-drop zones within a single panel 
 	* 
 	* The first drag-drop zone is empty panel space for "Bind To New Controller"
@@ -175,6 +179,15 @@ private:
 	*/
 	void OnFinishedChangingProperties(const FPropertyChangedEvent& PropertyChangedEvent);
 
+	/** Called when a MultiController Value Type changes. */
+	void OnControllerValueTypeChanged(URCVirtualPropertyBase* InController, EPropertyBagPropertyType InValueType);
+
+	/**
+	 * Called when a Controller Value changes.
+	 * Currently used to update handled controllers, in case of MultiControllers.
+	 */
+	void OnControllerValueChanged(URCVirtualPropertyBase* InController);
+
 	/** Creates a new Controller for the given Remote Control Property and also binds to it */
 	void CreateAutoBindForProperty(TSharedPtr<const FRemoteControlProperty> RemoteControlProperty);
 
@@ -191,7 +204,7 @@ private:
 	TArray<TSharedPtr<FRCControllerModel>> ControllerItems;
 
 	/** List View widget for representing our Controllers List*/
-	TSharedPtr<SListView<TSharedPtr<FRCControllerModel>>>  ListView;
+	TSharedPtr<SListView<TSharedPtr<FRCControllerModel>>> ListView;
 	
 	/** Refreshes the list from the latest state of the model*/
 	virtual void Reset() override;
@@ -202,7 +215,21 @@ private:
 	/** Removes the given Controller UI model item from the list of UI models*/
 	virtual int32 RemoveModel(const TSharedPtr<FRCLogicModeBase> InModel) override;
 
+	/** Set the visibility of the List header the Value Type Column used for MultiControllers */
+	void ShowValueTypeHeaderColumn(bool bInShowColumn);
+
+	/** Set the visibility of the List header the Field Id Column */
+	void ShowFieldIdHeaderColumn(bool bInShowColumn);
+
 	/** Panel Style reference. */
 	const FRCPanelStyle* RCPanelStyle;
+
+	/** Keeps track of current MultiControllers */
+	FRCMultiControllersState MultiControllers;
+	
+	bool bIsInMultiControllerMode = false;
+
+	/** Storing the header so we can add/remove columns */
+	TSharedPtr<SHeaderRow> ControllersHeaderRow;
 };
 
diff --git a/Engine/Plugins/VirtualProduction/RemoteControl/Source/RemoteControlUI/Private/UI/Controller/Tests/RCTypeTranslator.spec.cpp b/Engine/Plugins/VirtualProduction/RemoteControl/Source/RemoteControlUI/Private/UI/Controller/Tests/RCTypeTranslator.spec.cpp
new file mode 100644
--- /dev/null
+++ b/Engine/Plugins/VirtualProduction/RemoteControl/Source/RemoteControlUI/Private/UI/Controller/Tests/RCTypeTranslator.spec.cpp
@@ -0,0 +1,78 @@
+﻿// Copyright Epic Games, Inc. All Rights Reserved.
+
+#include "Misc/AutomationTest.h"
+#include "RemoteControlPreset.h"
+#include "UI/Controller/TypeTranslator/RCTypeTranslator.h"
+#include "UObject/StrongObjectPtr.h"
+
+BEGIN_DEFINE_SPEC(
+	FRCTypeTranslatorSpec,
+	"Plugin.RemoteControlUI.RCTypeTranslator",
+	EAutomationTestFlags::ApplicationContextMask | EAutomationTestFlags::ProductFilter)
+	TStrongObjectPtr<URemoteControlPreset> Preset;
+	TStrongObjectPtr<URCVirtualPropertyBase> BoolVirtualProperty;
+	TStrongObjectPtr<URCVirtualPropertyBase> StringVirtualProperty;
+END_DEFINE_SPEC(FRCTypeTranslatorSpec)
+
+void FRCTypeTranslatorSpec::Define()
+{
+	//todo: Add tests for other conversions
+	
+	BeforeEach([this]
+	{
+		// Init preset
+		Preset.Reset(NewObject<URemoteControlPreset>());
+
+		// Create controllers (these could actually be URCVirtualPropertyBase)
+		BoolVirtualProperty.Reset(Preset->AddController(URCController::StaticClass(), EPropertyBagPropertyType::Bool, nullptr, TEXT("BoolController")));
+		StringVirtualProperty.Reset(Preset->AddController(URCController::StaticClass(), EPropertyBagPropertyType::String, nullptr, TEXT("StringController")));
+	});
+
+	Describe(TEXT("From String"), [this]()
+	{
+		Describe(TEXT("To Bool"), [this]()
+		{
+			Describe(TEXT("String is Empty"), [this]()
+			{
+				It("Bool should be set to false", [this]()
+				{
+					const FString& SourceValue = TEXT("");
+					constexpr bool ExpectedResult = false;
+					bool Result;
+					StringVirtualProperty->SetValueString(SourceValue);
+					FRCTypeTranslator::Get()->Translate(StringVirtualProperty.Get(), {BoolVirtualProperty.Get()});					
+					BoolVirtualProperty->GetValueBool(Result);
+					TestEqual(TEXT("Values should match"), Result, ExpectedResult);
+				});
+			});
+
+			Describe(TEXT("String is 0"), [this]()
+			{
+				It("Bool should be set to false", [this]()
+				{
+					const FString& SourceValue = TEXT("0");
+					constexpr bool ExpectedResult = false;
+					bool Result;
+					StringVirtualProperty->SetValueString(SourceValue);
+					FRCTypeTranslator::Get()->Translate(StringVirtualProperty.Get(), {BoolVirtualProperty.Get()});
+					BoolVirtualProperty->GetValueBool(Result);
+					TestEqual(TEXT("Values should match"), Result, ExpectedResult);
+				});
+			});
+
+			Describe(TEXT("String is 123"), [this]()
+			{
+				It("Bool should be set to true", [this]()
+				{
+					const FString& SourceValue = TEXT("123");
+					constexpr bool ExpectedResult = true;
+					bool Result;
+					StringVirtualProperty->SetValueString(SourceValue);
+					FRCTypeTranslator::Get()->Translate(StringVirtualProperty.Get(), {BoolVirtualProperty.Get()});
+					BoolVirtualProperty->GetValueBool(Result);
+					TestEqual(TEXT("Values should match"), Result, ExpectedResult);
+				});
+			});
+		});
+	});
+}
diff --git a/Engine/Plugins/VirtualProduction/RemoteControl/Source/RemoteControlUI/Private/UI/Controller/TypeTranslator/RCBoolHandler.cpp b/Engine/Plugins/VirtualProduction/RemoteControl/Source/RemoteControlUI/Private/UI/Controller/TypeTranslator/RCBoolHandler.cpp
new file mode 100644
--- /dev/null
+++ b/Engine/Plugins/VirtualProduction/RemoteControl/Source/RemoteControlUI/Private/UI/Controller/TypeTranslator/RCBoolHandler.cpp
@@ -0,0 +1,70 @@
+﻿// Copyright Epic Games, Inc. All Rights Reserved.
+
+#include "RCBoolHandler.h"
+#include "RCVirtualProperty.h"
+
+void FRCBoolHandler::Apply(const bool& InValue, URCVirtualPropertyBase* InVirtualProperty)
+{
+	UpdateValue(InVirtualProperty, InValue);
+}
+
+void FRCBoolHandler::Apply(const uint8& InValue, URCVirtualPropertyBase* InVirtualProperty)
+{
+	UpdateValue(InVirtualProperty, static_cast<bool>(InValue));
+}
+
+void FRCBoolHandler::Apply(const double& InValue, URCVirtualPropertyBase* InVirtualProperty)
+{
+	UpdateValue(InVirtualProperty, static_cast<bool>(InValue));
+}
+
+void FRCBoolHandler::Apply(const float& InValue, URCVirtualPropertyBase* InVirtualProperty)
+{
+	UpdateValue(InVirtualProperty, static_cast<bool>(InValue));
+}
+
+void FRCBoolHandler::Apply(const int32& InValue, URCVirtualPropertyBase* InVirtualProperty)
+{
+	UpdateValue(InVirtualProperty, static_cast<bool>(InValue));
+}
+
+void FRCBoolHandler::Apply(const int64& InValue, URCVirtualPropertyBase* InVirtualProperty)
+{
+	UpdateValue(InVirtualProperty, static_cast<bool>(InValue));
+}
+
+void FRCBoolHandler::Apply(const FName& InValue, URCVirtualPropertyBase* InVirtualProperty)
+{
+	Apply(InValue.ToString(), InVirtualProperty);
+}
+
+void FRCBoolHandler::Apply(const FString& InValue, URCVirtualPropertyBase* InVirtualProperty)
+{
+	if (InVirtualProperty)
+	{
+		if (InValue.IsEmpty() || InValue.Equals(TEXT("0")))
+		{
+			InVirtualProperty->SetValueBool(false);
+		}
+		else
+		{
+			InVirtualProperty->SetValueBool(true);
+		}
+
+		OnControllerPropertyModified(InVirtualProperty);
+	}
+}
+
+void FRCBoolHandler::Apply(const FText& InValue, URCVirtualPropertyBase* InVirtualProperty)
+{
+	Apply(InValue.ToString(), InVirtualProperty);
+}
+
+void FRCBoolHandler::UpdateValue(URCVirtualPropertyBase* InVirtualProperty, const bool InValue)
+{
+	if (InVirtualProperty)
+	{
+		InVirtualProperty->SetValueBool(InValue);
+		OnControllerPropertyModified(InVirtualProperty);
+	}
+}
diff --git a/Engine/Plugins/VirtualProduction/RemoteControl/Source/RemoteControlUI/Private/UI/Controller/TypeTranslator/RCBoolHandler.h b/Engine/Plugins/VirtualProduction/RemoteControl/Source/RemoteControlUI/Private/UI/Controller/TypeTranslator/RCBoolHandler.h
new file mode 100644
--- /dev/null
+++ b/Engine/Plugins/VirtualProduction/RemoteControl/Source/RemoteControlUI/Private/UI/Controller/TypeTranslator/RCBoolHandler.h
@@ -0,0 +1,23 @@
+﻿// Copyright Epic Games, Inc. All Rights Reserved.
+
+#pragma once
+
+#include "RCTypeHandler.h"
+
+class FRCBoolHandler : public FRCTypeHandler
+{
+public:
+	virtual void Apply(const bool& InValue, URCVirtualPropertyBase* InController) override;
+	virtual void Apply(const uint8& InValue, URCVirtualPropertyBase* InController) override;
+	virtual void Apply(const double& InValue, URCVirtualPropertyBase* InController) override;
+	virtual void Apply(const float& InValue, URCVirtualPropertyBase* InController) override;
+	virtual void Apply(const int32& InValue, URCVirtualPropertyBase* InController) override;
+	virtual void Apply(const int64& InValue, URCVirtualPropertyBase* InController) override;
+	virtual void Apply(const FName& InValue, URCVirtualPropertyBase* InController) override;
+	virtual void Apply(const FString& InValue, URCVirtualPropertyBase* InController) override;
+	virtual void Apply(const FText& InValue, URCVirtualPropertyBase* InController) override;
+
+protected:
+	static void UpdateValue(URCVirtualPropertyBase* InVirtualProperty, const bool InValue);
+	~FRCBoolHandler() = default;
+};
diff --git a/Engine/Plugins/VirtualProduction/RemoteControl/Source/RemoteControlUI/Private/UI/Controller/TypeTranslator/RCFloatHandler.cpp b/Engine/Plugins/VirtualProduction/RemoteControl/Source/RemoteControlUI/Private/UI/Controller/TypeTranslator/RCFloatHandler.cpp
new file mode 100644
--- /dev/null
+++ b/Engine/Plugins/VirtualProduction/RemoteControl/Source/RemoteControlUI/Private/UI/Controller/TypeTranslator/RCFloatHandler.cpp
@@ -0,0 +1,59 @@
+﻿// Copyright Epic Games, Inc. All Rights Reserved.
+
+#include "RCFloatHandler.h"
+#include "RCVirtualProperty.h"
+
+void FRCFloatHandler::Apply(const bool& InValue, URCVirtualPropertyBase* InVirtualProperty)
+{
+	UpdateValue(InVirtualProperty, InValue);
+}
+
+void FRCFloatHandler::Apply(const uint8& InValue, URCVirtualPropertyBase* InVirtualProperty)
+{
+	UpdateValue(InVirtualProperty, InValue);
+}
+
+void FRCFloatHandler::Apply(const double& InValue, URCVirtualPropertyBase* InVirtualProperty)
+{
+	UpdateValue(InVirtualProperty, InValue);
+}
+
+void FRCFloatHandler::Apply(const float& InValue, URCVirtualPropertyBase* InVirtualProperty)
+{
+	UpdateValue(InVirtualProperty, InValue);
+}
+
+void FRCFloatHandler::Apply(const int32& InValue, URCVirtualPropertyBase* InVirtualProperty)
+{
+	UpdateValue(InVirtualProperty, InValue);
+}
+
+void FRCFloatHandler::Apply(const int64& InValue, URCVirtualPropertyBase* InVirtualProperty)
+{
+	UpdateValue(InVirtualProperty, InValue);
+}
+
+void FRCFloatHandler::Apply(const FName& InValue, URCVirtualPropertyBase* InVirtualProperty)
+{
+	Apply(InValue.ToString(), InVirtualProperty);
+}
+
+void FRCFloatHandler::Apply(const FString& InValue, URCVirtualPropertyBase* InVirtualProperty)
+{
+	const float OutValue = FCString::Atof(*InValue);
+	UpdateValue(InVirtualProperty, OutValue);
+}
+
+void FRCFloatHandler::Apply(const FText& InValue, URCVirtualPropertyBase* InVirtualProperty)
+{
+	Apply(InValue.ToString(), InVirtualProperty);
+}
+
+void FRCFloatHandler::UpdateValue(URCVirtualPropertyBase* InVirtualProperty, const float InValue)
+{
+	if (InVirtualProperty)
+	{
+		InVirtualProperty->SetValueFloat(InValue);
+		OnControllerPropertyModified(InVirtualProperty);
+	}
+}
diff --git a/Engine/Plugins/VirtualProduction/RemoteControl/Source/RemoteControlUI/Private/UI/Controller/TypeTranslator/RCFloatHandler.h b/Engine/Plugins/VirtualProduction/RemoteControl/Source/RemoteControlUI/Private/UI/Controller/TypeTranslator/RCFloatHandler.h
new file mode 100644
--- /dev/null
+++ b/Engine/Plugins/VirtualProduction/RemoteControl/Source/RemoteControlUI/Private/UI/Controller/TypeTranslator/RCFloatHandler.h
@@ -0,0 +1,23 @@
+﻿// Copyright Epic Games, Inc. All Rights Reserved.
+
+#pragma once
+
+#include "RCTypeHandler.h"
+
+class FRCFloatHandler : public FRCTypeHandler
+{
+public:
+	virtual void Apply(const bool& InValue, URCVirtualPropertyBase* InVirtualProperty) override;
+	virtual void Apply(const uint8& InValue, URCVirtualPropertyBase* InVirtualProperty) override;
+	virtual void Apply(const double& InValue, URCVirtualPropertyBase* InVirtualProperty) override;
+	virtual void Apply(const float& InValue, URCVirtualPropertyBase* InVirtualProperty) override;
+	virtual void Apply(const int32& InValue, URCVirtualPropertyBase* InVirtualProperty) override;
+	virtual void Apply(const int64& InValue, URCVirtualPropertyBase* InVirtualProperty) override;
+	virtual void Apply(const FName& InValue, URCVirtualPropertyBase* InVirtualProperty) override;
+	virtual void Apply(const FString& InValue, URCVirtualPropertyBase* InVirtualProperty) override;
+	virtual void Apply(const FText& InValue, URCVirtualPropertyBase* InVirtualProperty) override;
+
+protected:
+	static void UpdateValue(URCVirtualPropertyBase* InVirtualProperty, const float InValue);
+	~FRCFloatHandler() = default;
+};
diff --git a/Engine/Plugins/VirtualProduction/RemoteControl/Source/RemoteControlUI/Private/UI/Controller/TypeTranslator/RCStringHandler.cpp b/Engine/Plugins/VirtualProduction/RemoteControl/Source/RemoteControlUI/Private/UI/Controller/TypeTranslator/RCStringHandler.cpp
new file mode 100644
--- /dev/null
+++ b/Engine/Plugins/VirtualProduction/RemoteControl/Source/RemoteControlUI/Private/UI/Controller/TypeTranslator/RCStringHandler.cpp
@@ -0,0 +1,64 @@
+﻿// Copyright Epic Games, Inc. All Rights Reserved.
+
+#include "RCStringHandler.h"
+#include "RCVirtualProperty.h"
+
+void FRCStringHandler::Apply(const bool& InValue, URCVirtualPropertyBase* InVirtualProperty)
+{	
+	const FString& OutValue = InValue ? TEXT("true") : TEXT("false");
+	UpdateValue(InVirtualProperty, OutValue);
+}
+
+void FRCStringHandler::Apply(const uint8& InValue, URCVirtualPropertyBase* InVirtualProperty)
+{
+	const FString& OutValue = FString::SanitizeFloat(InValue, 0);
+	UpdateValue(InVirtualProperty, OutValue);
+}
+
+void FRCStringHandler::Apply(const double& InValue, URCVirtualPropertyBase* InVirtualProperty)
+{
+	const FString& OutValue = FString::SanitizeFloat(InValue);
+	UpdateValue(InVirtualProperty, OutValue);
+}
+
+void FRCStringHandler::Apply(const float& InValue, URCVirtualPropertyBase* InVirtualProperty)
+{
+	const FString& OutValue = FString::SanitizeFloat(InValue);
+	UpdateValue(InVirtualProperty, OutValue);
+}
+
+void FRCStringHandler::Apply(const int32& InValue, URCVirtualPropertyBase* InVirtualProperty)
+{
+	const FString& OutValue = FString::SanitizeFloat(InValue, 0);
+	UpdateValue(InVirtualProperty, OutValue);
+}
+
+void FRCStringHandler::Apply(const int64& InValue, URCVirtualPropertyBase* InVirtualProperty)
+{
+	const FString& OutValue = FString::SanitizeFloat(InValue, 0);
+	UpdateValue(InVirtualProperty, OutValue);
+}
+
+void FRCStringHandler::Apply(const FName& InValue, URCVirtualPropertyBase* InVirtualProperty)
+{
+	Apply(InValue.ToString(), InVirtualProperty);
+}
+
+void FRCStringHandler::Apply(const FString& InValue, URCVirtualPropertyBase* InVirtualProperty)
+{
+	UpdateValue(InVirtualProperty, InValue);
+}
+
+void FRCStringHandler::Apply(const FText& InValue, URCVirtualPropertyBase* InVirtualProperty)
+{
+	Apply(InValue.ToString(), InVirtualProperty);
+}
+
+void FRCStringHandler::UpdateValue(URCVirtualPropertyBase* InVirtualProperty, const FString& InValue)
+{
+	if (InVirtualProperty)
+	{
+		InVirtualProperty->SetValueString(InValue);
+		OnControllerPropertyModified(InVirtualProperty);
+	}
+}
diff --git a/Engine/Plugins/VirtualProduction/RemoteControl/Source/RemoteControlUI/Private/UI/Controller/TypeTranslator/RCStringHandler.h b/Engine/Plugins/VirtualProduction/RemoteControl/Source/RemoteControlUI/Private/UI/Controller/TypeTranslator/RCStringHandler.h
new file mode 100644
--- /dev/null
+++ b/Engine/Plugins/VirtualProduction/RemoteControl/Source/RemoteControlUI/Private/UI/Controller/TypeTranslator/RCStringHandler.h
@@ -0,0 +1,23 @@
+﻿// Copyright Epic Games, Inc. All Rights Reserved.
+
+#pragma once
+
+#include "RCTypeHandler.h"
+
+class FRCStringHandler : public FRCTypeHandler
+{
+public:
+	virtual void Apply(const bool& InValue, URCVirtualPropertyBase* InController) override;
+	virtual void Apply(const uint8& InValue, URCVirtualPropertyBase* InController) override;
+	virtual void Apply(const double& InValue, URCVirtualPropertyBase* InController) override;
+	virtual void Apply(const float& InValue, URCVirtualPropertyBase* InController) override;
+	virtual void Apply(const int32& InValue, URCVirtualPropertyBase* InController) override;
+	virtual void Apply(const int64& InValue, URCVirtualPropertyBase* InController) override;
+	virtual void Apply(const FName& InValue, URCVirtualPropertyBase* InController) override;
+	virtual void Apply(const FString& InValue, URCVirtualPropertyBase* InController) override;
+	virtual void Apply(const FText& InValue, URCVirtualPropertyBase* InController) override;
+
+protected:
+	static void UpdateValue(URCVirtualPropertyBase* InVirtualProperty, const FString& InValue);
+	~FRCStringHandler() = default;
+};
diff --git a/Engine/Plugins/VirtualProduction/RemoteControl/Source/RemoteControlUI/Private/UI/Controller/TypeTranslator/RCTypeHandler.cpp b/Engine/Plugins/VirtualProduction/RemoteControl/Source/RemoteControlUI/Private/UI/Controller/TypeTranslator/RCTypeHandler.cpp
new file mode 100644
--- /dev/null
+++ b/Engine/Plugins/VirtualProduction/RemoteControl/Source/RemoteControlUI/Private/UI/Controller/TypeTranslator/RCTypeHandler.cpp
@@ -0,0 +1,12 @@
+﻿// Copyright Epic Games, Inc. All Rights Reserved.
+
+#include "RCTypeHandler.h"
+#include "Controller/RCController.h"
+
+void FRCTypeHandler::OnControllerPropertyModified(URCVirtualPropertyBase* InVirtualProperty)
+{
+	if (URCController* Controller = Cast<URCController>(InVirtualProperty))
+	{
+		Controller->OnModifyPropertyValue();
+	}
+}
diff --git a/Engine/Plugins/VirtualProduction/RemoteControl/Source/RemoteControlUI/Private/UI/Controller/TypeTranslator/RCTypeHandler.h b/Engine/Plugins/VirtualProduction/RemoteControl/Source/RemoteControlUI/Private/UI/Controller/TypeTranslator/RCTypeHandler.h
new file mode 100644
--- /dev/null
+++ b/Engine/Plugins/VirtualProduction/RemoteControl/Source/RemoteControlUI/Private/UI/Controller/TypeTranslator/RCTypeHandler.h
@@ -0,0 +1,24 @@
+﻿// Copyright Epic Games, Inc. All Rights Reserved.
+
+#pragma once
+
+#include "Controller/RCController.h"
+
+class FRCTypeHandler
+{
+public:
+	static void OnControllerPropertyModified(URCVirtualPropertyBase* InVirtualProperty);
+	
+	virtual void Apply(const bool& InValue, URCVirtualPropertyBase* InVirtualProperty) = 0;
+	virtual void Apply(const uint8& InValue, URCVirtualPropertyBase* InVirtualProperty) = 0;
+	virtual void Apply(const double& InValue, URCVirtualPropertyBase* InVirtualProperty) = 0;
+	virtual void Apply(const float& InValue, URCVirtualPropertyBase* InVirtualProperty) = 0;
+	virtual void Apply(const int32& InValue, URCVirtualPropertyBase* InVirtualProperty) = 0;
+	virtual void Apply(const int64& InValue, URCVirtualPropertyBase* InVirtualProperty) = 0;
+	virtual void Apply(const FName& InValue, URCVirtualPropertyBase* InVirtualProperty) = 0;
+	virtual void Apply(const FString& InValue, URCVirtualPropertyBase* InVirtualProperty) = 0;
+	virtual void Apply(const FText& InValue, URCVirtualPropertyBase* InVirtualProperty) = 0;
+
+protected:
+	~FRCTypeHandler() = default;
+};
diff --git a/Engine/Plugins/VirtualProduction/RemoteControl/Source/RemoteControlUI/Private/UI/Controller/TypeTranslator/RCTypeTranslator.cpp b/Engine/Plugins/VirtualProduction/RemoteControl/Source/RemoteControlUI/Private/UI/Controller/TypeTranslator/RCTypeTranslator.cpp
new file mode 100644
--- /dev/null
+++ b/Engine/Plugins/VirtualProduction/RemoteControl/Source/RemoteControlUI/Private/UI/Controller/TypeTranslator/RCTypeTranslator.cpp
@@ -0,0 +1,161 @@
+﻿// Copyright Epic Games, Inc. All Rights Reserved.
+
+#include "RCTypeTranslator.h"
+#include "Misc/EnumRange.h"
+#include "PropertyBag.h"
+#include "RCBoolHandler.h"
+#include "RCFloatHandler.h"
+#include "RCStringHandler.h"
+
+// Needed to nicely cycle through Value Types
+ENUM_RANGE_BY_FIRST_AND_LAST(EPropertyBagPropertyType, EPropertyBagPropertyType::Bool, EPropertyBagPropertyType::SoftObject);
+
+FRCTypeTranslator::FRCTypeTranslator()
+{
+	for (const EPropertyBagPropertyType RCValueType : TEnumRange<EPropertyBagPropertyType>())
+	{
+		FRCTypeHandler* TypeHandler = CreateTypeHandler(RCValueType);
+		TypeHandlers.Emplace(RCValueType, TypeHandler);
+	}
+}
+
+FRCTypeTranslator* FRCTypeTranslator::Get()
+{
+	if (!TypeTranslatorSingleton)
+	{
+		TypeTranslatorSingleton = new FRCTypeTranslator();
+	}
+
+	return TypeTranslatorSingleton;
+}
+
+void FRCTypeTranslator::Translate(const URCVirtualPropertyBase* InSourceVirtualProperty, const TArray<URCVirtualPropertyBase*>& InTargetVirtualProperties)
+{
+	if (InSourceVirtualProperty)
+	{
+		const EPropertyBagPropertyType SourceValueType = InSourceVirtualProperty->GetValueType();
+
+		// Switch could have been an option, but if/else seems more readable,
+		// especially since each case is defining a type variable and that would lead to a lot of brackets
+		if (SourceValueType == EPropertyBagPropertyType::Bool)
+		{
+			bool Value;
+			InSourceVirtualProperty->GetValueBool(Value);
+			Apply(Value, InTargetVirtualProperties);
+		}
+		else if (SourceValueType == EPropertyBagPropertyType::Byte)
+		{
+			uint8 Value;
+			InSourceVirtualProperty->GetValueByte(Value);
+			Apply(Value, InTargetVirtualProperties);
+		}
+		else if (SourceValueType == EPropertyBagPropertyType::Double)
+		{
+			double Value;
+			InSourceVirtualProperty->GetValueDouble(Value);
+			Apply(Value, InTargetVirtualProperties);
+		}
+		else if (SourceValueType == EPropertyBagPropertyType::Float)
+		{
+			float Value;
+			InSourceVirtualProperty->GetValueFloat(Value);
+			Apply(Value, InTargetVirtualProperties);
+		}
+		else if (SourceValueType == EPropertyBagPropertyType::Int32)
+		{
+			int32 Value;
+			InSourceVirtualProperty->GetValueInt32(Value);
+			Apply(Value, InTargetVirtualProperties);
+		}
+		else if (SourceValueType == EPropertyBagPropertyType::Int64)
+		{
+			int64 Value;
+			InSourceVirtualProperty->GetValueInt64(Value);
+			Apply(Value, InTargetVirtualProperties);
+		}
+		else if (SourceValueType == EPropertyBagPropertyType::Name)
+		{
+			FName Value;
+			InSourceVirtualProperty->GetValueName(Value);
+			Apply(Value, InTargetVirtualProperties);
+		}
+		else if (SourceValueType == EPropertyBagPropertyType::String)
+		{
+			FString Value;
+			InSourceVirtualProperty->GetValueString(Value);
+			Apply(Value, InTargetVirtualProperties);
+		}
+		else if (SourceValueType == EPropertyBagPropertyType::Text)
+		{
+			FText Value;
+			InSourceVirtualProperty->GetValueText(Value);
+			Apply(Value, InTargetVirtualProperties);
+		}
+		else if (SourceValueType == EPropertyBagPropertyType::Struct)
+		{
+			// todo: Work out structs (e.g. vector, rotator, color, etc)
+		}
+
+		// todo: Handle other types if needed, or remove/collapse types
+	}
+}
+
+EPropertyBagPropertyType FRCTypeTranslator::GetOptimalValueType(const TArray<EPropertyBagPropertyType>& ValueTypes)
+{
+	// todo: Improve this function, it is currently just going in order!
+	EPropertyBagPropertyType OptimalType = EPropertyBagPropertyType::Bool;
+	for (const EPropertyBagPropertyType ValueType : ValueTypes)
+	{
+		if (OptimalType < ValueType)
+		{
+			OptimalType = ValueType;
+		}
+	}
+
+	return OptimalType;
+}
+
+FRCTypeHandler* FRCTypeTranslator::CreateTypeHandler(EPropertyBagPropertyType InValueType)
+{
+	// todo: Implement missing types
+	switch (InValueType)
+	{
+		case EPropertyBagPropertyType::None:
+			break;
+
+		case EPropertyBagPropertyType::Bool:
+			return new FRCBoolHandler();
+
+		case EPropertyBagPropertyType::Byte:
+			break;
+			
+		case EPropertyBagPropertyType::Int32:
+			break;
+			
+		case EPropertyBagPropertyType::Int64:
+			break;
+			
+		case EPropertyBagPropertyType::Float:
+			return new FRCFloatHandler();
+			break;
+			
+		case EPropertyBagPropertyType::Double:
+			break;
+			
+		case EPropertyBagPropertyType::Name:
+			break;
+			
+		case EPropertyBagPropertyType::String:
+			return new FRCStringHandler();
+			
+		case EPropertyBagPropertyType::Text:
+			break;
+			
+		case EPropertyBagPropertyType::Struct:
+			break;
+				
+		default: ;
+	}
+
+	return nullptr;
+}
diff --git a/Engine/Plugins/VirtualProduction/RemoteControl/Source/RemoteControlUI/Private/UI/Controller/TypeTranslator/RCTypeTranslator.h b/Engine/Plugins/VirtualProduction/RemoteControl/Source/RemoteControlUI/Private/UI/Controller/TypeTranslator/RCTypeTranslator.h
new file mode 100644
--- /dev/null
+++ b/Engine/Plugins/VirtualProduction/RemoteControl/Source/RemoteControlUI/Private/UI/Controller/TypeTranslator/RCTypeTranslator.h
@@ -0,0 +1,53 @@
+﻿// Copyright Epic Games, Inc. All Rights Reserved.
+
+#pragma once
+
+#include "PropertyBag.h"
+#include "RCTypeHandler.h"
+#include "Controller/RCController.h"
+
+class FRCTypeTranslator
+{
+public:
+	static FRCTypeTranslator* Get();
+
+	/** 
+	 * Conversion function: tries to use the Source Virtual Property current value to update the specified Target Virtual Properties values.
+	 * Particularly useful for MultiControllers.
+	 *
+	 * @param InSourceVirtualProperty the virtual property providing the source value
+	 * @param InTargetVirtualProperties the properties which will be updated using the source value
+	 */
+	static void Translate(const URCVirtualPropertyBase* InSourceVirtualProperty, const TArray<URCVirtualPropertyBase*>& InTargetVirtualProperties);
+
+	/** Returns the "optimal" value type among the specified list */
+	static EPropertyBagPropertyType GetOptimalValueType(const TArray<EPropertyBagPropertyType>& ValueTypes);
+	
+private:
+	template<typename InValueType>
+	static void Apply(const InValueType& Value, const TArray<URCVirtualPropertyBase*>& InTargetVirtualProperties)
+	{
+		for (URCVirtualPropertyBase* TargetVirtualProperty : InTargetVirtualProperties)
+		{
+			if (!TargetVirtualProperty)
+			{
+				continue;
+			}
+
+			if (FRCTypeHandler** TypeHandlerPtr = TypeHandlers.Find(TargetVirtualProperty->GetValueType()))
+			{
+				if (FRCTypeHandler* TypeHandler = *TypeHandlerPtr)
+				{
+					TypeHandler->Apply(Value, TargetVirtualProperty);
+				}
+			}
+		}
+	}
+
+	static FRCTypeHandler* CreateTypeHandler(EPropertyBagPropertyType InValueType);
+
+	FRCTypeTranslator();
+	
+	inline static FRCTypeTranslator* TypeTranslatorSingleton = nullptr;
+	inline static TMap<EPropertyBagPropertyType, FRCTypeHandler*> TypeHandlers;
+};
diff --git a/Engine/Plugins/VirtualProduction/RemoteControl/Source/RemoteControlUI/Public/IRemoteControlUIModule.h b/Engine/Plugins/VirtualProduction/RemoteControl/Source/RemoteControlUI/Public/IRemoteControlUIModule.h
--- a/Engine/Plugins/VirtualProduction/RemoteControl/Source/RemoteControlUI/Public/IRemoteControlUIModule.h
+++ b/Engine/Plugins/VirtualProduction/RemoteControl/Source/RemoteControlUI/Public/IRemoteControlUIModule.h
@@ -141,6 +141,10 @@ public:
 		return FModuleManager::LoadModuleChecked<IRemoteControlUIModule>(ModuleName);
 	}
 	
+	/** Delegate called when a Remote Control Preset panel is created/opened for the specified Remote Control Preset */
+    DECLARE_MULTICAST_DELEGATE_OneParam(FOnRemoteControlPresetOpened, URemoteControlPreset* /*RemoteControlPreset*/);
+    virtual FOnRemoteControlPresetOpened& OnRemoteControlPresetOpened() = 0;
+	
 	/** 
 	 * Get the toolbar extension generators.
 	 * Usage: Bind a handler that adds a widget to the out array parameter.
