1. Remote Control Preset property valid status not updated after being rebound

2. Brighten up Vertical Line Dividers in Column Headers for Improved Visibility

diff --git a/Engine/Plugins/VirtualProduction/RemoteControl/Source/RemoteControlUI/Private/UI/RCPanelWidgetRegistry.cpp b/Engine/Plugins/VirtualProduction/RemoteControl/Source/RemoteControlUI/Private/UI/RCPanelWidgetRegistry.cpp
--- a/Engine/Plugins/VirtualProduction/RemoteControl/Source/RemoteControlUI/Private/UI/RCPanelWidgetRegistry.cpp
+++ b/Engine/Plugins/VirtualProduction/RemoteControl/Source/RemoteControlUI/Private/UI/RCPanelWidgetRegistry.cpp
@@ -116,6 +116,13 @@ TSharedPtr<IDetailTreeNode> FRCPanelWidgetRegistry::GetObjectTreeNode(UObject* I
 	{
 		if (Node->IsValid() && Node->Pin()->CreatePropertyHandle() && Node->Pin()->CreatePropertyHandle()->GetNumOuterObjects() != 0)
 		{
+			if (TSharedPtr<IPropertyRowGenerator>* FoundGenerator = ObjectToRowGenerator.Find({ InObject }))
+			{
+				// Since we are returning cached node from the second time refresh onwards
+				// It is important that we refresh the corresponding row to keep the consistency.
+				FoundGenerator->Get()->OnRowsRefreshed().Broadcast();
+			}
+
 			return Node->Pin();
 		}
 		else
diff --git a/Engine/Plugins/VirtualProduction/RemoteControl/Source/RemoteControlUI/Private/UI/RemoteControlPanelStyle.cpp b/Engine/Plugins/VirtualProduction/RemoteControl/Source/RemoteControlUI/Private/UI/RemoteControlPanelStyle.cpp
--- a/Engine/Plugins/VirtualProduction/RemoteControl/Source/RemoteControlUI/Private/UI/RemoteControlPanelStyle.cpp
+++ b/Engine/Plugins/VirtualProduction/RemoteControl/Source/RemoteControlUI/Private/UI/RemoteControlPanelStyle.cpp
@@ -163,6 +163,10 @@ void FRemoteControlPanelStyle::SetupPanelStyles(TSharedRef<FSlateStyleSet> InSty
 	// Header Row Style
 	FHeaderRowStyle HeaderRowStyle = AppStyle.GetWidgetStyle<FHeaderRowStyle>("PropertyTable.HeaderRow");
 	HeaderRowStyle.BackgroundBrush = FSlateColorBrush(FStyleColors::Background);
+	HeaderRowStyle.SplitterHandleSize = 1.f;
+	HeaderRowStyle.ColumnSplitterStyle = FSplitterStyle()
+	.SetHandleNormalBrush(FSlateColorBrush(FStyleColors::AccentGray))
+	.SetHandleHighlightBrush(FSlateColorBrush(FStyleColors::AccentGray));
 	HeaderRowStyle.HorizontalSeparatorBrush = FSlateColorBrush(FStyleColors::Background);
 
 	// Table Row Style
