Remote Control not releasing Midi Controller properly

diff --git a/Engine/Plugins/VirtualProduction/RemoteControl/Source/RemoteControlProtocolWidgets/Private/ViewModels/ProtocolEntityViewModel.cpp b/Engine/Plugins/VirtualProduction/RemoteControl/Source/RemoteControlProtocolWidgets/Private/ViewModels/ProtocolEntityViewModel.cpp
--- a/Engine/Plugins/VirtualProduction/RemoteControl/Source/RemoteControlProtocolWidgets/Private/ViewModels/ProtocolEntityViewModel.cpp
+++ b/Engine/Plugins/VirtualProduction/RemoteControl/Source/RemoteControlProtocolWidgets/Private/ViewModels/ProtocolEntityViewModel.cpp
@@ -3,8 +3,10 @@
 #include "ProtocolEntityViewModel.h"
 
 #include "Editor.h"
+#include "IRCProtocolBindingList.h"
 #include "IRemoteControlProtocol.h"
 #include "IRemoteControlProtocolModule.h"
+#include "IRemoteControlProtocolWidgetsModule.h"
 #include "ProtocolBindingViewModel.h"
 #include "ProtocolCommandChange.h"
 #include "RemoteControlPreset.h"
@@ -85,7 +87,24 @@ namespace Commands
 				// Unbind the protocol from the property explicitly while removing.
 				if (const TSharedPtr<IRemoteControlProtocol> Protocol = IRemoteControlProtocolModule::Get().GetProtocolByName(FoundElement->GetProtocolName()))
 				{
-					Protocol->Unbind(FoundElement->GetRemoteControlProtocolEntityPtr());
+					if (TSharedPtr<TStructOnScope<FRemoteControlProtocolEntity>> RCProtocolEntityPtr = FoundElement->GetRemoteControlProtocolEntityPtr())
+					{
+						const ERCBindingStatus BindingStatus = (*RCProtocolEntityPtr)->GetBindingStatus();
+
+						IRemoteControlProtocolWidgetsModule& RCProtocolWidgetsModule = IRemoteControlProtocolWidgetsModule::Get();
+						
+						if (BindingStatus == ERCBindingStatus::Bound)
+						{
+							if (TSharedPtr<IRCProtocolBindingList> RCProtocolBindingList = RCProtocolWidgetsModule.GetProtocolBindingList())
+							{
+								RCProtocolBindingList->OnStopRecording(RCProtocolEntityPtr);
+							}
+						}
+
+						(*RCProtocolEntityPtr)->ResetDefaultBindingState();
+
+						Protocol->Unbind(RCProtocolEntityPtr);
+					}
 				}
 
 				RCProperty->ProtocolBindings.RemoveByHash(GetTypeHash(InArgs.BindingId), InArgs.BindingId);
