Remote Control:
- Added fix for Double Properties not working for Bind

diff --git a/Engine/Plugins/VirtualProduction/RemoteControl/Source/RemoteControlLogic/Private/Action/Bind/RCPropertyBindAction.cpp b/Engine/Plugins/VirtualProduction/RemoteControl/Source/RemoteControlLogic/Private/Action/Bind/RCPropertyBindAction.cpp
--- a/Engine/Plugins/VirtualProduction/RemoteControl/Source/RemoteControlLogic/Private/Action/Bind/RCPropertyBindAction.cpp
+++ b/Engine/Plugins/VirtualProduction/RemoteControl/Source/RemoteControlLogic/Private/Action/Bind/RCPropertyBindAction.cpp
@@ -159,6 +159,10 @@ void URCPropertyBindAction::Execute() const
 				{
 					Handle->SetValue((int32)NumericValue);
 				}
+				else if (RemoteControlProperty->IsA(FDoubleProperty::StaticClass()))
+				{
+					Handle->SetValue((double)NumericValue);
+				}
 			}
 			// Numeric To Boolean
 			else if (RemoteControlProperty->IsA(FBoolProperty::StaticClass()))
