[RemoteControl] Fix property reset over HTTP. Fix property reset for properties with no CDO data

diff --git a/Engine/Plugins/VirtualProduction/RemoteControl/Source/RemoteControl/Private/RemoteControlModule.cpp b/Engine/Plugins/VirtualProduction/RemoteControl/Source/RemoteControl/Private/RemoteControlModule.cpp
--- a/Engine/Plugins/VirtualProduction/RemoteControl/Source/RemoteControl/Private/RemoteControlModule.cpp
+++ b/Engine/Plugins/VirtualProduction/RemoteControl/Source/RemoteControl/Private/RemoteControlModule.cpp
@@ -1760,13 +1760,35 @@ bool FRemoteControlModule::ResetObjectProperties(const FRCObjectReference& Objec
 		else
 #endif
 		{
-			// Copy the value from the field on the CDO.
 			FRCFieldPathInfo FieldPathInfo = ObjectAccess.PropertyPathInfo;
 			void* TargetAddress = FieldPathInfo.GetResolvedData().ContainerAddress;
+
 			UObject* DefaultObject = Object->GetClass()->GetDefaultObject();
-			FieldPathInfo.Resolve(DefaultObject);
-			FRCFieldResolvedData DefaultObjectResolvedData = FieldPathInfo.GetResolvedData();
-			ObjectAccess.Property->CopyCompleteValue_InContainer(TargetAddress, DefaultObjectResolvedData.ContainerAddress);
+			if (FieldPathInfo.Resolve(DefaultObject))
+			{
+				// Copy the value from the field on the CDO
+				FRCFieldResolvedData DefaultObjectResolvedData = FieldPathInfo.GetResolvedData();
+				ObjectAccess.Property->CopyCompleteValue_InContainer(TargetAddress, DefaultObjectResolvedData.ContainerAddress);
+			}
+			else if (UStruct* ContainerStruct = ObjectAccess.ContainerType.Get())
+			{
+				// Structs have no CDO, so initialize a new instance of the struct and copy the value from there
+				TArray<uint8> NewStructDataArray;
+				const int32 StructureSize = ContainerStruct->GetStructureSize();
+				NewStructDataArray.SetNumUninitialized(StructureSize);
+
+				void* NewStructData = NewStructDataArray.GetData();
+				ContainerStruct->InitializeStruct(NewStructData);
+
+				ObjectAccess.Property->CopyCompleteValue_InContainer(TargetAddress, NewStructData);
+
+				ContainerStruct->DestroyStruct(NewStructData);
+			}
+			else
+			{
+				// Containing object has no default value, so fall back to the initial value for this property type
+				ObjectAccess.Property->InitializeValue_InContainer(TargetAddress);
+			}
 		}
 
 		// if we are generating a transaction, also generate post edit property event, event if the change ended up unsuccessful
diff --git a/Engine/Plugins/VirtualProduction/RemoteControl/Source/WebRemoteControl/Private/WebRemoteControl.cpp b/Engine/Plugins/VirtualProduction/RemoteControl/Source/WebRemoteControl/Private/WebRemoteControl.cpp
--- a/Engine/Plugins/VirtualProduction/RemoteControl/Source/WebRemoteControl/Private/WebRemoteControl.cpp
+++ b/Engine/Plugins/VirtualProduction/RemoteControl/Source/WebRemoteControl/Private/WebRemoteControl.cpp
@@ -995,7 +995,6 @@ bool FWebRemoteControlModule::HandleObjectPropertyRoute(const FHttpServerRequest
 {
 	TUniquePtr<FHttpServerResponse> Response = WebRemoteControlInternalUtils::CreateHttpResponse();
 
-	bool bResetToDefault = false;
 	FRCObjectReference ObjectRef;
 	FRCObjectRequest DeserializedRequest;
 	if (!WebRemoteControl::ResolveObjectPropertyForRequest(Request, OnComplete, Response, ObjectRef, DeserializedRequest))
@@ -1022,7 +1021,7 @@ bool FWebRemoteControlModule::HandleObjectPropertyRoute(const FHttpServerRequest
 	case ERCAccess::WRITE_MANUAL_TRANSACTION_ACCESS:
 	{
 		const FBlockDelimiters& PropertyValueDelimiters = DeserializedRequest.GetStructParameters().FindChecked(FRCObjectRequest::PropertyValueLabel());
-		if (bResetToDefault)
+		if (DeserializedRequest.ResetToDefault)
 		{
 			constexpr bool bAllowIntercept = true;
 			if (IRemoteControlModule::Get().ResetObjectProperties(ObjectRef, bAllowIntercept))
