Remote Control:
- Fix issue with Exposed properties from Structs in Containers (e.g. Array) not working with Controllers

diff --git a/Engine/Plugins/VirtualProduction/RemoteControl/Source/RemoteControl/Private/RemoteControlPropertyHandle.cpp b/Engine/Plugins/VirtualProduction/RemoteControl/Source/RemoteControl/Private/RemoteControlPropertyHandle.cpp
--- a/Engine/Plugins/VirtualProduction/RemoteControl/Source/RemoteControl/Private/RemoteControlPropertyHandle.cpp
+++ b/Engine/Plugins/VirtualProduction/RemoteControl/Source/RemoteControl/Private/RemoteControlPropertyHandle.cpp
@@ -448,7 +448,7 @@ FRemoteControlPropertyHandle::FRemoteControlPropertyHandle(const TSharedPtr<FRem
 	{
 		if (const UScriptStruct* Struct = Cast<UScriptStruct>(InProperty->GetOwnerUObject()))
 		{
-			FieldPath = InRCProperty->FieldPathInfo.ToPathPropertyString();
+			FieldPath = InRCProperty->FieldPathInfo.ToString();
 		}
 		else
 		{
