[Remote Control] Rebinding Remote Control Property Does not change owner name of multiple properties

Adds an override to Refresh in SRCPanelExposedEntity, and call SRCPanelExposedEntity::Refresh from child classes SRCPanelExposedField::Refresh() and SRCPanelExposedActpr::Refresh().
This allows the Entity internal cached data to be updated after the Rebind.

Upvoted ver.5
Version 6 is merging changes coming from latest.

diff --git a/Engine/Plugins/VirtualProduction/RemoteControl/Source/RemoteControlUI/Private/UI/SRCPanelExposedActor.cpp b/Engine/Plugins/VirtualProduction/RemoteControl/Source/RemoteControlUI/Private/UI/SRCPanelExposedActor.cpp
--- a/Engine/Plugins/VirtualProduction/RemoteControl/Source/RemoteControlUI/Private/UI/SRCPanelExposedActor.cpp
+++ b/Engine/Plugins/VirtualProduction/RemoteControl/Source/RemoteControlUI/Private/UI/SRCPanelExposedActor.cpp
@@ -53,6 +53,8 @@ SRCPanelTreeNode::ENodeType SRCPanelExposedActor::GetRCType() const
 
 void SRCPanelExposedActor::Refresh()
 {
+	SRCPanelExposedEntity::Refresh();
+	
 	if (TSharedPtr<FRemoteControlActor> RCActor = WeakActor.Pin())
 	{
 		CachedLabel = RCActor->GetLabel();
@@ -157,4 +159,4 @@ TSharedRef<SWidget> SRCPanelExposedActor::CreateEmbeddedPresetActorPicker()
 	return ActorPicker;
 }
 
-#undef LOCTEXT_NAMESPACE /*RemoteControlPanelActor*/
\ No newline at end of file
+#undef LOCTEXT_NAMESPACE /*RemoteControlPanelActor*/
diff --git a/Engine/Plugins/VirtualProduction/RemoteControl/Source/RemoteControlUI/Private/UI/SRCPanelExposedEntity.cpp b/Engine/Plugins/VirtualProduction/RemoteControl/Source/RemoteControlUI/Private/UI/SRCPanelExposedEntity.cpp
--- a/Engine/Plugins/VirtualProduction/RemoteControl/Source/RemoteControlUI/Private/UI/SRCPanelExposedEntity.cpp
+++ b/Engine/Plugins/VirtualProduction/RemoteControl/Source/RemoteControlUI/Private/UI/SRCPanelExposedEntity.cpp
@@ -141,6 +141,14 @@ void SRCPanelExposedEntity::EnterRenameMode()
 	}
 }
 
+void SRCPanelExposedEntity::Refresh()
+{
+	if (EntityId.IsValid() && Preset.IsValid())
+	{
+		Initialize(EntityId, Preset.Get(), bLiveMode);
+	}
+}
+
 void SRCPanelExposedEntity::Initialize(const FGuid& InEntityId, URemoteControlPreset* InPreset, const TAttribute<bool>& InbLiveMode)
 {
 	EntityId = InEntityId;
diff --git a/Engine/Plugins/VirtualProduction/RemoteControl/Source/RemoteControlUI/Private/UI/SRCPanelExposedEntity.h b/Engine/Plugins/VirtualProduction/RemoteControl/Source/RemoteControlUI/Private/UI/SRCPanelExposedEntity.h
--- a/Engine/Plugins/VirtualProduction/RemoteControl/Source/RemoteControlUI/Private/UI/SRCPanelExposedEntity.h
+++ b/Engine/Plugins/VirtualProduction/RemoteControl/Source/RemoteControlUI/Private/UI/SRCPanelExposedEntity.h
@@ -30,6 +30,8 @@ struct SRCPanelExposedEntity : public SRCPanelTreeNode
 	/** Updates the highlight text to active search term. */
 	virtual void SetHighlightText(const FText& InHightlightText = FText::GetEmpty()) override { HighlightText = InHightlightText; }
 
+	virtual void Refresh() override;
+	
 protected:
 	void Initialize(const FGuid& InEntityId, URemoteControlPreset* InPreset, const TAttribute<bool>& InbLiveMode);
 	
@@ -146,4 +148,4 @@ public:
 private:
 	FGuid Id;
 	TSharedPtr<SWidget> DecoratorWidget;
-};
\ No newline at end of file
+};
diff --git a/Engine/Plugins/VirtualProduction/RemoteControl/Source/RemoteControlUI/Private/UI/SRCPanelExposedField.cpp b/Engine/Plugins/VirtualProduction/RemoteControl/Source/RemoteControlUI/Private/UI/SRCPanelExposedField.cpp
--- a/Engine/Plugins/VirtualProduction/RemoteControl/Source/RemoteControlUI/Private/UI/SRCPanelExposedField.cpp
+++ b/Engine/Plugins/VirtualProduction/RemoteControl/Source/RemoteControlUI/Private/UI/SRCPanelExposedField.cpp
@@ -276,6 +276,8 @@ const bool SRCPanelExposedField::SupportsProtocol(const FName& InProtocolName) c
 
 void SRCPanelExposedField::Refresh()
 {
+	SRCPanelExposedEntity::Refresh();
+	
 	if (TSharedPtr<FRemoteControlField> Field = WeakField.Pin())
 	{
 		CachedLabel = Field->GetLabel();
