Fix remote control crash on exit

diff --git a/Engine/Plugins/VirtualProduction/RemoteControl/Source/WebRemoteControl/Private/WebRemoteControl.cpp b/Engine/Plugins/VirtualProduction/RemoteControl/Source/WebRemoteControl/Private/WebRemoteControl.cpp
--- a/Engine/Plugins/VirtualProduction/RemoteControl/Source/WebRemoteControl/Private/WebRemoteControl.cpp
+++ b/Engine/Plugins/VirtualProduction/RemoteControl/Source/WebRemoteControl/Private/WebRemoteControl.cpp
@@ -386,7 +386,10 @@ void FWebRemoteControlModule::ShutdownModule()
 	}
 
 	EditorRoutes.UnregisterRoutes(this);
-	WebSocketHandler->UnregisterRoutes(this);
+	if (WebSocketHandler)
+	{
+		WebSocketHandler->UnregisterRoutes(this);
+	}
 	StopHttpServer();
 	StopWebSocketServer();
 	UnregisterConsoleCommands();
