[RemoteControl] Fix property changes being missed in manual transactions

diff --git a/Engine/Plugins/VirtualProduction/RemoteControl/Source/RemoteControl/Private/RemoteControlModule.cpp b/Engine/Plugins/VirtualProduction/RemoteControl/Source/RemoteControl/Private/RemoteControlModule.cpp
--- a/Engine/Plugins/VirtualProduction/RemoteControl/Source/RemoteControl/Private/RemoteControlModule.cpp
+++ b/Engine/Plugins/VirtualProduction/RemoteControl/Source/RemoteControl/Private/RemoteControlModule.cpp
@@ -1297,7 +1297,7 @@ bool FRemoteControlModule::StartPropertyTransaction(FRCObjectReference& ObjectRe
 		}
 	}
 
-	if (!bUseOngoingChangeOptimization)
+	if (!bOutGeneratedTransaction)
 	{
 		FEditPropertyChain PreEditChain;
 		ObjectReference.PropertyPathInfo.ToEditPropertyChain(PreEditChain);
@@ -1311,26 +1311,31 @@ bool FRemoteControlModule::StartPropertyTransaction(FRCObjectReference& ObjectRe
 #if WITH_EDITOR
 void FRemoteControlModule::SnapshotOrEndTransaction(FRCObjectReference& ObjectReference, bool bGeneratedTransaction)
 {
-	if (CVarRemoteControlEnableOngoingChangeOptimization.GetValueOnAnyThread() == 1)
+	const bool bUseOngoingChangeOptimization = CVarRemoteControlEnableOngoingChangeOptimization.GetValueOnAnyThread() == 1;
+	const bool bHasOngoingTransaction = bUseOngoingChangeOptimization && OngoingModification
+		&& GetTypeHash(*OngoingModification) == GetTypeHash(ObjectReference) && OngoingModification->bHasStartedTransaction;
+
+	if (bHasOngoingTransaction || ObjectReference.Access == ERCAccess::WRITE_MANUAL_TRANSACTION_ACCESS)
+	{
+		SnapshotTransactionBuffer(ObjectReference.Object.Get());
+		FPropertyChangedEvent PropertyEvent(ObjectReference.PropertyPathInfo.ToPropertyChangedEvent());
+		PropertyEvent.ChangeType = EPropertyChangeType::Interactive;
+		ObjectReference.Object->PostEditChangeProperty(PropertyEvent);
+
+		if (UActorComponent* Component = Cast<UActorComponent>(ObjectReference.Object.Get()))
+		{
+			Component->MarkRenderStateDirty();
+			Component->UpdateComponentToWorld();
+		}
+	}
+
+	if (bUseOngoingChangeOptimization)
 	{
 		// If we have modified the same object and property in the last frames,
 		// update the triggered flag and snapshot the object to the transaction buffer.
 		if (OngoingModification && GetTypeHash(*OngoingModification) == GetTypeHash(ObjectReference))
 		{
 			OngoingModification->bWasTriggeredSinceLastPass = true;
-			if (OngoingModification->bHasStartedTransaction)
-			{
-				SnapshotTransactionBuffer(ObjectReference.Object.Get());
-				FPropertyChangedEvent PropertyEvent(ObjectReference.PropertyPathInfo.ToPropertyChangedEvent());
-				PropertyEvent.ChangeType = EPropertyChangeType::Interactive;
-				ObjectReference.Object->PostEditChangeProperty(PropertyEvent);
-
-				if (UActorComponent* Component = Cast<UActorComponent>(ObjectReference.Object.Get()))
-				{
-					Component->MarkRenderStateDirty();
-					Component->UpdateComponentToWorld();
-				}
-			}
 
 			// Update the world lighting if we're modifying a color.
 			if (ObjectReference.IsValid())
