Fix remote control not finding object in PIE

diff --git a/Engine/Plugins/VirtualProduction/RemoteControl/Source/RemoteControl/Private/RemoteControlBinding.cpp b/Engine/Plugins/VirtualProduction/RemoteControl/Source/RemoteControl/Private/RemoteControlBinding.cpp
--- a/Engine/Plugins/VirtualProduction/RemoteControl/Source/RemoteControl/Private/RemoteControlBinding.cpp
+++ b/Engine/Plugins/VirtualProduction/RemoteControl/Source/RemoteControl/Private/RemoteControlBinding.cpp
@@ -64,6 +64,22 @@ namespace UE::RemoteControlBinding
 
 		return true;
 	}
+
+	FString ExtractObjectFromPath(const FSoftObjectPath& OriginalObjectPath)
+	{
+		const static FString PersistentLevelString("PersistentLevel.");
+		const int32 PersistentLevelStringLength = PersistentLevelString.Len();
+		// /Game/MapName.MapName:PersistentLevel.StaticMeshActor_42.StaticMeshComponent becomes PersistentLevel.StaticMeshActor_42.StaticMeshComponent
+		const FString& SubPathString = OriginalObjectPath.GetSubPathString();
+		const int32 IndexOfPersistentLevelInfo = SubPathString.Find(PersistentLevelString, ESearchCase::CaseSensitive);
+		if (IndexOfPersistentLevelInfo == INDEX_NONE)
+		{
+			return TEXT("");
+		}
+
+		return SubPathString.RightChop(PersistentLevelStringLength);
+	}
+
 }
 
 namespace
@@ -252,11 +268,16 @@ void URemoteControlLevelDependantBinding::UnbindObject(const TSoftObjectPtr<UObj
 
 UObject* URemoteControlLevelDependantBinding::Resolve() const
 {
-	// Find the object in PIE if possible
-	UObject* Object = ResolveForCurrentWorld().Get();
+	constexpr bool bAllowPIE = true;
+	UObject* Object = ResolveForCurrentWorld(bAllowPIE).Get();
 
 	if (Object)
 	{
+		if (Object->GetWorld() && Object->GetWorld()->WorldType == EWorldType::PIE)
+		{
+			// Don't update path if we manually resolved to a PIE object. (Can happen if editor world gets unloaded)
+			return Object;
+		}
 		// Make sure we don't resolve on a subobject of a dying parent actor.
 		if (AActor* OwnerActor = Object->GetTypedOuter<AActor>())
 		{
@@ -358,16 +379,28 @@ void URemoteControlLevelDependantBinding::PostLoad()
 	}
 }
 
-TSoftObjectPtr<UObject> URemoteControlLevelDependantBinding::ResolveForCurrentWorld() const
+TSoftObjectPtr<UObject> URemoteControlLevelDependantBinding::ResolveForCurrentWorld(bool bAllowPIE) const
 {
 	// Note that we use Find rather than FindRef throughout this function. This lets us dereference
 	// the BoundObjectMapByPath/SubLevelSelectionMapByPath's values directly so that their internal weak pointers
 	// are updated when stale (rather than dereferencing copies and leaving stale pointers in the originals).
 
-	if (UWorld* World = GetCurrentWorld())
+	if (UWorld* World = GetCurrentWorld(bAllowPIE))
 	{
+		FSoftObjectPath LevelPath = World;
+#if WITH_EDITOR
 		// Try finding the object using the sub level selection map first.
-		if (TSoftObjectPtr<ULevel>* LastBindingLevel = SubLevelSelectionMapByPath.Find(World))
+		if (World->WorldType == EWorldType::PIE)
+		{
+			FString PIEPackageName = World->GetPackage()->GetPathName();
+			FString OriginalPackage = UWorld::StripPIEPrefixFromPackageName(PIEPackageName, World->StreamingLevelsPrefix);
+			FTopLevelAssetPath TopLevelAssetPath(FName(OriginalPackage), World->GetFName());
+
+			LevelPath = FSoftObjectPath(TopLevelAssetPath);
+		}
+#endif
+
+		if (TSoftObjectPtr<ULevel>* LastBindingLevel = SubLevelSelectionMapByPath.Find(LevelPath))
 		{
 			if (const TSoftObjectPtr<UObject>* ObjectPtr = BoundObjectMapByPath.Find(LastBindingLevel->ToSoftObjectPath()))
 			{
@@ -375,6 +408,30 @@ TSoftObjectPtr<UObject> URemoteControlLevelDependantBinding::ResolveForCurrentWo
 				{
 					return *ObjectPtr;
 				}
+
+#if WITH_EDITOR
+				if (World->WorldType == EWorldType::PIE)
+				{
+					bool bPathToActor = true;
+
+					FString ActorName = UE::RemoteControlBinding::ExtractObjectFromPath(ObjectPtr->ToSoftObjectPath());
+					if (!ActorName.IsEmpty())
+					{
+						constexpr bool bExactClass = false;
+						
+						if (UObject* SimWorldObject = FindObject<UObject>(World->PersistentLevel, *ActorName, bExactClass))
+						{
+							if ((SimWorldObject->IsA<AActor>() && GEditor->ObjectsThatExistInEditorWorld.Get(SimWorldObject))
+								|| (GEditor->ObjectsThatExistInEditorWorld.Get(SimWorldObject->GetTypedOuter<AActor>())))
+							{
+								return SimWorldObject;
+							}
+						}
+					}
+
+				}
+#endif
+				
 			}
 		}
 		
@@ -425,12 +482,10 @@ UClass* URemoteControlLevelDependantBinding::GetSupportedOwnerClass() const
 	return BindingContext.OwnerActorClass.LoadSynchronous();
 }
 
-UWorld* URemoteControlLevelDependantBinding::GetCurrentWorld() const
+UWorld* URemoteControlLevelDependantBinding::GetCurrentWorld(bool bAllowPIE) const
 {
 	URemoteControlPreset* Preset = Cast<URemoteControlPreset>(GetOuter());
 
-	// Since this is used to retrieve the binding in the map, we never use the PIE world in editor.
-	constexpr bool bAllowPIE = false;
 	return URemoteControlPreset::GetWorld(Preset, bAllowPIE);
 }
 
diff --git a/Engine/Plugins/VirtualProduction/RemoteControl/Source/RemoteControl/Private/RemoteControlPreset.cpp b/Engine/Plugins/VirtualProduction/RemoteControl/Source/RemoteControl/Private/RemoteControlPreset.cpp
--- a/Engine/Plugins/VirtualProduction/RemoteControl/Source/RemoteControl/Private/RemoteControlPreset.cpp
+++ b/Engine/Plugins/VirtualProduction/RemoteControl/Source/RemoteControl/Private/RemoteControlPreset.cpp
@@ -2090,6 +2090,8 @@ void URemoteControlPreset::RegisterDelegates()
 
 	FCoreUObjectDelegates::OnPackageReloaded.AddUObject(this, &URemoteControlPreset::OnPackageReloaded);
 #endif
+	
+	FCoreUObjectDelegates::PostLoadMapWithWorld.AddUObject(this, &URemoteControlPreset::OnMapLoadFinished);
 
 	FCoreDelegates::OnBeginFrame.AddUObject(this, &URemoteControlPreset::OnBeginFrame);
 	FCoreDelegates::OnEndFrame.AddUObject(this, &URemoteControlPreset::OnEndFrame);
@@ -2100,6 +2102,8 @@ void URemoteControlPreset::UnregisterDelegates()
 	FCoreDelegates::OnBeginFrame.RemoveAll(this);
 	FCoreDelegates::OnEndFrame.RemoveAll(this);
 
+	FCoreUObjectDelegates::PostLoadMapWithWorld.RemoveAll(this);
+
 #if WITH_EDITOR
 	FCoreUObjectDelegates::OnPackageReloaded.RemoveAll(this);
 
@@ -2320,6 +2324,15 @@ void URemoteControlPreset::CleanUpBindings()
 }
 #endif
 
+void URemoteControlPreset::OnMapLoadFinished(UWorld* LoadedWorld)
+{
+	if (Registry)
+	{
+		Algo::Transform(Registry->GetExposedEntities(), PerFrameUpdatedEntities, [](const TSharedPtr<FRemoteControlEntity>& Entity) { return Entity->GetId(); });
+		Algo::Transform(Registry->GetExposedEntities<FRemoteControlProperty>(), PerFrameModifiedProperties, [](const TSharedPtr<FRemoteControlProperty>& RCProp) { return RCProp->GetId(); });
+	}
+}
+
 void URemoteControlPreset::OnBeginFrame()
 {
 	TRACE_CPUPROFILER_EVENT_SCOPE(URemoteControlPreset::OnBeginFrame);
diff --git a/Engine/Plugins/VirtualProduction/RemoteControl/Source/RemoteControl/Public/RemoteControlBinding.h b/Engine/Plugins/VirtualProduction/RemoteControl/Source/RemoteControl/Public/RemoteControlBinding.h
--- a/Engine/Plugins/VirtualProduction/RemoteControl/Source/RemoteControl/Public/RemoteControlBinding.h
+++ b/Engine/Plugins/VirtualProduction/RemoteControl/Source/RemoteControl/Public/RemoteControlBinding.h
@@ -227,11 +227,13 @@ public:
 	UClass* GetSupportedOwnerClass() const;
 
 	/** Get the current world associated with the preset or the editor world as a fallback. */
-	UWorld* GetCurrentWorld() const;
+	UWorld* GetCurrentWorld(bool bAllowPIE = false) const;
 
 private:
-	TSoftObjectPtr<UObject> ResolveForCurrentWorld() const;
+	/** Attempt resolving the binding for the current world. */
+	TSoftObjectPtr<UObject> ResolveForCurrentWorld(bool bAllowPIE = false) const;
 
+	/** Update this binding's context with the object passed as argument. */
 	void UpdateBindingContext(UObject* InObject) const;
 
 
diff --git a/Engine/Plugins/VirtualProduction/RemoteControl/Source/RemoteControl/Public/RemoteControlPreset.h b/Engine/Plugins/VirtualProduction/RemoteControl/Source/RemoteControl/Public/RemoteControlPreset.h
--- a/Engine/Plugins/VirtualProduction/RemoteControl/Source/RemoteControl/Public/RemoteControlPreset.h
+++ b/Engine/Plugins/VirtualProduction/RemoteControl/Source/RemoteControl/Public/RemoteControlPreset.h
@@ -744,6 +744,9 @@ private:
 	//~ Frame events handlers.
 	void OnBeginFrame();
 	void OnEndFrame();
+
+	/** Handles dynamic level change. */
+	void OnMapLoadFinished(UWorld* LoadedWorld);
 	
 	/** Get a field ptr using it's id. */
 	FRemoteControlField* GetFieldPtr(FGuid FieldId);
