[Remote Control] - Duplicating Action Doesn't Modify Values with Range Map Behavior

URCRangeMapBehaviour::DuplicateAction was creating a new RangeMapInput for the duplicated property without creating a new InputProperty.
Consequence was that all duplicated Range Map Actionswere acting as a single one.

Fix creates a new URCVirtualPropertySelfContainer, initialized with the Input Property taken from the source action RangeMapInput.

diff --git a/Engine/Plugins/VirtualProduction/RemoteControl/Source/RemoteControlLogic/Private/Behaviour/Builtin/RangeMap/RCRangeMapBehaviour.cpp b/Engine/Plugins/VirtualProduction/RemoteControl/Source/RemoteControlLogic/Private/Behaviour/Builtin/RangeMap/RCRangeMapBehaviour.cpp
--- a/Engine/Plugins/VirtualProduction/RemoteControl/Source/RemoteControlLogic/Private/Behaviour/Builtin/RangeMap/RCRangeMapBehaviour.cpp
+++ b/Engine/Plugins/VirtualProduction/RemoteControl/Source/RemoteControlLogic/Private/Behaviour/Builtin/RangeMap/RCRangeMapBehaviour.cpp
@@ -343,9 +343,12 @@ URCAction* URCRangeMapBehaviour::DuplicateAction(URCAction* InAction, URCBehavio
 				ActionValueProperty = PropertyAction->PropertySelfContainer;
 			}
 
-			if(ensureAlways(ActionValueProperty))
+			URCVirtualPropertySelfContainer* ActionInputProperty = NewObject<URCVirtualPropertySelfContainer>(InBehaviour);
+			ActionInputProperty->DuplicatePropertyWithCopy(RangeMapInputData->InputProperty);
+
+			if (ensureAlways(ActionValueProperty))
 			{
-				const FRCRangeMapInput NewData(RangeMapInputData->InputProperty, ActionValueProperty);
+				const FRCRangeMapInput NewData(ActionInputProperty, ActionValueProperty);
 				InBehaviourRangeMap->RangeMapActionContainer.Add(NewAction, NewData);
 			}
 		}
@@ -547,4 +550,4 @@ bool URCRangeMapBehaviour::IsActionUnique(const TSharedRef<const FRemoteControlF
 	}
 	
 	return true;
-}
\ No newline at end of file
+}
