[Remote Control] Crash in RangeMap action

- Code was directly accessing a map value with a non valid key
- PrevValueDifference was being computed starting from a non-normalized prev input value.

diff --git a/Engine/Plugins/VirtualProduction/RemoteControl/Source/RemoteControlLogic/Private/Behaviour/Builtin/RangeMap/RCRangeMapBehaviour.cpp b/Engine/Plugins/VirtualProduction/RemoteControl/Source/RemoteControlLogic/Private/Behaviour/Builtin/RangeMap/RCRangeMapBehaviour.cpp
--- a/Engine/Plugins/VirtualProduction/RemoteControl/Source/RemoteControlLogic/Private/Behaviour/Builtin/RangeMap/RCRangeMapBehaviour.cpp
+++ b/Engine/Plugins/VirtualProduction/RemoteControl/Source/RemoteControlLogic/Private/Behaviour/Builtin/RangeMap/RCRangeMapBehaviour.cpp
@@ -118,7 +118,8 @@ bool URCRangeMapBehaviour::GetNearestActionByThreshold(TTuple<URCAction*, bool>&
 		
 		if (InputActionArray.Num() > 1 && InputIndex > 0)
 		{
-			const double PrevValueDifference = FMath::Abs(InputActionArray[InputIndex-1] - NormalizedControllerValue);
+			const double PrevNormalizedInputValue = UKismetMathLibrary::NormalizeToRange(InputActionArray[InputIndex-1], InputMin, InputMax);
+			const double PrevValueDifference = FMath::Abs(PrevNormalizedInputValue - NormalizedControllerValue);
 			
 			// Break out early in case previous Input had a lesser difference
 			if (PrevValueDifference <= ValueDifference)
@@ -126,13 +127,12 @@ bool URCRangeMapBehaviour::GetNearestActionByThreshold(TTuple<URCAction*, bool>&
 				break;
 			}
 		}
-		OutTuple = TTuple<URCAction*, bool>(NonLerpActions[NormalizedInputValue], ValueDifference <= UE::RCRangeMapBehaviour::Threshold);
+		OutTuple = TTuple<URCAction*, bool>(NonLerpActions[InputValue], ValueDifference <= UE::RCRangeMapBehaviour::Threshold);
 	}
 	
 	return true;
 }
 
-
 bool FRCRangeMapInput::GetInputValue(double& OutValue) const
 {
 	if (InputProperty)
