[RemoteControl] Fix crash due to binding returning null despite being bound

diff --git a/Engine/Plugins/VirtualProduction/RemoteControl/Source/RemoteControl/Private/RemoteControlPreset.cpp b/Engine/Plugins/VirtualProduction/RemoteControl/Source/RemoteControl/Private/RemoteControlPreset.cpp
--- a/Engine/Plugins/VirtualProduction/RemoteControl/Source/RemoteControl/Private/RemoteControlPreset.cpp
+++ b/Engine/Plugins/VirtualProduction/RemoteControl/Source/RemoteControl/Private/RemoteControlPreset.cpp
@@ -875,7 +875,22 @@ TWeakPtr<FRemoteControlProperty> URemoteControlPreset::ExposeProperty(UObject* O
 
 		const FName FieldName = GetEntityName(*Args.Label, Object, FieldPath);
 
-		FRemoteControlProperty RCProperty{ this, Registry->GenerateUniqueLabel(FieldName), MoveTemp(FieldPath), { FindOrAddBinding(Object) } };
+		URemoteControlBinding* Binding = FindOrAddBinding(Object);
+		if (!Binding->Resolve())
+		{
+			// The binding exists, and the object exists at this path, but the binding can't resolve to it.
+			// Due to a bug (UE-187216), this can happen if:
+			// 1. The object was created, but not saved
+			// 2. The property was bound
+			// 3. The object was destroyed
+			// 4. Another object was created at the same path
+			// This can all happen without FSoftObjectPath::CurrentTag being invalidated (usually only done when assets load), meaning the binding's
+			// pointer cache won't update to the new object at that path.
+			// But we know the new object here, so just replace the bound object and use it going forward.
+			Binding->SetBoundObject(Object);
+		}
+
+		FRemoteControlProperty RCProperty{ this, Registry->GenerateUniqueLabel(FieldName), MoveTemp(FieldPath), { Binding } };
 
 		RCPropertyPtr = StaticCastSharedPtr<FRemoteControlProperty>(Expose(MoveTemp(RCProperty), FRemoteControlProperty::StaticStruct(), Args.GroupId));
 	}
