[RemoteControl] Fix incorrect property expose/unexpose events on undo/redo when exposure occurs outside a transaction or in a transient preset

diff --git a/Engine/Plugins/VirtualProduction/RemoteControl/Source/RemoteControl/Private/RemoteControlPreset.cpp b/Engine/Plugins/VirtualProduction/RemoteControl/Source/RemoteControl/Private/RemoteControlPreset.cpp
--- a/Engine/Plugins/VirtualProduction/RemoteControl/Source/RemoteControl/Private/RemoteControlPreset.cpp
+++ b/Engine/Plugins/VirtualProduction/RemoteControl/Source/RemoteControl/Private/RemoteControlPreset.cpp
@@ -218,13 +218,13 @@ FRemoteControlPresetGroup& FRemoteControlPresetLayout::CreateGroup(FName GroupNa
 	Owner->CacheLayoutData();
 	OnGroupAddedDelegate.Broadcast(Group);
 
-	FRCTransactionListenerHelper<FRemoteControlPresetGroup>(ERCTransaction::Undo, OnGroupDeleted(), Group);
-	FRCTransactionListenerHelper<const FRemoteControlPresetGroup&>(ERCTransaction::Redo, OnGroupAdded(), Group);
+	FRCTransactionListenerHelper<FRemoteControlPresetGroup>(ERCTransaction::Undo, Owner->GetPresetId(), OnGroupDeleted(), Group);
+	FRCTransactionListenerHelper<const FRemoteControlPresetGroup&>(ERCTransaction::Redo, Owner->GetPresetId(), OnGroupAdded(), Group);
 
 	Owner->OnPresetLayoutModified().Broadcast(Owner.Get());
 
-	FRCTransactionListenerHelper<URemoteControlPreset*>(ERCTransaction::Undo, Owner->OnPresetLayoutModified(), Owner.Get());
-	FRCTransactionListenerHelper<URemoteControlPreset*>(ERCTransaction::Redo, Owner->OnPresetLayoutModified(), Owner.Get());
+	FRCTransactionListenerHelper<URemoteControlPreset*>(ERCTransaction::Undo, Owner->GetPresetId(), Owner->OnPresetLayoutModified(), Owner.Get());
+	FRCTransactionListenerHelper<URemoteControlPreset*>(ERCTransaction::Redo, Owner->GetPresetId(), Owner->OnPresetLayoutModified(), Owner.Get());
 
 	return Group;
 }
@@ -287,7 +287,7 @@ void FRemoteControlPresetLayout::SwapGroups(FGuid OriginGroupId, FGuid TargetGro
 		TArray<FGuid> GroupIdsBeforeTransaction;
 		GroupIdsBeforeTransaction.Reserve(Groups.Num());
 		Algo::Transform(Groups, GroupIdsBeforeTransaction, &FRemoteControlPresetGroup::Id);
-		FRCTransactionListenerHelper<const TArray<FGuid>&>(ERCTransaction::Undo, OnGroupOrderChanged(), GroupIdsBeforeTransaction);
+		FRCTransactionListenerHelper<const TArray<FGuid>&>(ERCTransaction::Undo, Owner->GetPresetId(), OnGroupOrderChanged(), GroupIdsBeforeTransaction);
 
 		Groups.Insert(FRemoteControlPresetGroup{ *OriginGroup }, TargetGroupIndex);
 		Groups.Swap(TargetGroupIndex, OriginGroupIndex);
@@ -298,12 +298,12 @@ void FRemoteControlPresetLayout::SwapGroups(FGuid OriginGroupId, FGuid TargetGro
 		Algo::Transform(Groups, GroupIds, &FRemoteControlPresetGroup::Id);
 		OnGroupOrderChangedDelegate.Broadcast(GroupIds);
 
-		FRCTransactionListenerHelper<const TArray<FGuid>&>(ERCTransaction::Redo, OnGroupOrderChanged(), GroupIds);
+		FRCTransactionListenerHelper<const TArray<FGuid>&>(ERCTransaction::Redo, Owner->GetPresetId(), OnGroupOrderChanged(), GroupIds);
 
 		Owner->OnPresetLayoutModified().Broadcast(Owner.Get());
 
-		FRCTransactionListenerHelper<URemoteControlPreset*>(ERCTransaction::Undo, Owner->OnPresetLayoutModified(), Owner.Get());
-		FRCTransactionListenerHelper<URemoteControlPreset*>(ERCTransaction::Redo, Owner->OnPresetLayoutModified(), Owner.Get());
+		FRCTransactionListenerHelper<URemoteControlPreset*>(ERCTransaction::Undo, Owner->GetPresetId(), Owner->OnPresetLayoutModified(), Owner.Get());
+		FRCTransactionListenerHelper<URemoteControlPreset*>(ERCTransaction::Redo, Owner->GetPresetId(), Owner->OnPresetLayoutModified(), Owner.Get());
 	}
 }
 
@@ -339,7 +339,7 @@ void FRemoteControlPresetLayout::SwapFields(const FFieldSwapArgs& FieldSwapArgs)
 		// Here we don't want to trigger add/delete delegates since the fields just get moved around.
 		TArray<FGuid>& Fields = DragTargetGroup->AccessFields();
 
-		FRCTransactionListenerHelper<const FGuid&, const TArray<FGuid>&>(ERCTransaction::Undo, OnFieldOrderChanged(), FieldSwapArgs.TargetGroupId, Fields);
+		FRCTransactionListenerHelper<const FGuid&, const TArray<FGuid>&>(ERCTransaction::Undo, Owner->GetPresetId(), OnFieldOrderChanged(), FieldSwapArgs.TargetGroupId, Fields);
 
 		Fields.Insert(FieldSwapArgs.DraggedFieldId, DragTargetFieldIndex);
 		Fields.Swap(DragTargetFieldIndex, DragOriginFieldIndex);
@@ -349,10 +349,10 @@ void FRemoteControlPresetLayout::SwapFields(const FFieldSwapArgs& FieldSwapArgs)
 		OnFieldOrderChangedDelegate.Broadcast(FieldSwapArgs.TargetGroupId, Fields);
 		Owner->OnPresetLayoutModified().Broadcast(Owner.Get());
 
-		FRCTransactionListenerHelper<const FGuid&, const TArray<FGuid>&>(ERCTransaction::Redo, OnFieldOrderChanged(), FieldSwapArgs.TargetGroupId, Fields);
+		FRCTransactionListenerHelper<const FGuid&, const TArray<FGuid>&>(ERCTransaction::Redo, Owner->GetPresetId(), OnFieldOrderChanged(), FieldSwapArgs.TargetGroupId, Fields);
 
-		FRCTransactionListenerHelper<URemoteControlPreset*>(ERCTransaction::Undo, Owner->OnPresetLayoutModified(), Owner.Get());
-		FRCTransactionListenerHelper<URemoteControlPreset*>(ERCTransaction::Redo, Owner->OnPresetLayoutModified(), Owner.Get());
+		FRCTransactionListenerHelper<URemoteControlPreset*>(ERCTransaction::Undo, Owner->GetPresetId(), Owner->OnPresetLayoutModified(), Owner.Get());
+		FRCTransactionListenerHelper<URemoteControlPreset*>(ERCTransaction::Redo, Owner->GetPresetId(), Owner->OnPresetLayoutModified(), Owner.Get());
 	}
 	else
 	{
@@ -361,19 +361,19 @@ void FRemoteControlPresetLayout::SwapFields(const FFieldSwapArgs& FieldSwapArgs)
 		DragTargetFieldIndex = DragTargetFieldIndex == INDEX_NONE ? 0 : DragTargetFieldIndex;
 		DragTargetGroup->AccessFields().Insert(FieldSwapArgs.DraggedFieldId, DragTargetFieldIndex);
 
-		FRCTransactionListenerHelper<const FGuid&, const FGuid&, int32>(ERCTransaction::Undo, OnFieldDeleted(), FieldSwapArgs.TargetGroupId, FieldSwapArgs.DraggedFieldId, DragTargetFieldIndex);
-		FRCTransactionListenerHelper<const FGuid&, const FGuid&, int32>(ERCTransaction::Undo, OnFieldAdded(), FieldSwapArgs.OriginGroupId, FieldSwapArgs.DraggedFieldId, DragOriginFieldIndex);
+		FRCTransactionListenerHelper<const FGuid&, const FGuid&, int32>(ERCTransaction::Undo, Owner->GetPresetId(), OnFieldDeleted(), FieldSwapArgs.TargetGroupId, FieldSwapArgs.DraggedFieldId, DragTargetFieldIndex);
+		FRCTransactionListenerHelper<const FGuid&, const FGuid&, int32>(ERCTransaction::Undo, Owner->GetPresetId(), OnFieldAdded(), FieldSwapArgs.OriginGroupId, FieldSwapArgs.DraggedFieldId, DragOriginFieldIndex);
 
-		FRCTransactionListenerHelper<const FGuid&, const FGuid&, int32>(ERCTransaction::Redo, OnFieldDeleted(), FieldSwapArgs.OriginGroupId, FieldSwapArgs.DraggedFieldId, DragOriginFieldIndex);
-		FRCTransactionListenerHelper<const FGuid&, const FGuid&, int32>(ERCTransaction::Redo, OnFieldAdded(), FieldSwapArgs.TargetGroupId, FieldSwapArgs.DraggedFieldId, DragTargetFieldIndex);
+		FRCTransactionListenerHelper<const FGuid&, const FGuid&, int32>(ERCTransaction::Redo, Owner->GetPresetId(), OnFieldDeleted(), FieldSwapArgs.OriginGroupId, FieldSwapArgs.DraggedFieldId, DragOriginFieldIndex);
+		FRCTransactionListenerHelper<const FGuid&, const FGuid&, int32>(ERCTransaction::Redo, Owner->GetPresetId(), OnFieldAdded(), FieldSwapArgs.TargetGroupId, FieldSwapArgs.DraggedFieldId, DragTargetFieldIndex);
 
 		Owner->CacheLayoutData();
 		OnFieldDeletedDelegate.Broadcast(FieldSwapArgs.OriginGroupId, FieldSwapArgs.DraggedFieldId, DragOriginFieldIndex);
 		OnFieldAddedDelegate.Broadcast(FieldSwapArgs.TargetGroupId, FieldSwapArgs.DraggedFieldId, DragTargetFieldIndex);
 		Owner->OnPresetLayoutModified().Broadcast(Owner.Get());
 
-		FRCTransactionListenerHelper<URemoteControlPreset*>(ERCTransaction::Undo, Owner->OnPresetLayoutModified(), Owner.Get());
-		FRCTransactionListenerHelper<URemoteControlPreset*>(ERCTransaction::Redo, Owner->OnPresetLayoutModified(), Owner.Get());
+		FRCTransactionListenerHelper<URemoteControlPreset*>(ERCTransaction::Undo, Owner->GetPresetId(), Owner->OnPresetLayoutModified(), Owner.Get());
+		FRCTransactionListenerHelper<URemoteControlPreset*>(ERCTransaction::Redo, Owner->GetPresetId(), Owner->OnPresetLayoutModified(), Owner.Get());
 	}
 }
 
@@ -393,14 +393,14 @@ void FRemoteControlPresetLayout::DeleteGroup(FGuid GroupId)
 		
 		Owner->CacheLayoutData();
 
-		FRCTransactionListenerHelper<const FRemoteControlPresetGroup&>(ERCTransaction::Undo, OnGroupAdded(), CopyTemp(DeletedGroup));
-		FRCTransactionListenerHelper<FRemoteControlPresetGroup>(ERCTransaction::Redo, OnGroupDeleted(), CopyTemp(DeletedGroup));
+		FRCTransactionListenerHelper<const FRemoteControlPresetGroup&>(ERCTransaction::Undo, Owner->GetPresetId(), OnGroupAdded(), CopyTemp(DeletedGroup));
+		FRCTransactionListenerHelper<FRemoteControlPresetGroup>(ERCTransaction::Redo, Owner->GetPresetId(), OnGroupDeleted(), CopyTemp(DeletedGroup));
 
 		OnGroupDeletedDelegate.Broadcast(MoveTemp(DeletedGroup));
 		Owner->OnPresetLayoutModified().Broadcast(Owner.Get());
 
-		FRCTransactionListenerHelper<URemoteControlPreset*>(ERCTransaction::Undo, Owner->OnPresetLayoutModified(), Owner.Get());
-		FRCTransactionListenerHelper<URemoteControlPreset*>(ERCTransaction::Redo, Owner->OnPresetLayoutModified(), Owner.Get());
+		FRCTransactionListenerHelper<URemoteControlPreset*>(ERCTransaction::Undo, Owner->GetPresetId(), Owner->OnPresetLayoutModified(), Owner.Get());
+		FRCTransactionListenerHelper<URemoteControlPreset*>(ERCTransaction::Redo, Owner->GetPresetId(), Owner->OnPresetLayoutModified(), Owner.Get());
 	}
 }
 
@@ -408,15 +408,15 @@ void FRemoteControlPresetLayout::RenameGroup(FGuid GroupId, FName NewGroupName)
 {
 	if (FRemoteControlPresetGroup* Group = GetGroup(GroupId))
 	{
-		FRCTransactionListenerHelper<const FGuid&, FName>(ERCTransaction::Undo, OnGroupRenamed(), GroupId, Group->Name);
-		FRCTransactionListenerHelper<const FGuid&, FName>(ERCTransaction::Redo, OnGroupRenamed(), GroupId, NewGroupName);
+		FRCTransactionListenerHelper<const FGuid&, FName>(ERCTransaction::Undo, Owner->GetPresetId(), OnGroupRenamed(), GroupId, Group->Name);
+		FRCTransactionListenerHelper<const FGuid&, FName>(ERCTransaction::Redo, Owner->GetPresetId(), OnGroupRenamed(), GroupId, NewGroupName);
 
 		Group->Name = NewGroupName;
 		OnGroupRenamedDelegate.Broadcast(GroupId, NewGroupName);
 		Owner->OnPresetLayoutModified().Broadcast(Owner.Get());
 
-		FRCTransactionListenerHelper<URemoteControlPreset*>(ERCTransaction::Undo, Owner->OnPresetLayoutModified(), Owner.Get());
-		FRCTransactionListenerHelper<URemoteControlPreset*>(ERCTransaction::Redo, Owner->OnPresetLayoutModified(), Owner.Get());
+		FRCTransactionListenerHelper<URemoteControlPreset*>(ERCTransaction::Undo, Owner->GetPresetId(), Owner->OnPresetLayoutModified(), Owner.Get());
+		FRCTransactionListenerHelper<URemoteControlPreset*>(ERCTransaction::Redo, Owner->GetPresetId(), Owner->OnPresetLayoutModified(), Owner.Get());
 	}
 }
 
@@ -437,13 +437,13 @@ void FRemoteControlPresetLayout::AddField(FGuid GroupId, FGuid FieldId)
 		Group->AccessFields().Add(FieldId);
 		OnFieldAddedDelegate.Broadcast(GroupId, FieldId, Group->AccessFields().Num() - 1);
 
-		FRCTransactionListenerHelper<const FGuid&, const FGuid&, int32>(ERCTransaction::Undo, OnFieldDeleted(), GroupId, FieldId, Group->AccessFields().Num() - 1);
-		FRCTransactionListenerHelper<const FGuid&, const FGuid&, int32>(ERCTransaction::Redo, OnFieldAdded(), GroupId, FieldId, Group->AccessFields().Num() - 1);
+		FRCTransactionListenerHelper<const FGuid&, const FGuid&, int32>(ERCTransaction::Undo, Owner->GetPresetId(), OnFieldDeleted(), GroupId, FieldId, Group->AccessFields().Num() - 1);
+		FRCTransactionListenerHelper<const FGuid&, const FGuid&, int32>(ERCTransaction::Redo, Owner->GetPresetId(), OnFieldAdded(), GroupId, FieldId, Group->AccessFields().Num() - 1);
 
 		Owner->OnPresetLayoutModified().Broadcast(Owner.Get());
 
-		FRCTransactionListenerHelper<URemoteControlPreset*>(ERCTransaction::Undo, Owner->OnPresetLayoutModified(), Owner.Get());
-		FRCTransactionListenerHelper<URemoteControlPreset*>(ERCTransaction::Redo, Owner->OnPresetLayoutModified(), Owner.Get());
+		FRCTransactionListenerHelper<URemoteControlPreset*>(ERCTransaction::Undo, Owner->GetPresetId(), Owner->OnPresetLayoutModified(), Owner.Get());
+		FRCTransactionListenerHelper<URemoteControlPreset*>(ERCTransaction::Redo, Owner->GetPresetId(), Owner->OnPresetLayoutModified(), Owner.Get());
 	}
 }
 
@@ -454,13 +454,13 @@ void FRemoteControlPresetLayout::InsertFieldAt(FGuid GroupId, FGuid FieldId, int
 		Group->AccessFields().Insert(FieldId, Index);
 		OnFieldAddedDelegate.Broadcast(GroupId, FieldId, Index);
 
-		FRCTransactionListenerHelper<const FGuid&, const FGuid&, int32>(ERCTransaction::Undo, OnFieldDeleted(), GroupId, FieldId, Index);
-		FRCTransactionListenerHelper<const FGuid&, const FGuid&, int32>(ERCTransaction::Redo, OnFieldAdded(), GroupId, FieldId, Index);
+		FRCTransactionListenerHelper<const FGuid&, const FGuid&, int32>(ERCTransaction::Undo, Owner->GetPresetId(), OnFieldDeleted(), GroupId, FieldId, Index);
+		FRCTransactionListenerHelper<const FGuid&, const FGuid&, int32>(ERCTransaction::Redo, Owner->GetPresetId(), OnFieldAdded(), GroupId, FieldId, Index);
 
 		Owner->OnPresetLayoutModified().Broadcast(Owner.Get());
 
-		FRCTransactionListenerHelper<URemoteControlPreset*>(ERCTransaction::Undo, Owner->OnPresetLayoutModified(), Owner.Get());
-		FRCTransactionListenerHelper<URemoteControlPreset*>(ERCTransaction::Redo, Owner->OnPresetLayoutModified(), Owner.Get());
+		FRCTransactionListenerHelper<URemoteControlPreset*>(ERCTransaction::Undo, Owner->GetPresetId(), Owner->OnPresetLayoutModified(), Owner.Get());
+		FRCTransactionListenerHelper<URemoteControlPreset*>(ERCTransaction::Redo, Owner->GetPresetId(), Owner->OnPresetLayoutModified(), Owner.Get());
 	}
 }
 
@@ -489,13 +489,13 @@ void FRemoteControlPresetLayout::RemoveFieldAt(FGuid GroupId, int32 Index)
 
 		OnFieldDeletedDelegate.Broadcast(GroupId, FieldId, Index);
 
-		FRCTransactionListenerHelper<const FGuid&, const FGuid&, int32>(ERCTransaction::Undo, OnFieldAdded(), GroupId, FieldId, Index);
-		FRCTransactionListenerHelper<const FGuid&, const FGuid&, int32>(ERCTransaction::Redo, OnFieldDeleted(), GroupId, FieldId, Index);
+		FRCTransactionListenerHelper<const FGuid&, const FGuid&, int32>(ERCTransaction::Undo, Owner->GetPresetId(), OnFieldAdded(), GroupId, FieldId, Index);
+		FRCTransactionListenerHelper<const FGuid&, const FGuid&, int32>(ERCTransaction::Redo, Owner->GetPresetId(), OnFieldDeleted(), GroupId, FieldId, Index);
 
 		Owner->OnPresetLayoutModified().Broadcast(Owner.Get());
 
-		FRCTransactionListenerHelper<URemoteControlPreset*>(ERCTransaction::Undo, Owner->OnPresetLayoutModified(), Owner.Get());
-		FRCTransactionListenerHelper<URemoteControlPreset*>(ERCTransaction::Redo, Owner->OnPresetLayoutModified(), Owner.Get());
+		FRCTransactionListenerHelper<URemoteControlPreset*>(ERCTransaction::Undo, Owner->GetPresetId(), Owner->OnPresetLayoutModified(), Owner.Get());
+		FRCTransactionListenerHelper<URemoteControlPreset*>(ERCTransaction::Redo, Owner->GetPresetId(), Owner->OnPresetLayoutModified(), Owner.Get());
 	}
 }
 
@@ -713,8 +713,8 @@ URCVirtualPropertyInContainer* URemoteControlPreset::AddController(TSubclassOf<U
 	{
 		OnControllerAdded().Broadcast(this, NewPropertyName, NewController->Id);
 
-		FRCTransactionListenerHelper<URemoteControlPreset*, const FGuid&>(ERCTransaction::Undo, OnControllerRemoved(), this, NewController->Id);
-		FRCTransactionListenerHelper<URemoteControlPreset*, const FName, const FGuid&>(ERCTransaction::Redo, OnControllerAdded(), this, NewPropertyName, NewController->Id);
+		FRCTransactionListenerHelper<URemoteControlPreset*, const FGuid&>(ERCTransaction::Undo, GetPresetId(), OnControllerRemoved(), this, NewController->Id);
+		FRCTransactionListenerHelper<URemoteControlPreset*, const FName, const FGuid&>(ERCTransaction::Redo, GetPresetId(), OnControllerAdded(), this, NewPropertyName, NewController->Id);
 
 		InitializeEntityMetadata(NewController);
 	}
@@ -735,8 +735,8 @@ bool URemoteControlPreset::RemoveController(const FName& InPropertyName)
 			OnControllerRemoved().Broadcast(this, ControllerToDelete->Id);
 			ControllerContainer->RemoveProperty(ControllerToDelete->GetPropertyName());
 
-			FRCTransactionListenerHelper<URemoteControlPreset*, const FName, const FGuid&>(ERCTransaction::Undo, OnControllerAdded(), this, ControllerToDelete->DisplayName, ControllerToDelete->Id);
-			FRCTransactionListenerHelper<URemoteControlPreset*, const FGuid&>(ERCTransaction::Redo, OnControllerRemoved(), this, ControllerToDelete->Id);
+			FRCTransactionListenerHelper<URemoteControlPreset*, const FName, const FGuid&>(ERCTransaction::Undo, GetPresetId(), OnControllerAdded(), this, ControllerToDelete->DisplayName, ControllerToDelete->Id);
+			FRCTransactionListenerHelper<URemoteControlPreset*, const FGuid&>(ERCTransaction::Redo, GetPresetId(), OnControllerRemoved(), this, ControllerToDelete->Id);
 
 			return true;
 		}
@@ -769,8 +769,8 @@ void URemoteControlPreset::ResetControllers()
 		{
 			OnControllerRemoved().Broadcast(this, Controller->Id);
 
-			FRCTransactionListenerHelper<URemoteControlPreset*, const FName, const FGuid&>(ERCTransaction::Undo, OnControllerAdded(), this, Controller->DisplayName, Controller->Id);
-			FRCTransactionListenerHelper<URemoteControlPreset*, const FGuid&>(ERCTransaction::Redo, OnControllerRemoved(), this, Controller->Id);
+			FRCTransactionListenerHelper<URemoteControlPreset*, const FName, const FGuid&>(ERCTransaction::Undo, GetPresetId(), OnControllerAdded(), this, Controller->DisplayName, Controller->Id);
+			FRCTransactionListenerHelper<URemoteControlPreset*, const FGuid&>(ERCTransaction::Redo, GetPresetId(), OnControllerRemoved(), this, Controller->Id);
 		}
 		ControllerContainer->Reset();
 	}
@@ -827,8 +827,8 @@ void URemoteControlPreset::OnModifyController(const FPropertyChangedEvent& Prope
 			ControllerContainer->OnModifyPropertyValue(PropertyChangedEvent);
 			OnControllerModified().Broadcast(this, {ModifiedController->Id});
 
-			FRCTransactionListenerHelper<URemoteControlPreset*, const TSet<FGuid>&>(ERCTransaction::Undo, OnControllerModified(), this, {ModifiedController->Id});
-			FRCTransactionListenerHelper<URemoteControlPreset*, const TSet<FGuid>&>(ERCTransaction::Redo, OnControllerModified(), this, {ModifiedController->Id});
+			FRCTransactionListenerHelper<URemoteControlPreset*, const TSet<FGuid>&>(ERCTransaction::Undo, GetPresetId(), OnControllerModified(), this, {ModifiedController->Id});
+			FRCTransactionListenerHelper<URemoteControlPreset*, const TSet<FGuid>&>(ERCTransaction::Redo, GetPresetId(), OnControllerModified(), this, {ModifiedController->Id});
 		}
 	}	
 } 
@@ -985,8 +985,8 @@ TSharedPtr<FRemoteControlEntity> URemoteControlPreset::Expose(FRemoteControlEnti
 	
 	OnEntityExposed().Broadcast(this, RCEntity->GetId());
 
-	FRCTransactionListenerHelper<URemoteControlPreset*, const FGuid&>(ERCTransaction::Undo, OnEntityUnexposed(), this, RCEntity->GetId());
-	FRCTransactionListenerHelper<URemoteControlPreset*, const FGuid&>(ERCTransaction::Redo, OnEntityExposed(), this, RCEntity->GetId());
+	FRCTransactionListenerHelper<URemoteControlPreset*, const FGuid&>(ERCTransaction::Undo, GetPresetId(), OnEntityUnexposed(), this, RCEntity->GetId());
+	FRCTransactionListenerHelper<URemoteControlPreset*, const FGuid&>(ERCTransaction::Redo, GetPresetId(), OnEntityExposed(), this, RCEntity->GetId());
 
 	return RCEntity;
 }
@@ -1537,8 +1537,8 @@ void URemoteControlPreset::Unexpose(const FGuid& EntityId)
 {
 	if (EntityId.IsValid() && Registry->GetExposedEntity(EntityId).IsValid())
 	{
-		FRCTransactionListenerHelper<URemoteControlPreset*, const FGuid&>(ERCTransaction::Undo, OnEntityExposed(), this, EntityId);
-		FRCTransactionListenerHelper<URemoteControlPreset*, const FGuid&>(ERCTransaction::Redo, OnEntityUnexposed(), this, EntityId);
+		FRCTransactionListenerHelper<URemoteControlPreset*, const FGuid&>(ERCTransaction::Undo, GetPresetId(), OnEntityExposed(), this, EntityId);
+		FRCTransactionListenerHelper<URemoteControlPreset*, const FGuid&>(ERCTransaction::Redo, GetPresetId(), OnEntityUnexposed(), this, EntityId);
 
 		OnEntityUnexposedDelegate.Broadcast(this, EntityId);
 
@@ -1707,15 +1707,15 @@ void URemoteControlPreset::OnExpose(const FExposeInfo& Info)
 		Layout.AddField(FieldGroupId, Info.FieldId);
 		OnEntityExposed().Broadcast(this, Info.FieldId);
 
-		FRCTransactionListenerHelper<URemoteControlPreset*, const FGuid&>(ERCTransaction::Undo, OnEntityUnexposed(), this, Info.FieldId);
-		FRCTransactionListenerHelper<URemoteControlPreset*, const FGuid&>(ERCTransaction::Redo, OnEntityExposed(), this, Info.FieldId);
+		FRCTransactionListenerHelper<URemoteControlPreset*, const FGuid&>(ERCTransaction::Undo, GetPresetId(), OnEntityUnexposed(), this, Info.FieldId);
+		FRCTransactionListenerHelper<URemoteControlPreset*, const FGuid&>(ERCTransaction::Redo, GetPresetId(), OnEntityExposed(), this, Info.FieldId);
 	}
 }
 
 void URemoteControlPreset::OnUnexpose(FGuid UnexposedFieldId)
  {
-	FRCTransactionListenerHelper<URemoteControlPreset*, const FGuid&>(ERCTransaction::Undo, OnEntityExposed(), this, UnexposedFieldId);
-	FRCTransactionListenerHelper<URemoteControlPreset*, const FGuid&>(ERCTransaction::Redo, OnEntityUnexposed(), this, UnexposedFieldId);
+	FRCTransactionListenerHelper<URemoteControlPreset*, const FGuid&>(ERCTransaction::Undo, GetPresetId(), OnEntityExposed(), this, UnexposedFieldId);
+	FRCTransactionListenerHelper<URemoteControlPreset*, const FGuid&>(ERCTransaction::Redo, GetPresetId(), OnEntityUnexposed(), this, UnexposedFieldId);
 
 	OnEntityUnexposed().Broadcast(this, UnexposedFieldId);
 	
diff --git a/Engine/Plugins/VirtualProduction/RemoteControl/Source/RemoteControl/Private/RemoteControlTransactionListenerHelper.h b/Engine/Plugins/VirtualProduction/RemoteControl/Source/RemoteControl/Private/RemoteControlTransactionListenerHelper.h
--- a/Engine/Plugins/VirtualProduction/RemoteControl/Source/RemoteControl/Private/RemoteControlTransactionListenerHelper.h
+++ b/Engine/Plugins/VirtualProduction/RemoteControl/Source/RemoteControl/Private/RemoteControlTransactionListenerHelper.h
@@ -2,8 +2,8 @@
 #pragma once
 
 #include "CoreMinimal.h"
+#include "IRemoteControlModule.h"
 #include "RemoteControlTransactionListenerManager.h"
-
 /**
  * Helper class which create a Listener and Register to the Manager
  * @tparam CallbackTypeArgs Callback Args type
@@ -18,13 +18,17 @@ public:
 	/**
 	 * Create and Register a Listener to the ListenerManager
 	 * @param InType When the broadcast should be called (Undo or Redo), see ERCTransactionType
+	 * @param PresetGuid Preset Guid to check whether or not is Transient
 	 * @param InDelegateToBroadcast Multicast delegate to broadcast
 	 * @param InOrderedArgs Ordered arguments that should be passed to the MulticastDelegate
 	 */
-	FRCTransactionListenerHelper(ERCTransaction::Type InType, TMulticastDelegate<void(CallbackTypeArgs...)> InDelegateToBroadcast, CallbackTypeArgs... InOrderedArgs)
+	FRCTransactionListenerHelper(ERCTransaction::Type InType, const FGuid& PresetGuid, TMulticastDelegate<void(CallbackTypeArgs...)> InDelegateToBroadcast, CallbackTypeArgs... InOrderedArgs)
 	{
 #if WITH_EDITOR
-		FRCTransactionListenerManager<CallbackTypeArgs...>::CreateListenerAndRegister(InType, InDelegateToBroadcast, Forward<CallbackTypeArgs>(InOrderedArgs)...);
+		if (!IRemoteControlModule::Get().IsPresetTransient(PresetGuid))
+		{
+			FRCTransactionListenerManager<CallbackTypeArgs...>::CreateListenerAndRegister(InType, InDelegateToBroadcast, Forward<CallbackTypeArgs>(InOrderedArgs)...);
+		}
 #endif // WITH_EDITOR
 	}
 };
\ No newline at end of file
diff --git a/Engine/Plugins/VirtualProduction/RemoteControl/Source/RemoteControl/Private/RemoteControlTransactionListenerManager.h b/Engine/Plugins/VirtualProduction/RemoteControl/Source/RemoteControl/Private/RemoteControlTransactionListenerManager.h
--- a/Engine/Plugins/VirtualProduction/RemoteControl/Source/RemoteControl/Private/RemoteControlTransactionListenerManager.h
+++ b/Engine/Plugins/VirtualProduction/RemoteControl/Source/RemoteControl/Private/RemoteControlTransactionListenerManager.h
@@ -247,7 +247,11 @@ private:
 	 */
 	void OnTransactionChanged(const FTransactionContext& InTransactionContext, ETransactionStateEventType InTransactionState)
 	{
-		if (InTransactionState == ETransactionStateEventType::TransactionFinalized
+		if (InTransactionState == ETransactionStateEventType::TransactionStarted)
+		{
+			PendingListeners.Reset();
+		}
+		else if (InTransactionState == ETransactionStateEventType::TransactionFinalized
 			&& !Listeners.Contains(InTransactionContext.TransactionId))
 		{
 			if (!PendingListeners.IsEmpty())
