[Remote Control] Remote Control Logic - Adds Object controllers and Custom Texture controller

Adds support for some Object controllers:
- Static Mesh
- Material Interface
- Texture (as a custom controller)

Adds support for "custom" controllers, which can be considered as a specialized version of existing controller types
- Custom Texture controller is a specialized string controller loading a texture from a file path or from an asset, through a dedicated widget.

Known issues: transaction are not handled properly. It was difficult to differentiate between what's the expected RC handling of controllers and exposed properties, and the behavior resulting from undoing actions for these new controllers.

Further development of the Texture controller might include the creation of a dedicated struct e.g. FTextureSelector

diff --git a/Engine/Plugins/VirtualProduction/RemoteControl/Config/BaseRemoteControl.ini b/Engine/Plugins/VirtualProduction/RemoteControl/Config/BaseRemoteControl.ini
--- a/Engine/Plugins/VirtualProduction/RemoteControl/Config/BaseRemoteControl.ini
+++ b/Engine/Plugins/VirtualProduction/RemoteControl/Config/BaseRemoteControl.ini
@@ -1,6 +1,7 @@
 [CoreRedirects]
 +PropertyRedirects=(OldName="RemoteControlField.Label",NewName="RemoteControlEntity.Label")
 +PropertyRedirects=(OldName="RemoteControlField.Id",NewName="RemoteControlEntity.Id")
++PropertyRedirects=(OldName="/Script/RemoteControlLogic.RemoteControlLogicConfig.SupportedControllerObjectTypes",NewName="/Script/RemoteControlLogic.RemoteControlLogicConfig.SupportedControllerObjectClassPaths")
 
 [/Script/RemoteControlLogic.RemoteControlLogicConfig]
 +SupportedControllerTypes=EPropertyBagPropertyType::String
@@ -8,6 +9,10 @@
 +SupportedControllerTypes=EPropertyBagPropertyType::Float
 +SupportedControllerTypes=EPropertyBagPropertyType::Bool
 +SupportedControllerTypes=EPropertyBagPropertyType::Struct
++SupportedControllerTypes=EPropertyBagPropertyType::Object
 +SupportedControllerStructTypes=Vector
 +SupportedControllerStructTypes=Color
-+SupportedControllerStructTypes=Rotator
\ No newline at end of file
++SupportedControllerStructTypes=Rotator
++SupportedControllerObjectClassPaths=/Script/Engine.StaticMesh
++SupportedControllerObjectClassPaths=/Script/Engine.MaterialInterface
++SupportedControllerCustomTypes=Texture
diff --git a/Engine/Plugins/VirtualProduction/RemoteControl/Source/RemoteControl/Private/RCVirtualProperty.cpp b/Engine/Plugins/VirtualProduction/RemoteControl/Source/RemoteControl/Private/RCVirtualProperty.cpp
--- a/Engine/Plugins/VirtualProduction/RemoteControl/Source/RemoteControl/Private/RCVirtualProperty.cpp
+++ b/Engine/Plugins/VirtualProduction/RemoteControl/Source/RemoteControl/Private/RCVirtualProperty.cpp
@@ -350,6 +350,33 @@ struct FRCVirtualPropertyCastHelpers
 
 		return false;
 	}
+
+	static UObject* GetObjectValue(const URCVirtualPropertyBase* InVirtualProperty)
+	{
+		if (InVirtualProperty)
+		{
+			const FProperty* Property = InVirtualProperty->GetProperty();
+			if (const FObjectProperty* ObjectProperty = CastField<FObjectProperty>(Property))
+			{
+				const void* SrcAddress = Property->ContainerPtrToValuePtr<void>(InVirtualProperty->GetContainerPtr());
+				return ObjectProperty->GetObjectPropertyValue(SrcAddress);
+			}
+		}
+
+		return nullptr;
+	}
+
+	static bool SetObjectValue(URCVirtualPropertyBase* InVirtualProperty, UObject* InValue)
+	{
+		const FProperty* Property = InVirtualProperty->GetProperty();
+		if (const FObjectProperty* ObjectProperty = CastField<FObjectProperty>(Property))
+		{			
+			ObjectProperty->SetObjectPropertyValue(InVirtualProperty->GetContainerPtr(), InValue);
+			return true;
+		}
+
+		return false;
+	}
 };
 
 bool URCVirtualPropertyBase::GetValueBool(bool& OutBoolValue) const
@@ -460,6 +487,11 @@ bool URCVirtualPropertyBase::GetValueColor(FColor& OutColor) const
 	return FRCVirtualPropertyCastHelpers::GetStructValue<FColor>(this, TBaseStructure<FColor>::Get(), &OutColor);
 }
 
+UObject* URCVirtualPropertyBase::GetValueObject() const
+{
+	return FRCVirtualPropertyCastHelpers::GetObjectValue(this);
+}
+
 FString URCVirtualPropertyBase::GetDisplayValueAsString() const
 {
 	const FProperty* Property = GetProperty();
@@ -520,6 +552,13 @@ FString URCVirtualPropertyBase::GetDisplayValueAsString() const
 
 		return FString::Printf(TEXT("%lf"), Value);
 	}
+	else if (const FObjectProperty* ObjectProperty = CastField<FObjectProperty>(Property))
+	{
+		if (const UObject* Object = GetValueObject())
+		{
+			return Object->GetPathName();
+		}
+	}
 	else if (IsVectorType())
 	{
 		FVector Vector;
@@ -675,10 +714,22 @@ FName URCVirtualPropertyBase::GetVirtualPropertyTypeDisplayName(const EPropertyB
 		PropertyTypeString = TEXT("Integer");
 	}
 	
-	// For Structs we should derive the prefix from the subobject type (Vector / Color / Rotator / etc) rather than the property bag type (which is always "Struct")
+	// For Structs and Objects we should derive the prefix from the subobject type (Vector / Color / Rotator / etc)
+	// or Object type (StaticMesh, MaterialInterface/etc) rather than the property bag type (which is always "Struct" or "Object")
 	if (InValueTypeObject)
+	{
 		PropertyTypeString = InValueTypeObject->GetName();
 
+		if (PropertyTypeString == TEXT("StaticMesh"))
+		{
+			PropertyTypeString = TEXT("Mesh");
+		}
+		else if (PropertyTypeString == TEXT("MaterialInterface"))
+		{
+			PropertyTypeString = TEXT("Material");
+		}
+	}
+
 	return *PropertyTypeString;
 }
 
diff --git a/Engine/Plugins/VirtualProduction/RemoteControl/Source/RemoteControl/Private/RemoteControlPropertyHandle.cpp b/Engine/Plugins/VirtualProduction/RemoteControl/Source/RemoteControl/Private/RemoteControlPropertyHandle.cpp
--- a/Engine/Plugins/VirtualProduction/RemoteControl/Source/RemoteControl/Private/RemoteControlPropertyHandle.cpp
+++ b/Engine/Plugins/VirtualProduction/RemoteControl/Source/RemoteControl/Private/RemoteControlPropertyHandle.cpp
@@ -9,10 +9,20 @@
 #include "Algo/AllOf.h"
 #include "Backends/CborStructDeserializerBackend.h"
 #include "CborWriter.h"
+#include "Components/MeshComponent.h"
+#include "Components/StaticMeshComponent.h"
+#include "Engine/StaticMesh.h"
+#include "Materials/MaterialInterface.h"
 #include "Serialization/MemoryReader.h"
 #include "Serialization/MemoryWriter.h"
 #include "UObject/TextProperty.h"
 
+#if WITH_EDITOR
+#include "ScopedTransaction.h"
+#endif
+
+#define LOCTEXT_NAMESPACE "RemoteControlPropertyHandle"
+
 namespace
 {
 	/** Check if the given property is Array or Set or Map */
@@ -305,6 +315,10 @@ TSharedPtr<IRemoteControlPropertyHandle> FRemoteControlPropertyHandle::GetProper
 	{
 		PropertyHandle = MakeShared<FRemoteControlPropertyHandleMap>(InRCProperty, InProperty, InParentProperty, InParentFieldPath, InArrayIndex);
 	}
+	else if (FRemoteControlPropertyHandleObject::Supports(InProperty))
+	{
+		PropertyHandle = MakeShared<FRemoteControlPropertyHandleObject>(InRCProperty, InProperty, InParentProperty, InParentFieldPath, InArrayIndex);
+	}
 	else
 	{
 		// Untyped or doesn't support getting the property directly but the property is still valid(probably struct property)
@@ -617,6 +631,7 @@ IMPLEMENT_PROPERTY_VALUE(FRemoteControlPropertyHandleName)
 IMPLEMENT_PROPERTY_VALUE(FRemoteControlPropertyHandleArray)
 IMPLEMENT_PROPERTY_VALUE(FRemoteControlPropertyHandleSet)
 IMPLEMENT_PROPERTY_VALUE(FRemoteControlPropertyHandleMap)
+IMPLEMENT_PROPERTY_VALUE(FRemoteControlPropertyHandleObject)
 
 #undef IMPLEMENT_PROPERTY_VALUE
 
@@ -1761,4 +1776,198 @@ bool FRemoteControlPropertyHandleLinearColor::AreComponentsValid() const
 	{
 		return ColorComponentPtr.Pin().IsValid();
 	});
-}
\ No newline at end of file
+}
+
+bool FRemoteControlPropertyHandleObject::SetValue(UObject* InValue)
+{
+	if (!InValue)
+	{
+		return false;
+	}
+	
+	if (InValue->IsA(UStaticMesh::StaticClass()))
+	{
+		return SetStaticMeshValue(InValue);
+	}
+	else
+	{
+		return SetObjectValue(InValue);
+	}
+}
+
+bool FRemoteControlPropertyHandleObject::GetValue(UObject* OutValue) const
+{
+	if (const TSharedPtr<FRemoteControlProperty>& RemoteControlPropertyPtr = GetRCProperty())
+	{
+		if (const TSharedPtr<IRemoteControlPropertyHandle>& PropertyHandle = RemoteControlPropertyPtr->GetPropertyHandle())
+		{
+			return PropertyHandle->GetValue(OutValue);
+		}
+	}
+
+	return false;
+}
+
+bool FRemoteControlPropertyHandleObject::SetValueInArray(UObject* InValue, int32 InIndex)
+{
+	// Override materials e.g for Static Mesh are still considered array elements, so we need to
+	// address setting their value as array elements, even if the exposed property looks like a material
+	if (InValue->IsA(UMaterialInterface::StaticClass()))
+	{
+		return SetMaterialValue(InValue, InIndex);
+	}
+
+	return false;
+}
+
+bool FRemoteControlPropertyHandleObject::GetValueInArray(UObject* OutValue, int32 InIndex) const
+{
+	if (InIndex < 0)
+	{
+		return false;
+	}
+	
+	if (const TSharedPtr<FRemoteControlProperty>& RemoteControlPropertyPtr = GetRCProperty())
+	{
+		if (const TSharedPtr<IRemoteControlPropertyHandle>& PropertyHandle = RemoteControlPropertyPtr->GetPropertyHandle())
+		{
+			if (const TSharedPtr<IRemoteControlPropertyHandleArray>& PropertyHandleArray = PropertyHandle->AsArray())
+			{
+				if (InIndex < PropertyHandleArray->GetNumElements())
+				{
+					if (const TSharedPtr<IRemoteControlPropertyHandle>& ElementPropertyHandle = PropertyHandleArray->GetElement(InIndex))
+					{
+						return ElementPropertyHandle->GetValue(OutValue);
+					}
+				}
+			}
+		}
+	}
+
+	return false;
+}
+
+bool FRemoteControlPropertyHandleObject::SetStaticMeshValue(UObject* InValue) const
+{
+	if (!InValue)
+	{
+		return false;
+	}
+
+	if (!InValue->IsA(UStaticMesh::StaticClass()))
+	{
+		return false;
+	}
+	
+	bool bSuccess = false;
+	
+	if (const TSharedPtr<FRemoteControlProperty>& RemoteControlPropertyPtr = GetRCProperty())
+	{
+		for (UObject* Object : RemoteControlPropertyPtr->GetBoundObjects())
+		{							
+			if (UStaticMeshComponent* StaticMeshComponent = Cast<UStaticMeshComponent>(Object))
+			{
+				bSuccess = StaticMeshComponent->SetStaticMesh(Cast<UStaticMesh>(InValue));
+			}			
+		}
+	}
+
+	return bSuccess;
+}
+
+bool FRemoteControlPropertyHandleObject::SetMaterialValue(UObject* InValue, int32 InIndex) const
+{
+	if (!InValue || InIndex < 0)
+	{
+		return false;
+	}
+	
+	if (!InValue->IsA(UMaterialInterface::StaticClass()))
+	{
+		return false;
+	}
+	
+	bool bSuccess = false;
+	
+	if (const TSharedPtr<FRemoteControlProperty>& RemoteControlPropertyPtr = GetRCProperty())
+	{		
+		for (UObject* Object : RemoteControlPropertyPtr->GetBoundObjects())
+		{
+			if (UMeshComponent* MaterialOwnerComponent = Cast<UMeshComponent>(Object))
+			{
+				if (MaterialOwnerComponent->GetMaterials().IsValidIndex(InIndex))
+				{
+					MaterialOwnerComponent->SetMaterial(InIndex, Cast<UMaterialInterface>(InValue));
+					bSuccess = true;
+				}
+			}
+		}
+	}
+
+	return bSuccess;
+}
+
+bool FRemoteControlPropertyHandleObject::SetObjectValue(UObject* InValue) const
+{
+	if (!InValue)
+	{
+		return false;
+	}
+
+	const TSharedPtr<FRemoteControlProperty>& RemoteControlPropertyPtr = GetRCProperty();
+	
+	FRCObjectReference ObjectRef;
+	ObjectRef.Property = RemoteControlPropertyPtr->GetProperty();
+	ObjectRef.Access = RemoteControlPropertyPtr->GetPropertyHandle()->ShouldGenerateTransaction() ? ERCAccess::WRITE_TRANSACTION_ACCESS : ERCAccess::WRITE_ACCESS;
+	ObjectRef.PropertyPathInfo = RemoteControlPropertyPtr->FieldPathInfo.ToString();
+
+#if WITH_EDITOR
+	FScopedTransaction Transaction(LOCTEXT("SetObjectValue", "Set Object Value"));
+#endif
+	
+	for (UObject* Object : RemoteControlPropertyPtr->GetBoundObjects())
+	{
+		if (IRemoteControlModule::Get().ResolveObjectProperty(ObjectRef.Access, Object, ObjectRef.PropertyPathInfo, ObjectRef))
+		{
+#if WITH_EDITOR
+			Object->Modify();
+#endif
+			/** Setting with Pointer and using Serialization of itself afterwards as a workaround with the Asset not updating in the world. */
+			FObjectPropertyBase* ObjectProperty = CastField<FObjectPropertyBase>(RemoteControlPropertyPtr->GetProperty());
+			if (!ObjectProperty)
+			{
+				return false;
+			}
+			
+			FProperty* Property = RemoteControlPropertyPtr->GetProperty();
+			uint8* ValueAddress = Property->ContainerPtrToValuePtr<uint8>(ObjectRef.ContainerAdress);
+			ObjectProperty->SetObjectPropertyValue(ValueAddress, InValue);
+
+			TArray<uint8> Buffer;
+			FMemoryReader Reader(Buffer);
+			FCborStructDeserializerBackend DeserializerBackend(Reader);
+
+			IRemoteControlModule::Get().SetObjectProperties(ObjectRef, DeserializerBackend, ERCPayloadType::Cbor, Buffer);
+			return true;
+		}
+	}
+	
+	return false;
+}
+
+bool FRemoteControlPropertyHandleObject::Supports(const FProperty* InProperty)
+{
+	if (!InProperty)
+	{
+		return false;
+	}
+	
+	if (const FObjectProperty* ObjectProperty = CastField<FObjectProperty>(InProperty))
+	{
+		return true;
+	}
+
+	return false;
+}
+
+#undef LOCTEXT_NAMESPACE
diff --git a/Engine/Plugins/VirtualProduction/RemoteControl/Source/RemoteControl/Private/RemoteControlPropertyHandle.h b/Engine/Plugins/VirtualProduction/RemoteControl/Source/RemoteControl/Private/RemoteControlPropertyHandle.h
--- a/Engine/Plugins/VirtualProduction/RemoteControl/Source/RemoteControl/Private/RemoteControlPropertyHandle.h
+++ b/Engine/Plugins/VirtualProduction/RemoteControl/Source/RemoteControl/Private/RemoteControlPropertyHandle.h
@@ -94,6 +94,11 @@ public:
 
 	virtual bool SetValue(const TCHAR* InValue) override { return false; }
 
+	virtual bool SetValue(UObject* InValue) override { return false; }
+	virtual bool SetValueInArray(UObject* InValue, int32 InIndex) override { return false; }
+	virtual bool GetValue(UObject* OutValue) const override { return false; }
+	virtual bool GetValueInArray(UObject* OutValue, int32 InIndex) const override { return false; }
+
 	virtual TSharedPtr<IRemoteControlPropertyHandle> GetChildHandle(FName ChildName, bool bRecurse = true) const override;
 	virtual TSharedPtr<IRemoteControlPropertyHandle> GetChildHandle(int32 Index) const override;
 	virtual int32 GetNumChildren() const override { return Children.Num(); }
@@ -590,4 +595,38 @@ private:
 	TArray< TWeakPtr<IRemoteControlPropertyHandle> > ColorComponents;
 };
 
+/**
+* The object implementation of a property handle
+*/
+class FRemoteControlPropertyHandleObject : public FRemoteControlPropertyHandle
+{
+public:
+	DECLARE_PROPERTY_ACCESSOR_CONSTRUCTOR(FRemoteControlPropertyHandleObject)
+
+	/**
+	* Check if Object handle supports the given property
+	* @param InProperty Property to check
+	* @return true if this handle supports the given property
+	*/
+	static bool Supports(const FProperty* InProperty);
+
+	//~ Begin IRemoteControlPropertyHandle Interface
+	virtual bool SetValue(UObject* InValue) override;
+	virtual bool GetValue(UObject* OutValue) const override;
+
+	/** This is currently used for the specific case in which we are setting a material inside the override materials of e.g. a Static Mesh*/
+	virtual bool SetValueInArray(UObject* InValue, int32 InIndex) override;
+	virtual bool GetValueInArray(UObject* OutValue, int32 InIndex) const override;
+	//~ End IRemoteControlPropertyHandle Interface
+
+	/** Set a Static Mesh value using the handled Object */
+	bool SetStaticMeshValue(UObject* InValue) const;
+
+	/** Set a Material value using the handled Object */
+	bool SetMaterialValue(UObject* InValue, int32 InIndex) const;
+
+	/** Set Object Value */
+	bool SetObjectValue(UObject* InValue) const;
+};
+
 #undef DECLARE_PROPERTY_ACCESSOR_CONSTRUCTOR
diff --git a/Engine/Plugins/VirtualProduction/RemoteControl/Source/RemoteControl/Public/IRemoteControlPropertyHandle.h b/Engine/Plugins/VirtualProduction/RemoteControl/Source/RemoteControl/Public/IRemoteControlPropertyHandle.h
--- a/Engine/Plugins/VirtualProduction/RemoteControl/Source/RemoteControl/Public/IRemoteControlPropertyHandle.h
+++ b/Engine/Plugins/VirtualProduction/RemoteControl/Source/RemoteControl/Public/IRemoteControlPropertyHandle.h
@@ -69,6 +69,8 @@ public:
 	virtual bool GetValue(FRotator& OutValue) const = 0;
 	virtual bool GetValue(FColor& OutValue) const = 0;
 	virtual bool GetValue(FLinearColor& OutValue) const = 0;
+	virtual bool GetValue(UObject* OutValue) const = 0;
+	virtual bool GetValueInArray(UObject* OutValue, int32 InIndex) const = 0;
 	
 	/**
 	 * Sets the typed value of a property.
@@ -100,6 +102,8 @@ public:
 	virtual bool SetValue(FColor InValue) = 0;
 	virtual bool SetValue(FLinearColor InValue) = 0;
 	virtual bool SetValue(const TCHAR* InValue) = 0;
+	virtual bool SetValue(UObject* InValue) = 0;
+	virtual bool SetValueInArray(UObject* InValue, int32 InIndex) = 0;
 
 	/**
 	 * Gets a child handle of this handle.  Useful for accessing properties in structs.
diff --git a/Engine/Plugins/VirtualProduction/RemoteControl/Source/RemoteControl/Public/RCVirtualProperty.h b/Engine/Plugins/VirtualProduction/RemoteControl/Source/RemoteControl/Public/RCVirtualProperty.h
--- a/Engine/Plugins/VirtualProduction/RemoteControl/Source/RemoteControl/Public/RCVirtualProperty.h
+++ b/Engine/Plugins/VirtualProduction/RemoteControl/Source/RemoteControl/Public/RCVirtualProperty.h
@@ -203,7 +203,11 @@ public:
 
 	/** Get Color value from Virtual Property */
 	UFUNCTION(BlueprintCallable, Category = "Remote Control Behaviour")
-	bool GetValueColor(FColor& OutColor) const;	
+	bool GetValueColor(FColor& OutColor) const;
+
+	/** Get Object value from Virtual Property */
+	UFUNCTION(BlueprintCallable, Category = "Remote Control Behaviour")
+	UObject* GetValueObject() const;
 
 	/** Infers correct type internally, fetches value from memory and returns the value as a string 
 	* that can be immediately used for dispaly (without needing to create a generic readonly property widget)
diff --git a/Engine/Plugins/VirtualProduction/RemoteControl/Source/RemoteControlLogic/Private/Action/Bind/RCCustomBindActionUtilities.cpp b/Engine/Plugins/VirtualProduction/RemoteControl/Source/RemoteControlLogic/Private/Action/Bind/RCCustomBindActionUtilities.cpp
new file mode 100644
--- /dev/null
+++ b/Engine/Plugins/VirtualProduction/RemoteControl/Source/RemoteControlLogic/Private/Action/Bind/RCCustomBindActionUtilities.cpp
@@ -0,0 +1,43 @@
+﻿// Copyright Epic Games, Inc. All Rights Reserved.
+
+#include "Action/Bind/RCCustomBindActionUtilities.h"
+#include "AssetRegistry/AssetData.h"
+#include "Engine/Texture2D.h"
+#include "IRemoteControlPropertyHandle.h"
+#include "Kismet/KismetRenderingLibrary.h"
+#include "RemoteControlField.h"
+
+UTexture2D* UE::RCCustomBindActionUtilities::LoadTextureFromPath(const FString& InPath)
+{
+	if (UTexture2D* Texture = Cast<UTexture2D>(StaticLoadObject(UTexture2D::StaticClass(), nullptr, *InPath)))
+	{
+		return Texture;
+	}
+
+	if (FPaths::FileExists(InPath))
+	{	
+		return UKismetRenderingLibrary::ImportFileAsTexture2D(nullptr, InPath);
+	}
+
+	return nullptr;
+}
+
+void UE::RCCustomBindActionUtilities::SetTexturePropertyFromPath(const TSharedRef<FRemoteControlProperty>& InRemoteControlEntityAsProperty, const FString& InPath)
+{
+	if (UTexture2D* LoadedTexture = UE::RCCustomBindActionUtilities::LoadTextureFromPath(InPath))
+	{
+		const FProperty* RemoteControlProperty = InRemoteControlEntityAsProperty->GetProperty();
+		if (RemoteControlProperty == nullptr)
+		{
+			return;
+		}
+
+		const TSharedPtr<IRemoteControlPropertyHandle>& RemoteControlHandle = InRemoteControlEntityAsProperty->GetPropertyHandle();
+		if (!ensure(RemoteControlHandle))
+		{
+			return;
+		}
+			
+		RemoteControlHandle->SetValue(LoadedTexture);
+	}
+}
diff --git a/Engine/Plugins/VirtualProduction/RemoteControl/Source/RemoteControlLogic/Private/Action/Bind/RCPropertyBindAction.cpp b/Engine/Plugins/VirtualProduction/RemoteControl/Source/RemoteControlLogic/Private/Action/Bind/RCPropertyBindAction.cpp
--- a/Engine/Plugins/VirtualProduction/RemoteControl/Source/RemoteControlLogic/Private/Action/Bind/RCPropertyBindAction.cpp
+++ b/Engine/Plugins/VirtualProduction/RemoteControl/Source/RemoteControlLogic/Private/Action/Bind/RCPropertyBindAction.cpp
@@ -1,10 +1,11 @@
 // Copyright Epic Games, Inc. All Rights Reserved.
 
 #include "Action/Bind/RCPropertyBindAction.h"
-
+#include "Action/Bind/RCCustomBindActionUtilities.h"
 #include "Action/RCAction.h"
 #include "Components/SceneComponent.h"
 #include "Controller/RCController.h"
+#include "Controller/RCCustomControllerUtilities.h"
 #include "IRemoteControlPropertyHandle.h"
 #include "RCVirtualProperty.h"
 #include "RemoteControlPreset.h"
@@ -68,6 +69,97 @@ static void SetStructPropertyFromController(const FProperty* RemoteControlProper
 	}
 }
 
+// Workaround method to extract the index of an Array Property element from its name
+// e.g. MaterialOverride[3] will return 3, MaterialOverride[12] will return 12, etc
+static int32 GetPropertyIndexInArray(const TSharedRef<FRemoteControlProperty>& InRemoteControlEntityAsProperty)
+{
+	FString FieldName = InRemoteControlEntityAsProperty->FieldName.ToString();
+
+	const TSharedPtr<IRemoteControlPropertyHandle>& RemoteControlHandle = InRemoteControlEntityAsProperty->GetPropertyHandle();
+	if (!ensure(RemoteControlHandle))
+	{
+		return -1;
+	}
+
+	const TSharedPtr<IRemoteControlPropertyHandleArray>& AsArray = RemoteControlHandle->AsArray();
+	const int32 ElementsNum = AsArray->GetNumElements();
+
+	FieldName.RemoveFromEnd(TEXT("]"));
+	
+	// Number of integer digits, so that we know how many characters to keep at the end of the string
+	const int32 Digits = FMath::Floor(FMath::LogX(10.0f, ElementsNum) + 1);
+	const int32 CharactersToRemove = FieldName.Len() - Digits;
+	FieldName.RemoveAt(0, CharactersToRemove, true);
+	const int32 Index = FCString::Atoi(*FieldName);
+
+	return Index;
+}
+
+static void SetObjectPropertyFromController(const TSharedRef<FRemoteControlProperty>& InRemoteControlEntityAsProperty, const URCController* Controller)
+{
+	FProperty* RemoteControlProperty = InRemoteControlEntityAsProperty->GetProperty();
+	if (RemoteControlProperty == nullptr)
+	{
+		return;
+	}
+
+	const TSharedPtr<IRemoteControlPropertyHandle>& RemoteControlHandle = InRemoteControlEntityAsProperty->GetPropertyHandle();
+	if (!ensure(RemoteControlHandle))
+	{
+		return;
+	}
+	
+	const FObjectProperty* SourceObjectProperty = CastField<FObjectProperty>(Controller->GetProperty());
+	const FObjectProperty* TargetObjectProperty = CastField<FObjectProperty>(RemoteControlProperty);
+	
+	if (!TargetObjectProperty && RemoteControlProperty->IsA(FArrayProperty::StaticClass()))
+	{
+		if (const FArrayProperty* TargetArrayProperty = CastField<FArrayProperty>(RemoteControlProperty))
+		{
+			if (FProperty* InnerProperty = TargetArrayProperty->Inner)
+			{
+				TargetObjectProperty = CastField<FObjectProperty>(InnerProperty);
+			}
+		}
+	}
+
+	if (ensure(SourceObjectProperty && TargetObjectProperty))
+	{
+		bool bApplyValue = false;
+		
+		// Pass 1 - Objects of exactly the same type
+		if (SourceObjectProperty->SameType(TargetObjectProperty))
+		{
+			bApplyValue = true;
+		}
+		// Pass 2 - Objects that support conversion to other (related) Objects
+		else
+		{
+			// Nothing for now
+		}
+
+		if (bApplyValue)
+		{
+			if (UObject* ObjectValue = Controller->GetValueObject())
+			{
+				if (RemoteControlProperty->IsA(FArrayProperty::StaticClass()))
+				{
+					// Extract index in array, based on field name
+					const int32 Index = GetPropertyIndexInArray(InRemoteControlEntityAsProperty);	
+					if (const TSharedPtr<IRemoteControlPropertyHandle>& ElementToUpdate = RemoteControlHandle->AsArray()->GetElement(Index))
+					{
+						ElementToUpdate->SetValueInArray(ObjectValue, Index);					
+					}
+				}
+				else
+				{
+					RemoteControlHandle->SetValue(ObjectValue);
+				}		
+			}
+		}
+	}
+}
+
 void URCPropertyBindAction::Execute() const
 {
 	if (!ensure(Controller))
@@ -124,6 +216,22 @@ void URCPropertyBindAction::Execute() const
 				{
 					Handle->SetValue((bool)NumericValue);
 				}
+				else // This might be a Custom Controller
+				{
+					const FString& CustomControllerType = UE::RCCustomControllers::GetCustomControllerTypeName(Controller);
+					
+					if (!CustomControllerType.IsEmpty())
+					{
+						// External Texture controller (internally a String controller)
+						if (CustomControllerType == UE::RCCustomControllers::CustomTextureControllerName)
+						{
+							if (RemoteControlProperty->IsA(FObjectProperty::StaticClass()))
+							{
+								UE::RCCustomBindActionUtilities::SetTexturePropertyFromPath(RemoteControlEntityAsProperty.ToSharedRef(), StringValue);
+							}
+						}
+					}
+				}
 			}
 			// String Controller (end)
 		}
@@ -210,6 +318,10 @@ void URCPropertyBindAction::Execute() const
 		{
 			SetStructPropertyFromController(RemoteControlProperty, Handle.ToSharedRef(), Controller);
 		}
+		else if (ControllerAsProperty->IsA(FObjectProperty::StaticClass()))
+		{
+			SetObjectPropertyFromController(RemoteControlEntityAsProperty.ToSharedRef(), Controller);
+		}
 
 		// Editor specific updates
 	#if WITH_EDITOR
@@ -232,4 +344,4 @@ void URCPropertyBindAction::Execute() const
 		}
 	#endif
 	}
-}
\ No newline at end of file
+}
diff --git a/Engine/Plugins/VirtualProduction/RemoteControl/Source/RemoteControlLogic/Private/Action/RCActionContainer.cpp b/Engine/Plugins/VirtualProduction/RemoteControl/Source/RemoteControlLogic/Private/Action/RCActionContainer.cpp
--- a/Engine/Plugins/VirtualProduction/RemoteControl/Source/RemoteControlLogic/Private/Action/RCActionContainer.cpp
+++ b/Engine/Plugins/VirtualProduction/RemoteControl/Source/RemoteControlLogic/Private/Action/RCActionContainer.cpp
@@ -2,14 +2,17 @@
 
 #include "Action/RCActionContainer.h"
 
+#include "Action/RCAction.h"
+#include "Action/RCFunctionAction.h"
+#include "Action/RCPropertyAction.h"
+#include "Behaviour/Builtin/Path/RCSetAssetByPathBehaviour.h"
 #include "Behaviour/RCBehaviour.h"
 #include "Controller/RCController.h"
+#include "Controller/RCCustomControllerUtilities.h"
 #include "IRemoteControlModule.h"
 #include "RCVirtualProperty.h"
 #include "RemoteControlField.h"
-#include "Action/RCAction.h"
-#include "Action/RCFunctionAction.h"
-#include "Action/RCPropertyAction.h"
+#include "RemoteControlPreset.h"
 
 void URCActionContainer::ExecuteActions()
 {
@@ -166,6 +169,51 @@ void URCActionContainer::PostEditUndo()
 
 	OnActionsListModified.Broadcast();
 }
+
+void URCActionContainer::ExecuteActionsOnLoad()
+{
+	// In some cases, external resources are referenced, so we need to execute Controllers actions while loading
+	if (URCBehaviour* ParentBehaviour = GetParentBehaviour())
+	{
+		if (URCController* Controller = ParentBehaviour->ControllerWeakPtr.Get())
+		{
+			if (UE::RCCustomControllers::CustomControllerExecutesOnLoad(Controller))
+			{
+				Controller->ExecuteBehaviours();
+				return;
+			}
+		}
+
+		if (URCSetAssetByPathBehaviour* AssetPathBehaviour = Cast<URCSetAssetByPathBehaviour>(ParentBehaviour))
+		{
+			if (!AssetPathBehaviour->bInternal)
+			{
+				// Target entity might not be set up yet, we need it to be set in order for execution to work
+				AssetPathBehaviour->UpdateTargetEntity();
+				ParentBehaviour->Execute();
+			}
+		}
+	}
+}
+
+void URCActionContainer::PostLoad()
+{
+	UObject::PostLoad();
+
+	TWeakObjectPtr<URCActionContainer> ThisWeak(this);
+
+	FTSTicker::GetCoreTicker().AddTicker(FTickerDelegate::CreateLambda(
+		[ThisWeak](float InDeltaTime)->bool
+		{
+			if (URCActionContainer* const This = ThisWeak.Get())
+			{
+				This->ExecuteActionsOnLoad();
+			}
+
+			// Return false for one time execution
+			return false;
+		}));
+}
 #endif
 
 URCAction* URCActionContainer::FindActionByFieldId(const FGuid InId) const
@@ -210,4 +258,4 @@ int32 URCActionContainer::RemoveAction(URCAction* InAction)
 void URCActionContainer::EmptyActions()
 {
 	Actions.Empty();
-}
\ No newline at end of file
+}
diff --git a/Engine/Plugins/VirtualProduction/RemoteControl/Source/RemoteControlLogic/Private/Behaviour/Builtin/Bind/RCBehaviourBind.cpp b/Engine/Plugins/VirtualProduction/RemoteControl/Source/RemoteControlLogic/Private/Behaviour/Builtin/Bind/RCBehaviourBind.cpp
--- a/Engine/Plugins/VirtualProduction/RemoteControl/Source/RemoteControlLogic/Private/Behaviour/Builtin/Bind/RCBehaviourBind.cpp
+++ b/Engine/Plugins/VirtualProduction/RemoteControl/Source/RemoteControlLogic/Private/Behaviour/Builtin/Bind/RCBehaviourBind.cpp
@@ -2,14 +2,14 @@
 
 #include "Behaviour/Builtin/Bind/RCBehaviourBind.h"
 
-#include "Action/RCAction.h"
-#include "Action/RCFunctionAction.h"
 #include "Action/Bind/RCPropertyBindAction.h"
-#include "Behaviour/Builtin/Bind/RCBehaviourBindNode.h"
+#include "Action/RCAction.h"
 #include "Controller/RCController.h"
+#include "Controller/RCCustomControllerUtilities.h"
+#include "Engine/Texture2D.h"
 #include "IRemoteControlPropertyHandle.h"
-#include "RemoteControlField.h"
 #include "RCVirtualProperty.h"
+#include "RemoteControlField.h"
 
 URCBehaviourBind::URCBehaviourBind()
 {
@@ -110,6 +110,36 @@ static bool EvaluateBindCompatibility(const TMap<FFieldClass*, TArray<FFieldClas
 	return false;
 }
 
+
+static bool EvaluateCustomControllerBindCompatibility(const URCController* InController,  const FFieldClass* InControllerPropertyClass, const FFieldClass* InRemoteControlPropertyClass, const FProperty* InRemoteControlProperty)
+{
+	if (!InController)
+	{
+		return false;
+	}
+	
+	const FString CustomName = UE::RCCustomControllers::GetCustomControllerTypeName(InController);
+
+	if (CustomName.IsEmpty())
+	{
+		return false;
+	}
+
+	// External Texture for image files (controller string, rc exposed property is object
+	if (CustomName == UE::RCCustomControllers::CustomTextureControllerName && InControllerPropertyClass == FStrProperty::StaticClass() && InRemoteControlPropertyClass == FObjectProperty::StaticClass())
+	{
+		const FObjectProperty* TargetObjectProperty = CastField<FObjectProperty>(InRemoteControlProperty);
+
+		// Target type is texture - we're good to go
+		if (TargetObjectProperty->PropertyClass == UTexture::StaticClass())
+		{
+			return true;
+		}
+	}
+
+	return false;
+}
+
 bool URCBehaviourBind::CanHaveActionForField(URCController* Controller, TSharedRef<const FRemoteControlField> InRemoteControlField, const bool bInAllowNumericInputAsStrings)
 {
 	// Bind behaviour is only relevant for Exposed Properties
@@ -162,6 +192,7 @@ bool URCBehaviourBind::CanHaveActionForField(URCController* Controller, TSharedR
 	const FFieldClass* ControllerPropertyClass = ControllerAsProperty->GetClass();
 	const FFieldClass* RemoteControlPropertyClass = RemoteControlProperty->GetClass();
 	const bool bIsStructController = ControllerAsProperty->IsA(FStructProperty::StaticClass());
+	const bool bIsObjectController = ControllerAsProperty->IsA(FObjectProperty::StaticClass());
 
 	// Direct Binding (same type)
 	if (ControllerPropertyClass == RemoteControlPropertyClass && !bIsStructController) // For Structs we need to check the inner component as well
@@ -169,7 +200,7 @@ bool URCBehaviourBind::CanHaveActionForField(URCController* Controller, TSharedR
 		return true;
 	}
 	// Indirect Binding (related types)
-	else if(EvaluateBindCompatibility(SupportedIndirectBindsMap, ControllerPropertyClass, RemoteControlPropertyClass))
+	else if (EvaluateBindCompatibility(SupportedIndirectBindsMap, ControllerPropertyClass, RemoteControlPropertyClass))
 	{
 		return true;
 	}
@@ -207,6 +238,34 @@ bool URCBehaviourBind::CanHaveActionForField(URCController* Controller, TSharedR
 		}
 	}
 
+	// Objects
+	if (bIsObjectController)
+	{
+		if (const FObjectProperty* ObjectProperty = CastField<FObjectProperty>(ControllerAsProperty))
+		{
+			if (const FObjectProperty* RCFieldObjectProperty = CastField<FObjectProperty>(RemoteControlProperty))
+			{
+				if (ObjectProperty->SameType(RCFieldObjectProperty))
+				{
+					return true;
+				}
+			}
+			else if (const FArrayProperty* RCFieldArrayProperty = CastField<FArrayProperty>(RemoteControlProperty))
+			{
+				if (ObjectProperty->SameType(RCFieldArrayProperty->Inner))
+				{
+					return true;
+				}
+			}
+		}
+	}
+
+	// at the end, let's check if this is a Custom controller (a default controller with additional metadata)
+	if (EvaluateCustomControllerBindCompatibility(Controller, ControllerPropertyClass, RemoteControlPropertyClass, RemoteControlProperty))
+	{
+		return true;
+	}
+
 	return false; // Unsupported types!
 }
 
@@ -220,7 +279,8 @@ bool URCBehaviourBind::GetPropertyBagTypeFromFieldProperty(const FProperty* InPr
 														   FUInt64Property::StaticClass(),   FUInt16Property::StaticClass(),   FEnumProperty::StaticClass(),    FByteProperty::StaticClass() } },
 		{ EPropertyBagPropertyType::Float,   /* --> */    { FFloatProperty::StaticClass(),    FDoubleProperty::StaticClass() } },
 		{ EPropertyBagPropertyType::Bool,    /* --> */    { FBoolProperty::StaticClass() } },
-		{ EPropertyBagPropertyType::Struct,  /* --> */    { FStructProperty::StaticClass() } }
+		{ EPropertyBagPropertyType::Struct,  /* --> */    { FStructProperty::StaticClass() }, },
+		{ EPropertyBagPropertyType::Object,  /* --> */    { FObjectProperty::StaticClass() } /* this will probably never occur, so we manually set the out bag type later */ }
 	};
 
 	OutPropertyBagType = EPropertyBagPropertyType::None;
@@ -241,6 +301,35 @@ bool URCBehaviourBind::GetPropertyBagTypeFromFieldProperty(const FProperty* InPr
 	{
 		OutStructObject = StructProperty->Struct;
 	}
+	// Extract the Object type if this is an Object Property
+	else if (const FObjectProperty* ObjectProperty = CastField<FObjectProperty>(InProperty))
+	{
+		if (UClass* ObjectClass = ObjectProperty->PropertyClass)
+		{
+			OutStructObject = ObjectClass;
+			OutPropertyBagType = EPropertyBagPropertyType::Object;
+		}
+	}
+	// Handle an array property, and then treat the inner property in case it's an Object Property
+	else if (const FArrayProperty* ArrayProperty = CastField<FArrayProperty>(InProperty))
+	{
+		if (FProperty* InnerProperty = ArrayProperty->Inner)
+		{
+			if (const FObjectProperty* InnerObjectProperty = CastField<FObjectProperty>(InnerProperty))
+			{
+				if (UClass* InnerObjectClass = InnerObjectProperty->PropertyClass)
+				{
+					OutStructObject = InnerObjectClass;
+					OutPropertyBagType = EPropertyBagPropertyType::Object;
+				}
+			}
+		}
+	}
+
+	if (OutStructObject == UTexture::StaticClass())
+	{
+		OutPropertyBagType = EPropertyBagPropertyType::String;
+	}
 
 	return OutPropertyBagType != EPropertyBagPropertyType::None;
 }
@@ -332,4 +421,4 @@ bool URCBehaviourBind::CopyPropertyValueToController(URCController* InController
 	}
 
 	return false;
-}
\ No newline at end of file
+}
diff --git a/Engine/Plugins/VirtualProduction/RemoteControl/Source/RemoteControlLogic/Private/Controller/RCCustomControllerUtilities.cpp b/Engine/Plugins/VirtualProduction/RemoteControl/Source/RemoteControlLogic/Private/Controller/RCCustomControllerUtilities.cpp
new file mode 100644
--- /dev/null
+++ b/Engine/Plugins/VirtualProduction/RemoteControl/Source/RemoteControlLogic/Private/Controller/RCCustomControllerUtilities.cpp
@@ -0,0 +1,102 @@
+﻿// Copyright Epic Games, Inc. All Rights Reserved.
+
+#include "Controller/RCCustomControllerUtilities.h"
+#include "RCVirtualProperty.h"
+#include "RCVirtualPropertyContainer.h"
+
+namespace UE::RCCustomControllers
+{
+	// A Map to associate a custom controller name to its underlying type
+	inline static TMap<FString, EPropertyBagPropertyType> CustomControllerTypes =
+		{{CustomTextureControllerName, EPropertyBagPropertyType::String}};
+
+	// A Map to associate a custom controller name to its underlying type
+	inline static TSet<FString> ExecuteOnLoadControllers =
+		{CustomTextureControllerName};
+}
+
+bool UE::RCCustomControllers::IsCustomController(const URCVirtualPropertyBase* InController)
+{
+	return !InController->GetMetadataValue(CustomControllerNameKey).IsEmpty();
+}
+
+FString UE::RCCustomControllers::GetCustomControllerTypeName(const URCVirtualPropertyBase* InController)
+{
+	return InController->GetMetadataValue(CustomControllerNameKey);
+}
+
+bool UE::RCCustomControllers::CustomControllerExecutesOnLoad(const URCVirtualPropertyBase* InController)
+{
+	return InController->GetMetadataValue(CustomControllerExecuteOnLoadNameKey).ToBool();
+}
+
+bool UE::RCCustomControllers::CustomControllerExecutesOnLoad(const FName& InCustomControllerTypeName)
+{
+	return ExecuteOnLoadControllers.Contains(InCustomControllerTypeName.ToString());
+}
+
+bool UE::RCCustomControllers::IsValidCustomController(const FName& InCustomControllerTypeName)
+{
+	return CustomControllerTypes.Contains(InCustomControllerTypeName.ToString());
+}
+
+EPropertyBagPropertyType UE::RCCustomControllers::GetCustomControllerType(const FName& InCustomControllerTypeName)
+{
+	if (IsValidCustomController(InCustomControllerTypeName))
+	{
+		return CustomControllerTypes[InCustomControllerTypeName.ToString()];
+	}
+
+	return EPropertyBagPropertyType::None;
+}
+
+EPropertyBagPropertyType UE::RCCustomControllers::GetCustomControllerType(const FString& InCustomControllerTypeName)
+{
+	return GetCustomControllerType(FName(InCustomControllerTypeName));
+}
+
+TMap<FName, FString> UE::RCCustomControllers::GetCustomControllerMetaData(const FString& InCustomControllerTypeName)
+{
+	TMap<FName, FString> OutMetaData;
+	OutMetaData.Emplace(UE::RCCustomControllers::CustomControllerNameKey, InCustomControllerTypeName);
+
+	if (UE::RCCustomControllers::CustomControllerExecutesOnLoad(FName(*InCustomControllerTypeName)))
+	{
+		OutMetaData.Emplace(UE::RCCustomControllers::CustomControllerExecuteOnLoadNameKey, TEXT("true"));
+	}
+
+	return OutMetaData;
+}
+
+FName UE::RCCustomControllers::GetUniqueNameForController(const URCVirtualPropertyInContainer* InController)
+{
+	const FString& CustomTypeName = UE::RCCustomControllers::GetCustomControllerTypeName(InController);
+	if (!CustomTypeName.IsEmpty())
+	{
+		if (InController->ContainerWeakPtr.IsValid())
+		{
+			URCVirtualPropertyContainerBase* Container = InController->ContainerWeakPtr.Get();
+								
+			int32 Index = 0;
+			const FString InitialName = CustomTypeName;
+			FString FinalName = InitialName;
+
+			for (const TObjectPtr<URCVirtualPropertyBase>& VirtualProperty : Container->VirtualProperties)
+			{
+				if (VirtualProperty->DisplayName == FinalName)
+				{
+					Index++;
+					if (Index > 0)
+					{
+						FinalName = InitialName + TEXT("_") + FString::FromInt(Index++);
+					}
+				}
+			}
+			
+			return FName(FinalName);				
+		}
+	}
+
+	// Worst case - we return the plain type name.
+	return FName(CustomTypeName);
+}
diff --git a/Engine/Plugins/VirtualProduction/RemoteControl/Source/RemoteControlLogic/Public/Action/Bind/RCCustomBindActionUtilities.h b/Engine/Plugins/VirtualProduction/RemoteControl/Source/RemoteControlLogic/Public/Action/Bind/RCCustomBindActionUtilities.h
new file mode 100644
--- /dev/null
+++ b/Engine/Plugins/VirtualProduction/RemoteControl/Source/RemoteControlLogic/Public/Action/Bind/RCCustomBindActionUtilities.h
@@ -0,0 +1,19 @@
+﻿// Copyright Epic Games, Inc. All Rights Reserved.
+
+#pragma once
+
+#include "Engine/Texture2D.h"
+#include "Templates/SharedPointer.h"
+
+struct FRemoteControlProperty;
+
+namespace UE::RCCustomBindActionUtilities
+{
+	REMOTECONTROLLOGIC_API UTexture2D* LoadTextureFromPath(const FString& InPath);
+	
+	/**
+	 * Loads the texture found at the specified path, and uses it to set a Texture value in the specified RC Property
+	 * Useful e.g. for Custom External Texture Controller
+	 */
+	REMOTECONTROLLOGIC_API void SetTexturePropertyFromPath(const TSharedRef<FRemoteControlProperty>& InRemoteControlEntityAsProperty, const FString& InPath);
+};
diff --git a/Engine/Plugins/VirtualProduction/RemoteControl/Source/RemoteControlLogic/Public/Action/RCActionContainer.h b/Engine/Plugins/VirtualProduction/RemoteControl/Source/RemoteControlLogic/Public/Action/RCActionContainer.h
--- a/Engine/Plugins/VirtualProduction/RemoteControl/Source/RemoteControlLogic/Public/Action/RCActionContainer.h
+++ b/Engine/Plugins/VirtualProduction/RemoteControl/Source/RemoteControlLogic/Public/Action/RCActionContainer.h
@@ -67,10 +67,14 @@ public:
 	/** Retrieves all the Property Actions present in this container */
 	TArray<const URCPropertyAction*> GetPropertyActions() const;
 
+	//~ Begin UObject
 #if WITH_EDITOR
 	/** Called after applying a transaction to the object. Used to broadcast Undo related container changes to UI */
 	virtual void PostEditUndo();
+	void ExecuteActionsOnLoad();
 #endif
+	virtual void PostLoad() override;
+	//~ End UObject
 
 	/** Set of child action container */
 	UPROPERTY()
diff --git a/Engine/Plugins/VirtualProduction/RemoteControl/Source/RemoteControlLogic/Public/Controller/RCCustomControllerUtilities.h b/Engine/Plugins/VirtualProduction/RemoteControl/Source/RemoteControlLogic/Public/Controller/RCCustomControllerUtilities.h
new file mode 100644
--- /dev/null
+++ b/Engine/Plugins/VirtualProduction/RemoteControl/Source/RemoteControlLogic/Public/Controller/RCCustomControllerUtilities.h
@@ -0,0 +1,42 @@
+﻿// Copyright Epic Games, Inc. All Rights Reserved.
+
+#pragma once
+#include "PropertyBag.h"
+
+class URCVirtualPropertyBase;
+class URCVirtualPropertyInContainer;
+
+namespace UE::RCCustomControllers
+{
+	/** Checks if the passed virtual property has Custom Controller metadata available */
+	REMOTECONTROLLOGIC_API bool IsCustomController(const URCVirtualPropertyBase* InController);
+
+	/** Checks if the specified Custom Controller Type Name exists (e.g. is not a made up type) */
+	REMOTECONTROLLOGIC_API bool IsValidCustomController(const FName& InCustomControllerTypeName);
+	
+	/** Returns the Custom Controller Type Name (e.g. Asset Path). Empty string if Controller is not a custom one. */
+	REMOTECONTROLLOGIC_API FString GetCustomControllerTypeName(const URCVirtualPropertyBase* InController);
+
+	REMOTECONTROLLOGIC_API bool CustomControllerExecutesOnLoad(const URCVirtualPropertyBase* InController);
+
+	REMOTECONTROLLOGIC_API bool CustomControllerExecutesOnLoad(const FName& InCustomControllerTypeName);
+	
+	/** Returns the type associated to the specified Custom Controller Type Name */
+	REMOTECONTROLLOGIC_API EPropertyBagPropertyType GetCustomControllerType(const FName& InCustomControllerTypeName);
+
+	/** Returns the type associated to the specified Custom Controller Type Name */
+	REMOTECONTROLLOGIC_API EPropertyBagPropertyType GetCustomControllerType(const FString& InCustomControllerTypeName);
+
+	REMOTECONTROLLOGIC_API TMap<FName, FString> GetCustomControllerMetaData(const FString& InCustomControllerTypeName);
+
+	/**
+	 * Generates a unique name for a Custom Controller, based on its Custom Type Name
+	 * Useful when creating a new controller.
+	 */
+	REMOTECONTROLLOGIC_API FName GetUniqueNameForController(const URCVirtualPropertyInContainer* InController);
+	
+	inline static FName CustomControllerNameKey = "CustomControllerTypeName";
+	inline static FName CustomControllerExecuteOnLoadNameKey = "CustomControllerExecuteOnLoad";
+
+	inline static FString CustomTextureControllerName = TEXT("Texture");
+};
diff --git a/Engine/Plugins/VirtualProduction/RemoteControl/Source/RemoteControlLogic/Public/RemoteControlLogicConfig.h b/Engine/Plugins/VirtualProduction/RemoteControl/Source/RemoteControlLogic/Public/RemoteControlLogicConfig.h
--- a/Engine/Plugins/VirtualProduction/RemoteControl/Source/RemoteControlLogic/Public/RemoteControlLogicConfig.h
+++ b/Engine/Plugins/VirtualProduction/RemoteControl/Source/RemoteControlLogic/Public/RemoteControlLogicConfig.h
@@ -23,4 +23,15 @@ public:
 	* This config represents the list of such types that may be used as Controllers*/
 	UPROPERTY(config)
 	TArray<FName> SupportedControllerStructTypes;
+
+	/** Add specific Object types Class Paths to this list in order to provide support for them */
+	UPROPERTY(config)
+	TArray<FName> SupportedControllerObjectClassPaths;
+
+	/** Custom Controllers list. The elements in this list are used to generate customized versions of available controllers
+	 * e.g. External Texture controller generates a specialized version of a String Controller.
+	 * Support is introduced by adding MetaData to each newly created controller.
+	 */
+	UPROPERTY(config)
+	TArray<FName> SupportedControllerCustomTypes;
 };
diff --git a/Engine/Plugins/VirtualProduction/RemoteControl/Source/RemoteControlLogic/RemoteControlLogic.Build.cs b/Engine/Plugins/VirtualProduction/RemoteControl/Source/RemoteControlLogic/RemoteControlLogic.Build.cs
--- a/Engine/Plugins/VirtualProduction/RemoteControl/Source/RemoteControlLogic/RemoteControlLogic.Build.cs
+++ b/Engine/Plugins/VirtualProduction/RemoteControl/Source/RemoteControlLogic/RemoteControlLogic.Build.cs
@@ -10,11 +10,13 @@ public class RemoteControlLogic : ModuleRules
 			new string[] {
 				"Core",
 				"CoreUObject",
+				"StructUtils"
 			}
 		);
 
 		PrivateDependencyModuleNames.AddRange(
 			new string[] {
+				"AssetRegistry",
 				"Cbor",
 				"Engine",
 				"HTTP",
diff --git a/Engine/Plugins/VirtualProduction/RemoteControl/Source/RemoteControlUI/Private/UI/Controller/CustomControllers/SCustomTextureControllerWidget.cpp b/Engine/Plugins/VirtualProduction/RemoteControl/Source/RemoteControlUI/Private/UI/Controller/CustomControllers/SCustomTextureControllerWidget.cpp
new file mode 100644
--- /dev/null
+++ b/Engine/Plugins/VirtualProduction/RemoteControl/Source/RemoteControlUI/Private/UI/Controller/CustomControllers/SCustomTextureControllerWidget.cpp
@@ -0,0 +1,300 @@
+﻿// Copyright Epic Games, Inc. All Rights Reserved.
+
+#include "SCustomTextureControllerWidget.h"
+#include "AssetThumbnail.h"
+#include "Controller/RCController.h"
+#include "EditorDirectories.h"
+#include "Engine/Texture2D.h"
+#include "Misc/MessageDialog.h"
+#include "Misc/PackageName.h"
+#include "PropertyCustomizationHelpers.h"
+#include "RCVirtualProperty.h"
+#include "SlateOptMacros.h"
+#include "Action/Bind/RCCustomBindActionUtilities.h"
+#include "Styling/AppStyle.h"
+#include "ThumbnailRendering/ThumbnailManager.h"
+#include "Widgets/Input/SFilePathPicker.h"
+#include "Widgets/Input/STextComboBox.h"
+#include "Widgets/Layout/SSpacer.h"
+#include "Widgets/SBoxPanel.h"
+
+#define LOCTEXT_NAMESPACE "ExternalTextureControllerWidget"
+
+void SCustomTextureControllerWidget::Construct(const FArguments& InArgs, URCVirtualPropertyBase* InController)
+{
+	if (!InController || InController->GetValueType() != EPropertyBagPropertyType::String)
+	{
+		return;
+	}
+	
+	ControllerWeakPtr = InController;
+
+	FString ControllerString;
+	InController->GetValueString(ControllerString);
+	bInternal = FPackageName::IsValidTextForLongPackageName(ControllerString);
+	if (FPackageName::IsValidTextForLongPackageName(ControllerString))
+	{
+		bInternal = true;
+		CurrentAssetPath = ControllerString;
+	}
+	else
+	{
+		bInternal = false;
+		CurrentExternalPath = ControllerString;
+	}
+
+	ControllerTypes.Add(MakeShared<FString>(TEXT("External")));
+	ControllerTypes.Add(MakeShared<FString>(TEXT("Asset")));
+	
+	ChildSlot
+	[
+		SNew(SHorizontalBox)
+		+ SHorizontalBox::Slot()
+		.AutoWidth()
+		[
+			SAssignNew(ThumbnailWidgetBox, SBox)
+			.WidthOverride(64)
+			.HeightOverride(64)
+		]
+		+ SHorizontalBox::Slot()
+		.FillWidth(0.05)
+		[
+			SNew(SSpacer)
+		]
+		+ SHorizontalBox::Slot()
+		.FillWidth(0.6)
+		[
+			SAssignNew(ValueWidgetBox, SBox)
+		]
+		+ SHorizontalBox::Slot()
+		.VAlign(VAlign_Center)
+		.AutoWidth()
+		[
+			SNew(STextComboBox)
+			.OptionsSource(&ControllerTypes)
+			.OnSelectionChanged(this, &SCustomTextureControllerWidget::OnControllerTypeChanged)
+			.InitiallySelectedItem(ControllerTypes[bInternal])
+		]
+	];
+
+	UpdateValueWidget();
+	UpdateThumbnailWidget();
+}
+
+FString SCustomTextureControllerWidget::GetFilePath() const
+{
+	return CurrentExternalPath;
+}
+
+FString SCustomTextureControllerWidget::GetAssetPath() const
+{
+	return CurrentAssetPath;
+}
+
+FString SCustomTextureControllerWidget::GetAssetPathName() const
+{
+	if (Texture)
+	{
+		return Texture->GetPathName();
+	}
+
+	return TEXT("");
+}
+
+URCVirtualPropertyBase* SCustomTextureControllerWidget::GetVirtualProperty() const
+{
+	return ControllerWeakPtr.Get();
+}
+
+void SCustomTextureControllerWidget::OnControllerTypeChanged(TSharedPtr<FString, ESPMode::ThreadSafe> InString, ESelectInfo::Type InArg)
+{
+	const FString& Selection = *InString.Get();
+	
+	if (Selection == TEXT("Asset"))
+	{
+		bInternal = true;
+	}
+	if (Selection == TEXT("External"))
+	{
+		bInternal = false;
+	}
+
+	Texture = nullptr;
+	
+	UpdateValueWidget();
+	UpdateThumbnailWidget();
+	UpdateControllerValue();
+}
+
+void SCustomTextureControllerWidget::OnAssetSelected(const FAssetData& AssetData)
+{
+	UObject* TextureAsset = AssetData.GetAsset();
+	if (TextureAsset != nullptr)
+	{
+		if (TextureAsset->IsA<UTexture2D>())
+		{
+			Texture = Cast<UTexture2D>(TextureAsset);
+			CurrentAssetPath = AssetData.PackageName.ToString();
+			UpdateControllerValue();
+		}
+	}
+}
+
+FString SCustomTextureControllerWidget::GetCurrentPath() const
+{
+	if (bInternal)
+	{
+		return GetAssetPath();
+	}
+	else
+	{
+		return GetFilePath();
+	}
+}
+
+void SCustomTextureControllerWidget::UpdateControllerValue()
+{
+	const FString& Path = GetCurrentPath();
+	
+	if (URCController* Controller = Cast<URCController>(GetVirtualProperty()))
+	{
+		Controller->SetValueString(Path);
+		RefreshThumbnailImage();
+		
+		// Make sure behaviors/actions get executed
+		Controller->ExecuteBehaviours();
+	}
+}
+
+void SCustomTextureControllerWidget::HandleFilePathPickerPathPicked(const FString& InPickedPath)
+{
+	// The following code is adapted from FFilePath customization
+	FString FinalPath = InPickedPath;
+
+	// Keeping these here in case need additional customization
+	bool bLongPackageName = false;
+	bool bRelativeToGameDir = false;
+	
+	if (bLongPackageName)
+	{
+		FString LongPackageName;
+		FString StringFailureReason;
+		if (FPackageName::TryConvertFilenameToLongPackageName(InPickedPath, LongPackageName, &StringFailureReason) == false)
+		{
+			FMessageDialog::Open(EAppMsgType::Ok, FText::FromString(StringFailureReason));
+		}
+		FinalPath = LongPackageName;
+	}
+	else if (bRelativeToGameDir && !InPickedPath.IsEmpty())
+	{
+		// A filepath under the project directory will be made relative to the project directory
+		// Otherwise, the absolute path will be returned unless it doesn't exist, the current path will
+		// be kept. This can happen if it's already relative to project dir (tabbing when selected)
+
+		const FString ProjectDir = FPaths::ProjectDir();
+		const FString AbsoluteProjectDir = FPaths::ConvertRelativePathToFull(ProjectDir);
+		const FString AbsolutePickedPath = FPaths::ConvertRelativePathToFull(InPickedPath);
+
+		// Verify if absolute path to file exists. If it was already relative to content directory
+		// the absolute will be to binaries and will possibly be garbage
+		if (FPaths::FileExists(AbsolutePickedPath))
+		{
+			// If file is part of the project dir, chop the project dir part
+			// Otherwise, use the absolute path
+			if (AbsolutePickedPath.StartsWith(AbsoluteProjectDir))
+			{
+				FinalPath = AbsolutePickedPath.RightChop(AbsoluteProjectDir.Len());
+			}
+			else
+			{
+				FinalPath = AbsolutePickedPath;
+			}
+		}
+		else
+		{
+			// If absolute file doesn't exist, it might already be relative to project dir
+			// If not, then it might be a manual entry, so keep it untouched either way
+			FinalPath = InPickedPath;
+		}
+	}
+
+	CurrentExternalPath = FinalPath;
+	
+	UpdateControllerValue();
+	
+	FEditorDirectories::Get().SetLastDirectory(ELastDirectory::GENERIC_OPEN, FPaths::GetPath(InPickedPath));
+}
+
+TSharedRef<SWidget> SCustomTextureControllerWidget::GetAssetThumbnailWidget()
+{
+	if (!Texture)
+	{
+		RefreshThumbnailImage();
+	}
+	
+	const TSharedRef<FAssetThumbnail> AssetThumbnail = MakeShared<FAssetThumbnail>(Texture, 64, 64, UThumbnailManager::Get().GetSharedThumbnailPool());
+	return AssetThumbnail->MakeThumbnailWidget();
+}
+
+TSharedRef<SWidget> SCustomTextureControllerWidget::GetExternalTextureValueWidget()
+{
+	static const FString FileTypeFilter = TEXT("Image files (*.jpg; *.png; *.bmp; *.ico; *.exr; *.icns; *.jpeg; *.tga; *.hdr)|*.jpg; *.png; *.bmp; *.ico; *.exr; *.icns; *.jpeg; *.tga; *.hdr");
+	
+	return SNew(SFilePathPicker)
+			.BrowseButtonImage(FAppStyle::GetBrush("PropertyWindow.Button_Ellipsis"))
+			.BrowseButtonStyle(FAppStyle::Get(), "HoverHintOnly")
+			.BrowseButtonToolTip(LOCTEXT("FileButtonToolTipText", "Choose a file from this computer"))
+			.BrowseDirectory(FEditorDirectories::Get().GetLastDirectory(ELastDirectory::GENERIC_OPEN))
+			.BrowseTitle(LOCTEXT("PropertyEditorTitle", "File picker..."))
+			.FilePath(this, &SCustomTextureControllerWidget::GetFilePath)
+			.FileTypeFilter(FileTypeFilter)
+			.OnPathPicked(this, &SCustomTextureControllerWidget::HandleFilePathPickerPathPicked);
+}
+
+TSharedRef<SWidget> SCustomTextureControllerWidget::GetAssetTextureValueWidget()
+{
+	return SNew(SObjectPropertyEntryBox)
+			.AllowedClass(UTexture2D::StaticClass())
+			.OnObjectChanged(this, &SCustomTextureControllerWidget::OnAssetSelected)
+			.ObjectPath(this, &SCustomTextureControllerWidget::GetAssetPathName)
+			.DisplayUseSelected(true)
+			.DisplayBrowse(true);
+}
+
+void SCustomTextureControllerWidget::UpdateValueWidget()
+{
+	if (!ValueWidgetBox)
+	{
+		return;
+	}
+	
+	if (bInternal)
+	{
+		ValueWidgetBox->SetContent(GetAssetTextureValueWidget());
+	}
+	else
+	{
+		ValueWidgetBox->SetContent(GetExternalTextureValueWidget());
+	}
+}
+
+void SCustomTextureControllerWidget::UpdateThumbnailWidget()
+{
+	if (!ThumbnailWidgetBox)
+	{
+		return;
+	}
+	
+	ThumbnailWidgetBox->SetContent(GetAssetThumbnailWidget());
+}
+
+void SCustomTextureControllerWidget::RefreshThumbnailImage()
+{
+	if (UTexture2D* LoadedTexture = UE::RCCustomBindActionUtilities::LoadTextureFromPath(GetCurrentPath()))
+	{
+		Texture = LoadedTexture;
+		UpdateThumbnailWidget();
+	}
+}
+
+#undef LOCTEXT_NAMESPACE
diff --git a/Engine/Plugins/VirtualProduction/RemoteControl/Source/RemoteControlUI/Private/UI/Controller/CustomControllers/SCustomTextureControllerWidget.h b/Engine/Plugins/VirtualProduction/RemoteControl/Source/RemoteControlUI/Private/UI/Controller/CustomControllers/SCustomTextureControllerWidget.h
new file mode 100644
--- /dev/null
+++ b/Engine/Plugins/VirtualProduction/RemoteControl/Source/RemoteControlUI/Private/UI/Controller/CustomControllers/SCustomTextureControllerWidget.h
@@ -0,0 +1,90 @@
+﻿// Copyright Epic Games, Inc. All Rights Reserved.
+
+#pragma once
+
+#include "Widgets/Layout/SBox.h"
+#include "Widgets/SCompoundWidget.h"
+
+class URCVirtualPropertyBase;
+class UTexture2D;
+
+class REMOTECONTROLUI_API SCustomTextureControllerWidget : public SCompoundWidget
+{
+public:
+	SLATE_BEGIN_ARGS(SCustomTextureControllerWidget)
+		{
+		}
+
+	SLATE_END_ARGS()
+	
+	/**
+	 * Constructs this widget with InArgs
+	 * @param InController: the controller handled by this widget
+	 */
+	void Construct(const FArguments& InArgs, URCVirtualPropertyBase* InController);
+
+private:
+	/** Returns the handled controller as URCVirtualPropertyBase */
+	URCVirtualPropertyBase* GetVirtualProperty() const;
+
+	/** Handle toggling Texture Controller type change: External vs. Internal */
+	void OnControllerTypeChanged(TSharedPtr<FString, ESPMode::ThreadSafe> InString, ESelectInfo::Type InArg);
+
+	/** Called when a Texture is selected in Asset mode */
+	void OnAssetSelected(const FAssetData& AssetData);
+
+	/** Sets the controller value starting from the current path set via this Widget */
+	void UpdateControllerValue();
+
+	/** Called to handle the file selection via OS window */
+	void HandleFilePathPickerPathPicked(const FString& String);
+
+	/** Refresh the value widget content. Useful when switching between External and Asset modes */
+	void UpdateValueWidget();
+
+	/** Refresh the thumbnail widget content */
+	void UpdateThumbnailWidget();
+
+	/** Loads the selected texture and updates the thumbnail Widget*/
+	void RefreshThumbnailImage();
+
+	/** Returns the current external file path */
+	FString GetFilePath() const;
+
+	/** Returns the current internal/asset file path */
+	FString GetAssetPath() const;
+
+	/** Returns the path name for the currently stored texture */
+	FString GetAssetPathName() const;
+
+	/** Returns current texture path value. Can be External or Asset, depending on current mode */
+	FString GetCurrentPath() const;
+
+	TSharedRef<SWidget> GetAssetThumbnailWidget();
+	TSharedRef<SWidget> GetExternalTextureValueWidget();
+	TSharedRef<SWidget> GetAssetTextureValueWidget();
+
+	/** The controller handled by this widget */
+	TWeakObjectPtr<URCVirtualPropertyBase> ControllerWeakPtr;
+
+	/** The texture selected by the user via the asset or external file path */
+	TObjectPtr<UTexture2D> Texture;
+
+	/** The widget containing the controller value */
+	TSharedPtr<SBox> ValueWidgetBox;
+
+	/** The widget containing the thumbnail for the currently selected Texture */
+	TSharedPtr<SBox> ThumbnailWidgetBox;
+
+	/** Current path for texture project asset */
+	FString CurrentAssetPath;
+
+	/** Current path pointing to external texture (e.g. png file) */
+	FString CurrentExternalPath;
+
+	/** Is the Texture controller currently used to select a texture asset (internal) or an external texture? */
+	bool bInternal = false;
+
+	/** List of available Texture Controller types, used by combo box selection. Supports Asset and External*/
+	TArray<TSharedPtr<FString>> ControllerTypes;
+};
diff --git a/Engine/Plugins/VirtualProduction/RemoteControl/Source/RemoteControlUI/Private/UI/Controller/RCControllerModel.cpp b/Engine/Plugins/VirtualProduction/RemoteControl/Source/RemoteControlUI/Private/UI/Controller/RCControllerModel.cpp
--- a/Engine/Plugins/VirtualProduction/RemoteControl/Source/RemoteControlUI/Private/UI/Controller/RCControllerModel.cpp
+++ b/Engine/Plugins/VirtualProduction/RemoteControl/Source/RemoteControlUI/Private/UI/Controller/RCControllerModel.cpp
@@ -1,7 +1,8 @@
 ﻿// Copyright Epic Games, Inc. All Rights Reserved.
 
 #include "RCControllerModel.h"
-
+#include "Controller/RCCustomControllerUtilities.h"
+#include "CustomControllers/SCustomTextureControllerWidget.h"
 #include "IDetailTreeNode.h"
 #include "RCVirtualProperty.h"
 #include "RemoteControlPreset.h"
@@ -46,22 +47,36 @@ TSharedRef<SWidget> FRCControllerModel::GetWidget() const
 	const FNodeWidgets NodeWidgets = DetailTreeNodeWeakPtr.Pin()->CreateNodeWidgets();
 
 	const TSharedRef<SHorizontalBox> FieldWidget = SNew(SHorizontalBox);
-	if (NodeWidgets.ValueWidget)
-	{
-		FieldWidget->AddSlot()
-			.Padding(FMargin(10.0f, 2.0f))
-			.HAlign(HAlign_Left)
-			[
-				NodeWidgets.ValueWidget.ToSharedRef()
-			];
-	}
-	else if (NodeWidgets.WholeRowWidget)
-	{
-		FieldWidget->AddSlot()
-			.Padding(FMargin(10.0f, 2.0f))
-			[
-				NodeWidgets.WholeRowWidget.ToSharedRef()
-			];
+	if (VirtualPropertyWeakPtr.IsValid())
+	{		
+		const FString& CustomName = UE::RCCustomControllers::GetCustomControllerTypeName(VirtualPropertyWeakPtr.Get());
+
+		// If this is a custom controller, we will use its custom widget
+		if (!CustomName.IsEmpty())
+		{
+			FieldWidget->AddSlot()
+				.Padding(FMargin(10.0f, 2.0f))
+				[
+					GetCustomControllerWidget(CustomName)
+				];
+		}
+		else if (NodeWidgets.ValueWidget)
+		{
+			FieldWidget->AddSlot()
+				.Padding(FMargin(10.0f, 2.0f))
+				.HAlign(HAlign_Left)
+				[
+					NodeWidgets.ValueWidget.ToSharedRef()
+				];
+		}
+		else if (NodeWidgets.WholeRowWidget)
+		{
+			FieldWidget->AddSlot()
+				.Padding(FMargin(10.0f, 2.0f))
+				[
+					NodeWidgets.WholeRowWidget.ToSharedRef()
+				];
+		}
 	}
 
 	return FieldWidget;
@@ -108,6 +123,16 @@ TSharedRef<SWidget> FRCControllerModel::GetTypeSelectionWidget()
 	return SNullWidget::NullWidget;
 }
 
+TSharedRef<SWidget> FRCControllerModel::GetCustomControllerWidget(const FString& InCustomControllerTypeName) const
+{
+	if (InCustomControllerTypeName == UE::RCCustomControllers::CustomTextureControllerName)
+	{		
+		return SNew(SCustomTextureControllerWidget, VirtualPropertyWeakPtr.Get());		
+	}
+
+	return SNullWidget::NullWidget;
+}
+
 void FRCControllerModel::SetMultiController(bool bInIsMultiController)
 {
 	bIsMultiController = bInIsMultiController;
diff --git a/Engine/Plugins/VirtualProduction/RemoteControl/Source/RemoteControlUI/Private/UI/Controller/RCControllerModel.h b/Engine/Plugins/VirtualProduction/RemoteControl/Source/RemoteControlUI/Private/UI/Controller/RCControllerModel.h
--- a/Engine/Plugins/VirtualProduction/RemoteControl/Source/RemoteControlUI/Private/UI/Controller/RCControllerModel.h
+++ b/Engine/Plugins/VirtualProduction/RemoteControl/Source/RemoteControlUI/Private/UI/Controller/RCControllerModel.h
@@ -48,6 +48,9 @@ public:
 	/** The widget showing the Value Type selection dropdown for this Controller (used for MultiControllers) */
 	TSharedRef<SWidget> GetTypeSelectionWidget();
 
+	/** Retrieve the widget for a custom controller */
+	TSharedRef<SWidget> GetCustomControllerWidget(const FString& InCustomControllerTypeName) const;
+
 	/**
 	 * Specify if this is a MultiController handling other hidden controllers.
 	 * MultiControllers will setup and show a Type Selection Widget.
diff --git a/Engine/Plugins/VirtualProduction/RemoteControl/Source/RemoteControlUI/Private/UI/Controller/SRCControllerPanel.cpp b/Engine/Plugins/VirtualProduction/RemoteControl/Source/RemoteControlUI/Private/UI/Controller/SRCControllerPanel.cpp
--- a/Engine/Plugins/VirtualProduction/RemoteControl/Source/RemoteControlUI/Private/UI/Controller/SRCControllerPanel.cpp
+++ b/Engine/Plugins/VirtualProduction/RemoteControl/Source/RemoteControlUI/Private/UI/Controller/SRCControllerPanel.cpp
@@ -3,22 +3,25 @@
 #include "SRCControllerPanel.h"
 
 #include "Controller/RCController.h"
+#include "Controller/RCCustomControllerUtilities.h"
+#include "Engine/Texture.h"
 #include "Framework/MultiBox/MultiBoxBuilder.h"
+#include "Materials/MaterialInterface.h"
 #include "Misc/MessageDialog.h"
 #include "RCControllerModel.h"
 #include "RCVirtualPropertyContainer.h"
 #include "RemoteControlLogicConfig.h"
 #include "RemoteControlPreset.h"
-#include "SlateOptMacros.h"
 #include "SRCControllerPanelList.h"
+#include "SlateOptMacros.h"
 #include "Styling/RemoteControlStyles.h"
 #include "UI/BaseLogicUI/RCLogicModeBase.h"
 #include "UI/Panels/SRCDockPanel.h"
-#include "UI/RemoteControlPanelStyle.h"
 #include "UI/RCUIHelpers.h"
+#include "UI/RemoteControlPanelStyle.h"
 #include "UI/SRemoteControlPanel.h"
-#include "Widgets/SBoxPanel.h"
 #include "Widgets/Input/SComboButton.h"
+#include "Widgets/SBoxPanel.h"
 #include "Widgets/Views/SListView.h"
 
 #if WITH_EDITOR
@@ -106,6 +109,21 @@ void SRCControllerPanel::Construct(const FArguments& InArgs, const TSharedRef<SR
 
 END_SLATE_FUNCTION_BUILD_OPTIMIZATION
 
+SRCControllerPanel::FRCControllerPropertyInfo SRCControllerPanel::GetPropertyInfoForCustomType(const FName& InType)
+{
+	FRCControllerPropertyInfo PropertyInfo;
+	const EPropertyBagPropertyType CustomControllerType = UE::RCCustomControllers::GetCustomControllerType(InType);
+
+	if (CustomControllerType != EPropertyBagPropertyType::None)
+	{
+		PropertyInfo.Type = CustomControllerType;
+		PropertyInfo.CustomMetaData = UE::RCCustomControllers::GetCustomControllerMetaData(InType.ToString());
+	}
+
+	return PropertyInfo;
+}
+
+
 static UObject* GetBaseStructForType(FName StructType)
 {
 	if (StructType == NAME_Vector)
@@ -304,13 +322,12 @@ EVisibility SRCControllerPanel::GetMultiControllerSwitchVisibility() const
 	return EVisibility::Visible;
 }
 
-
 TSharedRef<SWidget> SRCControllerPanel::GetControllerMenuContentWidget() const
 {
 	constexpr  bool bShouldCloseWindowAfterMenuSelection = true;
 	FMenuBuilder MenuBuilder(bShouldCloseWindowAfterMenuSelection, nullptr);
 
-	TArray<TPair<EPropertyBagPropertyType, UObject*>> VirtualPropertyFieldClassNames;
+	TArray<TPair<EPropertyBagPropertyType, FRCControllerPropertyInfo>> VirtualPropertyFieldClassNames;
 
 	// See Config file: BaseRemoteControl.ini
 	const URemoteControlLogicConfig* RCLogicConfig = GetDefault<URemoteControlLogicConfig>();
@@ -318,33 +335,70 @@ TSharedRef<SWidget> SRCControllerPanel::GetControllerMenuContentWidget() const
 	{
 		if (ControllerType == EPropertyBagPropertyType::Struct)
 		{
-			for (const FName StructType : RCLogicConfig->SupportedControllerStructTypes)
+			for (const FName& StructType : RCLogicConfig->SupportedControllerStructTypes)
 			{
 				UObject* ValueTypeObject = GetBaseStructForType(StructType);
 				if (ensure(ValueTypeObject))
 				{
-					VirtualPropertyFieldClassNames.Add({ ControllerType, ValueTypeObject });
+					FRCControllerPropertyInfo PropertyInfo(ValueTypeObject);
+					VirtualPropertyFieldClassNames.Add({ ControllerType, PropertyInfo });
 				}
 			}
 		}
+		else if (ControllerType == EPropertyBagPropertyType::Object)
+		{
+			for (const FName& ObjectType : RCLogicConfig->SupportedControllerObjectClassPaths)
+			{			
+				FSoftClassPath ClassPath(ObjectType.ToString());
+				UObject* ValueTypeObject = ClassPath.TryLoad();
+				
+				if (ensure(ValueTypeObject))
+				{
+					FRCControllerPropertyInfo PropertyInfo(ValueTypeObject);
+					VirtualPropertyFieldClassNames.Add({ ControllerType, PropertyInfo });
+				}	
+			}
+		}
 		else
 		{
 			VirtualPropertyFieldClassNames.Add({ ControllerType, nullptr });
 		}
 	}
 
+	for (const FName& CustomType : RCLogicConfig->SupportedControllerCustomTypes)
+	{
+		FRCControllerPropertyInfo PropertyInfo = GetPropertyInfoForCustomType(CustomType);
+		VirtualPropertyFieldClassNames.Add({ PropertyInfo.Type, PropertyInfo });
+	}
+
 	// Generate a menu from the list of supported Controllers
-	for (const TPair<EPropertyBagPropertyType, UObject*>& Pair : VirtualPropertyFieldClassNames)
+	for (const TPair<EPropertyBagPropertyType, FRCControllerPropertyInfo>& Pair : VirtualPropertyFieldClassNames)
 	{
+		const EPropertyBagPropertyType Type = Pair.Key;
+		const FRCControllerPropertyInfo& PropertyInfo = Pair.Value;
+		UObject* ValueTypeObject = PropertyInfo.ValueTypeObject;
+		
 		// Display Name
-		const FName DefaultName = URCVirtualPropertyBase::GetVirtualPropertyTypeDisplayName(Pair.Key, Pair.Value);
+		FName DefaultName;
+
+		const FName& CustomNameKey = UE::RCCustomControllers::CustomControllerNameKey;
+		
+		// Check if this is a custom property with a dedicated name
+		if (PropertyInfo.CustomMetaData.Contains(CustomNameKey))
+		{
+			DefaultName = FName(PropertyInfo.CustomMetaData[CustomNameKey]);
+		}
+		else
+		{
+			DefaultName = URCVirtualPropertyBase::GetVirtualPropertyTypeDisplayName(Type, ValueTypeObject);
+		}
 
 		// Type Color
 		FLinearColor TypeColor = FColor::White;
 
 		// Generate a transient virtual property for deducing type color:
 		FInstancedPropertyBag Bag;
-		Bag.AddProperty(DefaultName, Pair.Key, Pair.Value);
+		Bag.AddProperty(DefaultName, Type, ValueTypeObject);
 		const FPropertyBagPropertyDesc* BagPropertyDesc = Bag.FindPropertyDescByName(DefaultName);
 		if (!ensure(BagPropertyDesc))
 		{
@@ -370,7 +424,7 @@ TSharedRef<SWidget> SRCControllerPanel::GetControllerMenuContentWidget() const
 			];
 
 		// Menu Item
-		FUIAction Action(FExecuteAction::CreateSP(this, &SRCControllerPanel::OnAddControllerClicked, Pair.Key, Pair.Value));
+		FUIAction Action(FExecuteAction::CreateSP(this, &SRCControllerPanel::OnAddControllerClicked, Type, PropertyInfo));
 
 		MenuBuilder.AddMenuEntry(Action, MenuItemWidget);
 	}
@@ -400,7 +454,7 @@ TSharedRef<SWidget> SRCControllerPanel::GetMultiControllerSwitchWidget()
 		];
 }
 
-void SRCControllerPanel::OnAddControllerClicked(const EPropertyBagPropertyType InValueType, UObject* InValueTypeObject) const
+void SRCControllerPanel::OnAddControllerClicked(const EPropertyBagPropertyType InValueType, const FRCControllerPropertyInfo InControllerPropertyInfo) const
 {
 	// Add to the asset
 	if (URemoteControlPreset* Preset = GetPreset())
@@ -408,7 +462,20 @@ void SRCControllerPanel::OnAddControllerClicked(const EPropertyBagPropertyType I
 		FScopedTransaction Transaction(LOCTEXT("AddController", "Add Controller"));
 		Preset->Modify();
 
-		URCVirtualPropertyInContainer* NewVirtualProperty = Preset->AddController(URCController::StaticClass(), InValueType, InValueTypeObject);
+		URCVirtualPropertyInContainer* NewVirtualProperty = Preset->AddController(URCController::StaticClass(), InValueType, InControllerPropertyInfo.ValueTypeObject);
+
+		// Copy over the custom metadata of this property
+		for (const TPair<FName, FString>& Pair : InControllerPropertyInfo.CustomMetaData)
+		{
+			NewVirtualProperty->SetMetadataValue(Pair.Key, Pair.Value);
+		}
+
+		// Additional logic for Custom Controllers naming
+		if (UE::RCCustomControllers::IsCustomController(NewVirtualProperty))
+		{
+			const FName& NewControllerName = UE::RCCustomControllers::GetUniqueNameForController(NewVirtualProperty);
+			NewVirtualProperty->DisplayName = NewControllerName;
+		}
 
 		if (ControllerPanelList.IsValid())
 		{
diff --git a/Engine/Plugins/VirtualProduction/RemoteControl/Source/RemoteControlUI/Private/UI/Controller/SRCControllerPanel.h b/Engine/Plugins/VirtualProduction/RemoteControl/Source/RemoteControlUI/Private/UI/Controller/SRCControllerPanel.h
--- a/Engine/Plugins/VirtualProduction/RemoteControl/Source/RemoteControlUI/Private/UI/Controller/SRCControllerPanel.h
+++ b/Engine/Plugins/VirtualProduction/RemoteControl/Source/RemoteControlUI/Private/UI/Controller/SRCControllerPanel.h
@@ -20,6 +20,28 @@ enum class EPropertyBagPropertyType : uint8;
 */
 class REMOTECONTROLUI_API SRCControllerPanel : public SRCLogicPanelBase
 {
+	/** Helper struct used when generating Controllers.
+	 * Wraps the TypeObject and the Bag Type.
+	 * 
+	 * Additionally, some facultative MetaData can be added to create a Custom Controller extending the initial type
+	 * e.g. External Texture custom controller is a specialized String controller
+	 * See RCCustomControllerUtilities.h for helper functions related to custom controllers.
+	 */
+	struct FRCControllerPropertyInfo
+	{
+		FRCControllerPropertyInfo(){}
+		FRCControllerPropertyInfo(UObject* InValueTypeObject) : ValueTypeObject(InValueTypeObject){}
+
+		/** The type of the controller which will be created starting from this FRCControllerPropertyInfo */
+		EPropertyBagPropertyType Type = EPropertyBagPropertyType::Name;
+		
+		/** Object defining an Enum, Struct, or Class */
+		UObject* ValueTypeObject = nullptr;
+
+		/** MetaData Map which can be used to pass additional settings or information to be used when customizing a Custom Controller */
+		TMap<FName, FString> CustomMetaData;
+	};
+	
 public:
 	SLATE_BEGIN_ARGS(SRCControllerPanel)
 	{}
@@ -28,6 +50,8 @@ public:
 
 	SLATE_END_ARGS()
 
+	static FRCControllerPropertyInfo GetPropertyInfoForCustomType(const FName& InType);
+
 	/** Constructs this widget with InArgs */
 	void Construct(const FArguments& InArgs, const TSharedRef<SRemoteControlPanel>& InPanel);
 
@@ -71,7 +95,7 @@ private:
 	TSharedRef<SWidget> GetMultiControllerSwitchWidget();
 
 	/** Handles click event for Add Controller button*/
-	void OnAddControllerClicked(const EPropertyBagPropertyType InValueType, UObject* InValueTypeObject = nullptr) const;
+	void OnAddControllerClicked(const EPropertyBagPropertyType InValueType, const FRCControllerPropertyInfo InControllerPropertyInfo) const;
 
 	/** Handles checkbox changes event for the MultiController mode selection widget */
 	void OnToggleMultiControllersMode(ECheckBoxState CheckBoxState);
diff --git a/Engine/Plugins/VirtualProduction/RemoteControl/Source/RemoteControlUI/Private/UI/Controller/SRCControllerPanelList.cpp b/Engine/Plugins/VirtualProduction/RemoteControl/Source/RemoteControlUI/Private/UI/Controller/SRCControllerPanelList.cpp
--- a/Engine/Plugins/VirtualProduction/RemoteControl/Source/RemoteControlUI/Private/UI/Controller/SRCControllerPanelList.cpp
+++ b/Engine/Plugins/VirtualProduction/RemoteControl/Source/RemoteControlUI/Private/UI/Controller/SRCControllerPanelList.cpp
@@ -7,10 +7,13 @@
 #include "Commands/RemoteControlCommands.h"
 #include "Controller/RCController.h"
 #include "Controller/RCControllerContainer.h"
+#include "Controller/RCCustomControllerUtilities.h"
+#include "Engine/Texture.h"
+#include "Engine/Texture2D.h"
 #include "IDetailTreeNode.h"
-#include "Interfaces/IMainFrameModule.h"
 #include "IPropertyRowGenerator.h"
 #include "IRemoteControlModule.h"
+#include "Interfaces/IMainFrameModule.h"
 #include "RCControllerModel.h"
 #include "RCMultiController.h"
 #include "RCVirtualProperty.h"
@@ -18,15 +21,15 @@
 #include "RemoteControlField.h"
 #include "RemoteControlPreset.h"
 #include "SDropTarget.h"
-#include "SlateOptMacros.h"
 #include "SRCControllerPanel.h"
+#include "SlateOptMacros.h"
 #include "Styling/RemoteControlStyles.h"
 #include "UI/Action/SRCActionPanelList.h"
 #include "UI/BaseLogicUI/RCLogicModeBase.h"
 #include "UI/RCUIHelpers.h"
 #include "UI/RemoteControlPanelStyle.h"
-#include "UI/SRCPanelExposedEntity.h"
 #include "UI/SRCPanelDragHandle.h"
+#include "UI/SRCPanelExposedEntity.h"
 #include "UI/SRemoteControlPanel.h"
 #include "Widgets/Views/SHeaderRow.h"
 
@@ -250,7 +253,6 @@ namespace UE::RCControllerPanelList
 								}
 							}
 						}
-
 					}
 				}
 			}
@@ -787,10 +789,31 @@ void SRCControllerPanelList::CreateAutoBindForProperty(TSharedPtr<const FRemoteC
 		bool bSuccess = URCBehaviourBind::GetPropertyBagTypeFromFieldProperty(Property, PropertyBagType, StructObject);
 		if(bSuccess)
 		{
+			// Preparation step, in case we are dealing with a custom controller
+			FString CustomControllerName = TEXT("");
+			if (StructObject == UTexture::StaticClass() || StructObject == UTexture2D::StaticClass())
+			{
+				if (PropertyBagType == EPropertyBagPropertyType::String)
+				{
+					StructObject = nullptr;
+					CustomControllerName = UE::RCCustomControllers::CustomTextureControllerName;
+				}
+			}
+			
 			// Step 1. Create a Controller of matching type
 			URCController* NewController = Cast<URCController>(Preset->AddController(URCController::StaticClass(), PropertyBagType, StructObject, RemoteControlProperty->FieldName));
 			NewController->DisplayIndex = Preset->GetNumControllers() - 1;
 
+			// Add metadata to this controller, if this is a custom controller
+			if (!CustomControllerName.IsEmpty())
+			{
+				const TMap<FName, FString>& CustomControllerMetaData = UE::RCCustomControllers::GetCustomControllerMetaData(CustomControllerName);
+				for (const TPair<FName, FString>& Pair : CustomControllerMetaData)
+				{
+					NewController->SetMetadataValue(Pair.Key, Pair.Value);
+				}
+			}
+
 			// Transfer property value from Exposed Property to the New Controller.
 			// The goal is to keep values synced for a Controller newly created via "Auto Bind"
 			URCBehaviourBind::CopyPropertyValueToController(NewController, RemoteControlProperty.ToSharedRef());
diff --git a/Engine/Plugins/VirtualProduction/RemoteControl/Source/RemoteControlUI/RemoteControlUI.Build.cs b/Engine/Plugins/VirtualProduction/RemoteControl/Source/RemoteControlUI/RemoteControlUI.Build.cs
--- a/Engine/Plugins/VirtualProduction/RemoteControl/Source/RemoteControlUI/RemoteControlUI.Build.cs
+++ b/Engine/Plugins/VirtualProduction/RemoteControl/Source/RemoteControlUI/RemoteControlUI.Build.cs
@@ -17,6 +17,7 @@ public class RemoteControlUI : ModuleRules
 				"Core",
 				"CoreUObject",
 				"ContentBrowser",
+				"DesktopWidgets",
 				"DeveloperSettings",
 				"EditorWidgets",
 				"EditorStyle",
@@ -35,6 +36,7 @@ public class RemoteControlUI : ModuleRules
 				"RemoteControlProtocol",
 				"RemoteControlProtocolWidgets",
 				"SceneOutliner",
+				"Serialization",
 				"Slate",
 				"SlateCore",
 				"StructUtils",
