[Remote Control] Fix non-unity compile issue #316763

Horde issue 316763: https://horde.devtools.epicgames.com/job/64825d9637416fd06c7bed23?step=b9d4&issue=316763

diff --git a/Engine/Plugins/VirtualProduction/RemoteControl/Source/RemoteControlLogic/Private/Action/Bind/RCCustomBindActionUtilities.cpp b/Engine/Plugins/VirtualProduction/RemoteControl/Source/RemoteControlLogic/Private/Action/Bind/RCCustomBindActionUtilities.cpp
--- a/Engine/Plugins/VirtualProduction/RemoteControl/Source/RemoteControlLogic/Private/Action/Bind/RCCustomBindActionUtilities.cpp
+++ b/Engine/Plugins/VirtualProduction/RemoteControl/Source/RemoteControlLogic/Private/Action/Bind/RCCustomBindActionUtilities.cpp
@@ -5,6 +5,7 @@
 #include "Engine/Texture2D.h"
 #include "IRemoteControlPropertyHandle.h"
 #include "Kismet/KismetRenderingLibrary.h"
+#include "Misc/Paths.h"
 #include "RemoteControlField.h"
 
 UTexture2D* UE::RCCustomBindActionUtilities::LoadTextureFromPath(const FString& InPath)
