[Remote Control] Expose Sub Properties

(on behalf of Andrea Botti)

diff --git a/Engine/Plugins/VirtualProduction/RemoteControl/Source/RemoteControl/Private/RemoteControlModule.cpp b/Engine/Plugins/VirtualProduction/RemoteControl/Source/RemoteControl/Private/RemoteControlModule.cpp
--- a/Engine/Plugins/VirtualProduction/RemoteControl/Source/RemoteControl/Private/RemoteControlModule.cpp
+++ b/Engine/Plugins/VirtualProduction/RemoteControl/Source/RemoteControl/Private/RemoteControlModule.cpp
@@ -1618,7 +1618,9 @@ void FRemoteControlModule::RefreshEditorPostSetObjectProperties(const FRCObjectR
 	{
 		if (FProperty* Property = ObjectAccess.Property.Get())
 		{
-			if(Property->GetFName() == RemoteControlUtil::NAME_RelativeLocation)
+			// Check if the property is the RelativeLocation or if a child by checking the Segment[0]
+			const bool bCheckSegmentParent = ObjectAccess.PropertyPathInfo.Segments.Num() ? ObjectAccess.PropertyPathInfo.Segments[0].Name == RemoteControlUtil::NAME_RelativeLocation : false;
+			if(Property->GetFName() == RemoteControlUtil::NAME_RelativeLocation || bCheckSegmentParent)
 			{
 				if (AActor* Actor = SceneComponent->GetOwner())
 				{
@@ -2695,6 +2697,7 @@ void FRemoteControlModule::TestOrFinalizeOngoingChange(bool bForceEndChange)
 						ChainEvent.SetArrayIndexPerObject(ArrayIndicesPerObject);
 
 						Object->PostEditChangeChainProperty(ChainEvent);
+						RefreshEditorPostSetObjectProperties(ObjectReference);
 					}
 				}
 			}
diff --git a/Engine/Plugins/VirtualProduction/RemoteControl/Source/RemoteControlUI/Private/RemoteControlUIModule.cpp b/Engine/Plugins/VirtualProduction/RemoteControl/Source/RemoteControlUI/Private/RemoteControlUIModule.cpp
--- a/Engine/Plugins/VirtualProduction/RemoteControl/Source/RemoteControlUI/Private/RemoteControlUIModule.cpp
+++ b/Engine/Plugins/VirtualProduction/RemoteControl/Source/RemoteControlUI/Private/RemoteControlUIModule.cpp
@@ -5,6 +5,7 @@
 #include "AssetEditor/RemoteControlPresetEditorToolkit.h"
 #include "AssetTools/RemoteControlPresetActions.h"
 #include "AssetToolsModule.h"
+#include "DetailTreeNode.h"
 #include "Commands/RemoteControlCommands.h"
 #include "Components/SkeletalMeshComponent.h"
 #include "Components/StaticMeshComponent.h"
@@ -31,6 +32,7 @@
 #include "Styling/AppStyle.h"
 #include "Styling/AppStyle.h"
 #include "Textures/SlateIcon.h"
+#include "UI/RemoteControlExposeMenuStyle.h"
 #include "UI/Action/SRCActionPanel.h"
 #include "UI/Behaviour/Builtin/Conditional/SRCBehaviourConditional.h"
 #include "UI/Behaviour/SRCBehaviourPanel.h"
@@ -121,14 +123,9 @@ FProperty* FRCExposesPropertyArgs::GetProperty() const
 	const EType Type = GetType();
 	if (Type == EType::E_Handle)
 	{
-		PropertyHandle->GetProperty();
+		return PropertyHandle->GetProperty();
 	}
-	else
-	{
-		return Property.Get();
-	}
-
-	return nullptr;
+	return Property.Get();
 }
 
 FProperty* FRCExposesPropertyArgs::GetPropertyChecked() const
@@ -151,7 +148,7 @@ namespace RemoteControlUIModule
 
 		return CustomizedStructNames;
 	};
-	
+
 	static void OnOpenRemoteControlPanel(const EToolkitMode::Type Mode, const TSharedPtr<class IToolkitHost>& EditWithinLevelEditor, class URemoteControlPreset* Preset)
 	{
 		TSharedRef<FRemoteControlPresetEditorToolkit> Toolkit = FRemoteControlPresetEditorToolkit::CreateEditor(Mode, EditWithinLevelEditor,  Preset);
@@ -253,6 +250,7 @@ namespace RemoteControlUIModule
 void FRemoteControlUIModule::StartupModule()
 {
 	FRemoteControlPanelStyle::Initialize();
+	FRemoteControlExposeMenuStyle::Initialize();
 	BindRemoteControlCommands();
 	RegisterAssetTools();
 	RegisterDetailRowExtension();
@@ -273,6 +271,7 @@ void FRemoteControlUIModule::ShutdownModule()
 	UnregisterAssetTools();
 	UnbindRemoteControlCommands();
 	FRemoteControlPanelStyle::Shutdown();
+	FRemoteControlExposeMenuStyle::Shutdown();
 	SRemoteControlPanel::Shutdown();
 	SRCActionPanel::Shutdown();
 	SRCBehaviourPanel::Shutdown();
@@ -318,7 +317,7 @@ TSharedRef<SRemoteControlPanel> FRemoteControlUIModule::CreateRemoteControlPanel
 
 				if (!IsValid(Preset) || Preset->IsEmbeddedPreset() == false)
 				{
-					// Activating the live mode on a panel sets it as the active panel 
+					// Activating the live mode on a panel sets it as the active panel
 					if (bLiveMode)
 					{
 						if (TSharedPtr<SRemoteControlPanel> ActivePanel = WeakActivePanel.Pin())
@@ -355,17 +354,23 @@ TSharedRef<SRemoteControlPanel> FRemoteControlUIModule::CreateRemoteControlPanel
 		for (const FName& DetailsTabIdentifier : DetailsTabIdentifiers)
 		{
 			SharedDetailsPanel = PropertyEditor.FindDetailView(DetailsTabIdentifier);
-			
+
 			if (SharedDetailsPanel.IsValid())
 			{
-				SharedDetailsPanel->SetRightColumnMinWidth(50);
+				SharedDetailsPanel->SetRightColumnMinWidth(100);
+				SharedDetailsPanel->ForceRefresh();
 				break;
 			}
 		}
 	}
+	else
+	{
+		SharedDetailsPanel->SetRightColumnMinWidth(100);
+		SharedDetailsPanel->ForceRefresh();
+	}
 
 	OnRemoteControlPresetOpened().Broadcast(Preset);
-	
+
 	return PanelRef;
 }
 
@@ -477,7 +482,7 @@ void FRemoteControlUIModule::UnbindRemoteControlCommands()
 		ActionList.UnmapAction(Commands.PasteItem);
 		ActionList.UnmapAction(Commands.DuplicateItem);
 	}
-	
+
 	FRemoteControlCommands::Unregister();
 }
 
@@ -504,6 +509,24 @@ void FRemoteControlUIModule::HandleCreatePropertyRowExtension(const FOnGenerateG
 			FIsActionButtonVisible::CreateRaw(this, &FRemoteControlUIModule::CanToggleExposeProperty, ExposesPropertyArgs)
 		);
 
+		if (ExposesPropertyArgs.PropertyHandle.IsValid() && HasChildProperties(ExposesPropertyArgs.PropertyHandle->GetProperty()))
+		{
+			// Expose/Unexpose button.
+			FPropertyRowExtensionButton& ExposeSubPropertiesButton = OutExtensions.AddDefaulted_GetRef();
+			ExposeSubPropertiesButton.Icon = TAttribute<FSlateIcon>::Create([]()
+			{
+				return FSlateIcon(FRemoteControlExposeMenuStyle::GetStyleSetName(),"RemoteControlExposeMenu.Expand");
+			});
+
+			ExposeSubPropertiesButton.Label = LOCTEXT("RemoteControlExpandMenu_Title","Expose Sub-Property");
+			ExposeSubPropertiesButton.ToolTip = LOCTEXT("RemoteControlExpandMenu_Tooltip","Expose Sub-Property");
+			ExposeSubPropertiesButton.UIAction = FUIAction(
+				FExecuteAction::CreateRaw(this, &FRemoteControlUIModule::OnToggleExposePropertyWithChild, ExposesPropertyArgs),
+				FCanExecuteAction::CreateRaw(this, &FRemoteControlUIModule::CanToggleExposeProperty, ExposesPropertyArgs),
+				FGetActionCheckState(),
+				FIsActionButtonVisible::CreateRaw(this, &FRemoteControlUIModule::CanToggleExposeProperty, ExposesPropertyArgs));
+		}
+
 		// Override material(s) warning.
 		FPropertyRowExtensionButton& OverrideMaterialButton = OutExtensions.AddDefaulted_GetRef();
 
@@ -564,7 +587,7 @@ TSharedPtr<SRemoteControlPanel> FRemoteControlUIModule::GetPanelForProperty(cons
 		if (OuterObjects.Num() > 0)
 		{
 			return GetPanelForObject(OuterObjects[0]);
-		}		
+		}
 	}
 	else if (ExtensionArgsType == FRCExposesPropertyArgs::EType::E_OwnerObject)
 	{
@@ -586,7 +609,7 @@ TSharedPtr<SRemoteControlPanel> FRemoteControlUIModule::GetPanelForPropertyChang
 
 FSlateIcon FRemoteControlUIModule::OnGetExposedIcon(const FRCExposesPropertyArgs& InPropertyArgs) const
 {
-	FName BrushName("NoBrush");
+	FName BrushName("RemoteControlExposeMenu.NoBrush");
 
 	if (TSharedPtr<SRemoteControlPanel> Panel = GetPanelForProperty(InPropertyArgs))
 	{
@@ -595,16 +618,16 @@ FSlateIcon FRemoteControlUIModule::OnGetExposedIcon(const FRCExposesPropertyArgs
 			EPropertyExposeStatus Status = GetPropertyExposeStatus(InPropertyArgs);
 			if (Status == EPropertyExposeStatus::Exposed)
 			{
-				BrushName = "Level.VisibleIcon16x";
+				BrushName = HasChildPropertiesExposed(InPropertyArgs) ? "RemoteControlExposeMenu.VisibleAndVisibleChildren" : "RemoteControlExposeMenu.Visible";
 			}
 			else
 			{
-				BrushName = "Level.NotVisibleIcon16x";
+				BrushName = HasChildPropertiesExposed(InPropertyArgs) ? "RemoteControlExposeMenu.HiddenAndVisibleChildren" : "RemoteControlExposeMenu.Hidden";
 			}
 		}
 	}
 
-	return FSlateIcon(FAppStyle::Get().GetStyleSetName(), BrushName);
+	return FSlateIcon(FRemoteControlExposeMenuStyle::GetStyleSetName(), BrushName);
 }
 
 bool FRemoteControlUIModule::CanToggleExposeProperty(const FRCExposesPropertyArgs InPropertyArgs) const
@@ -641,10 +664,62 @@ void FRemoteControlUIModule::OnToggleExposeProperty(const FRCExposesPropertyArgs
 	{
 		return;
 	}
+	if (ShouldCreateSubMenuForChildProperties(InPropertyArgs))
+	{
+		CreateSubMenuForChildProperties(InPropertyArgs);
+	}
+	else if (TSharedPtr<SRemoteControlPanel> Panel = GetPanelForProperty(InPropertyArgs))
+	{
+		Panel->ToggleProperty(InPropertyArgs);
+	}
+}
 
+void FRemoteControlUIModule::OnToggleExposeSubProperty(const FRCExposesPropertyArgs InPropertyArgs, const FString InDesiredName) const
+{
+	if (!ensureMsgf(InPropertyArgs.IsValid(), TEXT("Property could not be exposed because the extension args was invalid.")))
+	{
+		return;
+	}
 	if (TSharedPtr<SRemoteControlPanel> Panel = GetPanelForProperty(InPropertyArgs))
 	{
-		Panel->ToggleProperty(InPropertyArgs);
+		Panel->ToggleProperty(InPropertyArgs, InDesiredName);
+	}
+}
+
+void FRemoteControlUIModule::OnToggleExposePropertyWithChild(const FRCExposesPropertyArgs InPropertyArgs)
+{
+	if (!ensureMsgf(InPropertyArgs.IsValid(), TEXT("Property could not be exposed because the extension args was invalid.")))
+	{
+		return;
+	}
+	CreateSubMenuForChildProperties(InPropertyArgs);
+}
+
+void FRemoteControlUIModule::OnExposeAll(const TArray<FRCExposesAllPropertiesArgs> InExposeAllArgs) const
+{
+	for (const FRCExposesAllPropertiesArgs& PropInfo : InExposeAllArgs)
+	{
+		if (TSharedPtr<SRemoteControlPanel> Panel = GetPanelForProperty(PropInfo.PropertyArgs))
+		{
+			if (!Panel->IsExposed(PropInfo.PropertyArgs))
+			{
+				OnToggleExposeSubProperty(PropInfo.PropertyArgs, PropInfo.DesiredName);
+			}
+		}
+	}
+}
+
+void FRemoteControlUIModule::OnUnexposeAll(const TArray<FRCExposesAllPropertiesArgs> InExposeAllArgs) const
+{
+	for (const FRCExposesAllPropertiesArgs& PropInfo : InExposeAllArgs)
+	{
+		if (TSharedPtr<SRemoteControlPanel> Panel = GetPanelForProperty(PropInfo.PropertyArgs))
+		{
+			if (Panel->IsExposed(PropInfo.PropertyArgs))
+			{
+				OnToggleExposeSubProperty(PropInfo.PropertyArgs, PropInfo.DesiredName);
+			}
+		}
 	}
 }
 
@@ -837,12 +912,12 @@ bool FRemoteControlUIModule::IsAllowedOwnerObjects(TArray<UObject*> InOuterObjec
 void FRemoteControlUIModule::RegisterStructCustomizations()
 {
 	FPropertyEditorModule& PropertyEditorModule = FModuleManager::GetModuleChecked<FPropertyEditorModule>("PropertyEditor");
-	
+
 	for (FName Name : RemoteControlUIModule::GetCustomizedStructNames())
 	{
 		PropertyEditorModule.RegisterCustomClassLayout(Name, FOnGetDetailCustomizationInstance::CreateStatic(&FRemoteControlEntityCustomization::MakeInstance));
 	}
-	
+
 	PropertyEditorModule.RegisterCustomPropertyTypeLayout(FRCNetworkAddress::StaticStruct()->GetFName(), FOnGetPropertyTypeCustomizationInstance::CreateStatic(&FNetworkAddressCustomization::MakeInstance));
 	PropertyEditorModule.RegisterCustomPropertyTypeLayout(FRCPassphrase::StaticStruct()->GetFName(), FOnGetPropertyTypeCustomizationInstance::CreateStatic(&FPassphraseCustomization::MakeInstance));
 }
@@ -850,7 +925,7 @@ void FRemoteControlUIModule::RegisterStructCustomizations()
 void FRemoteControlUIModule::UnregisterStructCustomizations()
 {
 	FPropertyEditorModule& PropertyEditorModule = FModuleManager::GetModuleChecked<FPropertyEditorModule>("PropertyEditor");
-	
+
 	// Unregister delegates in reverse order.
 	for (int8 NameIndex = RemoteControlUIModule::GetCustomizedStructNames().Num() - 1; NameIndex >= 0; NameIndex--)
 	{
@@ -982,7 +1057,7 @@ void FRemoteControlUIModule::TryOverridingMaterials(const FRCExposesPropertyArgs
 	{
 		return;
 	}
-	
+
 	const FRCExposesPropertyArgs::EType ExtensionArgsType = InPropertyArgs.GetType();
 
 	if (ExtensionArgsType == FRCExposesPropertyArgs::EType::E_Handle)
@@ -1134,6 +1209,146 @@ void FRemoteControlUIModule::RefreshPanels()
 	}
 }
 
+bool FRemoteControlUIModule::ShouldCreateSubMenuForChildProperties(const FRCExposesPropertyArgs InPropertyArgs) const
+{
+	if (InPropertyArgs.PropertyHandle.IsValid() && FSlateApplication::Get().GetModifierKeys().IsControlDown())
+	{
+		return HasChildProperties(InPropertyArgs.GetProperty());
+	}
+	return false;
+}
+
+bool FRemoteControlUIModule::HasChildProperties(const FProperty* InProperty) const
+{
+	if (const FStructProperty* StructProperty = CastField<FStructProperty>(InProperty))
+	{
+		// if at least it contains one child than it is consider valid to expose child properties
+		return StructProperty->Struct->Children->IsValidLowLevel() || StructProperty->Struct->ChildProperties->IsValidLowLevel();
+	}
+	return false;
+}
+
+void FRemoteControlUIModule::CreateSubMenuForChildProperties(const FRCExposesPropertyArgs InPropertyArgs) const
+{
+	FMenuBuilder MenuBuilder(false, nullptr, nullptr, true );
+	MenuBuilder.SetStyle(FRemoteControlExposeMenuStyle::Get().Get(), FRemoteControlExposeMenuStyle::GetStyleSetName());
+
+	TArray<FRCExposesAllPropertiesArgs> ExposeAllArgs;
+	GetAllExposableSubPropertyFromStruct(InPropertyArgs, ExposeAllArgs);
+	MenuBuilder.BeginSection(NAME_None, FText::Format(LOCTEXT("ExposeUnexposeExpansion", "{0}"), InPropertyArgs.GetProperty()->GetDisplayNameText()));
+
+	// Expose and Unexpose ALL section
+	MenuBuilder.AddMenuEntry(
+	LOCTEXT("RCExposeAll", "Expose All"),
+	LOCTEXT("RCExposeAll", "Expose all sub-properties"),
+	FSlateIcon(),
+	FUIAction(FExecuteAction::CreateRaw(this, &FRemoteControlUIModule::OnExposeAll, ExposeAllArgs)),
+	NAME_None,
+	EUserInterfaceActionType::CollapsedButton);
+
+	MenuBuilder.AddMenuEntry(
+		LOCTEXT("RCUnexposeAll", "Unexpose All"),
+		LOCTEXT("RCUnexposeAll", "Unexpose all sub-properties"),
+		FSlateIcon(),
+		FUIAction(FExecuteAction::CreateRaw(this, &FRemoteControlUIModule::OnUnexposeAll, ExposeAllArgs)),
+		NAME_None,
+		EUserInterfaceActionType::CollapsedButton);
+
+	MenuBuilder.EndSection();
+
+	// Sub-Properties section
+	MenuBuilder.BeginSection(NAME_None, LOCTEXT("ExposeUnexposeExpansion", "EXPOSE"));
+
+	for (FRCExposesAllPropertiesArgs PropArgs : ExposeAllArgs)
+	{
+		MenuBuilder.AddMenuEntry(
+			PropArgs.ExposedPropertyLabel,
+			PropArgs.ToolTip,
+			FSlateIcon(),
+			FUIAction(
+				FExecuteAction::CreateRaw(this, &FRemoteControlUIModule::OnToggleExposeSubProperty, PropArgs.PropertyArgs, PropArgs.DesiredName),
+				FCanExecuteAction::CreateRaw(this, &FRemoteControlUIModule::CanToggleExposeProperty, PropArgs.PropertyArgs),
+				FGetActionCheckState::CreateRaw(this, &FRemoteControlUIModule::GetPropertyExposedCheckState, PropArgs.PropertyArgs),
+				FIsActionButtonVisible::CreateRaw(this, &FRemoteControlUIModule::CanToggleExposeProperty, PropArgs.PropertyArgs)),
+			NAME_None,
+			EUserInterfaceActionType::ToggleButton);
+	}
+
+	MenuBuilder.EndSection();
+
+	UE::Slate::FDeprecateVector2DParameter MousePos = FSlateApplication::Get().GetCursorPos();
+	FWidgetPath WindowPath = FSlateApplication::Get().LocateWindowUnderMouse(MousePos, FSlateApplication::Get().GetInteractiveTopLevelWindows(), false, 0);
+
+	FSlateApplication::Get().PushMenu(WindowPath.GetLastWidget(), WindowPath, MenuBuilder.MakeWidget(), MousePos, FPopupTransitionEffect::ContextMenu);
+}
+
+bool FRemoteControlUIModule::HasChildPropertiesExposed(const FRCExposesPropertyArgs& InPropertyArgs) const
+{
+	if (InPropertyArgs.PropertyHandle.IsValid() && InPropertyArgs.PropertyHandle->GetProperty()->IsA<FStructProperty>())
+	{
+		// Check if it has sub-properties exposed
+		FStructProperty* StructProperty = CastFieldChecked<FStructProperty>(InPropertyArgs.PropertyHandle->GetProperty());
+		int32 ChildHandleIndex = 0;
+		for (TFieldIterator<FProperty> It(StructProperty->Struct); It; ++It)
+		{
+			FRCExposesPropertyArgs ChildArgs;
+			ChildArgs.OwnerObject = InPropertyArgs.OwnerObject.Get();
+			ChildArgs.Property = *It;
+			ChildArgs.PropertyPath = (*It)->GetPathName(nullptr);
+			ChildArgs.PropertyHandle = InPropertyArgs.PropertyHandle->GetChildHandle(ChildHandleIndex++);
+			EPropertyExposeStatus ChildStatus = GetPropertyExposeStatus(ChildArgs);
+
+			if (ChildStatus == EPropertyExposeStatus::Exposed)
+			{
+				return true;
+			}
+		}
+	}
+	return false;
+}
+
+void FRemoteControlUIModule::GetAllExposableSubPropertyFromStruct(const FRCExposesPropertyArgs InPropertyArgs, TArray<FRCExposesAllPropertiesArgs>& OutAllProperty) const
+{
+	FStructProperty* StructProperty = CastFieldChecked<FStructProperty>(InPropertyArgs.PropertyHandle->GetProperty());
+	int32 ChildHandleIndex = 0;
+	for (TFieldIterator<FProperty> It(StructProperty->Struct); It; ++It)
+	{
+		FRCExposesPropertyArgs ChildArgs;
+		ChildArgs.OwnerObject = InPropertyArgs.OwnerObject.Get();
+		ChildArgs.Property = *It;
+		ChildArgs.PropertyPath = (*It)->GetPathName(nullptr);
+		ChildArgs.PropertyHandle = InPropertyArgs.PropertyHandle->GetChildHandle(ChildHandleIndex++);
+		// Skip if not valid or special case for FColor and FLinearColor since they can decide to show the Alpha channel or not
+		if (!ChildArgs.IsValid() || (InPropertyArgs.GetProperty()->HasMetaData("HideAlphaChannel") && ChildArgs.Property->GetFName() == FName("A")))
+		{
+			continue;
+		}
+		FText PropertyName = It->GetDisplayNameText();
+		if (StructProperty->GetFName() == USceneComponent::GetRelativeRotationPropertyName() || StructProperty->GetFName() == USceneComponent::GetAbsoluteRotationPropertyName())
+		{
+			PropertyName = FText::FromString(It->GetName());
+		}
+		FText ExposedPropertyLabel = FText::Format(LOCTEXT("RC_ChildProperty", "{0}"), PropertyName);
+		TAttribute<FText> ToolTip = TAttribute<FText>::CreateLambda([this, ChildArgs, PropertyName]()
+		{
+			const EPropertyExposeStatus ChildStatus = GetPropertyExposeStatus(ChildArgs);
+			switch (ChildStatus)
+			{
+				case EPropertyExposeStatus::Exposed:
+					return FText::Format(LOCTEXT("RCUnexpose", "Unexpose {0}"), PropertyName);
+
+				case EPropertyExposeStatus::Unexposed:
+					return FText::Format(LOCTEXT("RCExpose", "Expose {0}"), PropertyName);
+
+				case EPropertyExposeStatus::Unexposable:
+				default:
+					return FText::Format(LOCTEXT("RCUnexposable", "Property {0} cannot be exposed"), PropertyName);
+			}
+		});
+		OutAllProperty.Add(FRCExposesAllPropertiesArgs({ChildArgs, PropertyName, PropertyName.ToString(), ExposedPropertyLabel, ToolTip}));
+	}
+}
+
 TSharedPtr<SRCPanelTreeNode> FRemoteControlUIModule::GenerateEntityWidget(const FGenerateWidgetArgs& Args)
 {
 	TRACE_CPUPROFILER_EVENT_SCOPE(FRemoteControlUIModule::GenerateEntityWidget);
@@ -1141,7 +1356,7 @@ TSharedPtr<SRCPanelTreeNode> FRemoteControlUIModule::GenerateEntityWidget(const
 	if (Args.Preset && Args.Entity)
 	{
 		const UScriptStruct* EntityType = Args.Preset->GetExposedEntityType(Args.Entity->GetId());
-		
+
 		if (FOnGenerateRCWidget* OnGenerateWidget = GenerateWidgetDelegates.Find(const_cast<UScriptStruct*>(EntityType)))
 		{
 			return OnGenerateWidget->Execute(Args);
@@ -1153,6 +1368,13 @@ TSharedPtr<SRCPanelTreeNode> FRemoteControlUIModule::GenerateEntityWidget(const
 
 void FRemoteControlUIModule::UnregisterRemoteControlPanel(SRemoteControlPanel* Panel)
 {
+	if (SharedDetailsPanel.IsValid())
+	{
+		// reset column min width to standard size
+		SharedDetailsPanel->SetRightColumnMinWidth(22);
+		SharedDetailsPanel->ForceRefresh();
+	}
+
 	if (!Panel)
 	{
 		return;
diff --git a/Engine/Plugins/VirtualProduction/RemoteControl/Source/RemoteControlUI/Private/RemoteControlUIModule.h b/Engine/Plugins/VirtualProduction/RemoteControl/Source/RemoteControlUI/Private/RemoteControlUIModule.h
--- a/Engine/Plugins/VirtualProduction/RemoteControl/Source/RemoteControlUI/Private/RemoteControlUIModule.h
+++ b/Engine/Plugins/VirtualProduction/RemoteControl/Source/RemoteControlUI/Private/RemoteControlUIModule.h
@@ -22,7 +22,7 @@ class URemoteControlPreset;
 class FRemoteControlUIModule : public IRemoteControlUIModule
 {
 public:
-	
+
 	FRemoteControlUIModule()
 		: RemoteControlAssetCategoryBit(EAssetTypeCategories::Misc)
 	{
@@ -105,7 +105,7 @@ private:
 	//~ Context menu extenders
 	void RegisterContextMenuExtender();
 	void UnregisterContextMenuExtender();
-	
+
 	//~ Detail row extensions
 	void RegisterDetailRowExtension();
 	void UnregisterDetailRowExtension();
@@ -137,6 +137,26 @@ private:
 	/** Handle clicking the expose button. */
 	void OnToggleExposeProperty(const FRCExposesPropertyArgs InArgs);
 
+	/** Handle clicking the expose SubProperty button. */
+	void OnToggleExposeSubProperty(const FRCExposesPropertyArgs InArgs, const FString InDesiredName = TEXT("")) const;
+
+	void OnToggleExposePropertyWithChild(const FRCExposesPropertyArgs InArgs);
+
+	/** Store the information to expose/unexpose all properties from the parent one */
+	struct FRCExposesAllPropertiesArgs
+	{
+		FRCExposesPropertyArgs PropertyArgs;
+		FText PropName;
+		FString DesiredName;
+		FText ExposedPropertyLabel;
+		TAttribute<FText> ToolTip;
+	};
+	/** Handle clicking the expose all in the SubProperty menu */
+	void OnExposeAll(const TArray<FRCExposesAllPropertiesArgs> InExposeAllArgs) const;
+
+	/** Handle clicking the expose all in the SubProperty menu */
+	void OnUnexposeAll(const TArray<FRCExposesAllPropertiesArgs> InExposeAllArgs) const;
+
 	/** Returns whether a property is exposed, unexposed or unexposable. */
 	EPropertyExposeStatus GetPropertyExposeStatus(const FRCExposesPropertyArgs& InArgs) const;
 
@@ -176,7 +196,7 @@ private:
 
 	/** Returns expose button text based on exposed state */
 	FText GetExposePropertyButtonText(const FRCExposesPropertyArgs InArgs) const;
-	
+
 	/** Attempts to replace the static or skeletal materials with their corressponding overrides. */
 	void TryOverridingMaterials(const FRCExposesPropertyArgs InArgs);
 
@@ -186,6 +206,19 @@ private:
 	/** Attempts to refresh the details panel. */
 	void RefreshPanels();
 
+	/** Check if a sub menu should be created when ctrl+click on the eyeball icon in the details view */
+	bool ShouldCreateSubMenuForChildProperties(const FRCExposesPropertyArgs InPropertyArgs) const;
+
+	/** Check if the property has child properties that can be exposed */
+	bool HasChildProperties(const FProperty* InProperty) const;
+
+	/** Create a sub menu for each sub property */
+	void CreateSubMenuForChildProperties(const FRCExposesPropertyArgs InPropertyArgs) const;
+
+	/** Check if any sub-property is exposed */
+	bool HasChildPropertiesExposed(const FRCExposesPropertyArgs& InPropertyArgs) const;
+
+	void GetAllExposableSubPropertyFromStruct(const FRCExposesPropertyArgs InPropertyArgs, TArray<FRCExposesAllPropertiesArgs>& OutAllProperty) const;
 private:
 	/** The custom actions added to the actor context menu. */
 	TSharedPtr<class FRemoteControlPresetActions> RemoteControlPresetActions;
@@ -198,10 +231,10 @@ private:
 
 	/** Holds a weak ptr to the active control panel. */
 	TWeakPtr<SRemoteControlPanel> WeakActivePanel;
-	
+
 	/** Holds a weak ptr to the owner tree node of the active details panel. */
 	TWeakPtr<IDetailTreeNode> WeakDetailsTreeNode;
-	
+
 	/** Holds a shared ptr to the global details panel. */
 	TSharedPtr<IDetailsView> SharedDetailsPanel;
 
diff --git a/Engine/Plugins/VirtualProduction/RemoteControl/Source/RemoteControlUI/Private/UI/RemoteControlExposeMenuStyle.cpp b/Engine/Plugins/VirtualProduction/RemoteControl/Source/RemoteControlUI/Private/UI/RemoteControlExposeMenuStyle.cpp
new file mode 100644
--- /dev/null
+++ b/Engine/Plugins/VirtualProduction/RemoteControl/Source/RemoteControlUI/Private/UI/RemoteControlExposeMenuStyle.cpp
@@ -0,0 +1,156 @@
+﻿// Copyright Epic Games, Inc. All Rights Reserved.
+
+#include "RemoteControlExposeMenuStyle.h"
+#include "Interfaces/IPluginManager.h"
+#include "Styling/AppStyle.h"
+#include "Styling/CoreStyle.h"
+#include "Styling/SlateStyle.h"
+#include "Styling/SlateStyleRegistry.h"
+#include "Styling/SlateTypes.h"
+#include "Styling/StarshipCoreStyle.h"
+#include "Styling/StyleColors.h"
+
+#define IMAGE_PLUGIN_BRUSH( RelativePath, ... ) FSlateImageBrush(FRemoteControlExposeMenuStyle::InContent(RelativePath, ".png" ), __VA_ARGS__)
+#define IMAGE_BRUSH(RelativePath, ...) FSlateImageBrush(StyleSet->RootToContentDir(RelativePath, TEXT(".png")), __VA_ARGS__)
+#define IMAGE_BRUSH_SVG( RelativePath, ... ) FSlateVectorImageBrush(StyleSet->RootToContentDir(RelativePath, TEXT(".svg")), __VA_ARGS__)
+#define IMAGE_PLUGIN_BRUSH_SVG( RelativePath, ... ) FSlateVectorImageBrush(FRemoteControlExposeMenuStyle::InContent(RelativePath, ".svg"), __VA_ARGS__)
+#define BOX_BRUSH(RelativePath, ...) FSlateBoxBrush(StyleSet->RootToContentDir(RelativePath, TEXT(".png")), __VA_ARGS__)
+#define CORE_BOX_BRUSH( RelativePath, ... ) FSlateBoxBrush(StyleSet->RootToCoreContentDir(RelativePath, TEXT(".png") ), __VA_ARGS__)
+#define BOX_PLUGIN_BRUSH( RelativePath, ... ) FSlateBoxBrush(FRemoteControlExposeMenuStyle::InContent( RelativePath, ".png" ), __VA_ARGS__)
+#define CORE_FONT(...) FSlateFontInfo(FCoreStyle::GetDefaultFont(), __VA_ARGS__)
+
+TSharedPtr<FSlateStyleSet> FRemoteControlExposeMenuStyle::StyleSet;
+
+void FRemoteControlExposeMenuStyle::Initialize()
+{
+	if (StyleSet.IsValid())
+	{
+		return;
+	}
+
+	const FVector2D Icon8x8(8.0f, 8.0f);
+	const FVector2D Icon3x10(3.0f, 10.0f);
+	const FVector2D Icon12x12(12.0f, 12.0f);
+	const FVector2D Icon16x16(16.0f, 16.0f);
+	const FVector2D Icon20x20(20.0f, 20.0f);
+	const FVector2D Icon32x32(32.0f);
+	const FVector2D Icon48x48(48.0f, 48.0f);
+	const FVector2D Icon28x14(28.0f, 14.0f);
+	const FVector2D Icon64x64(64.0f, 64.0f);
+
+	StyleSet = MakeShared<FSlateStyleSet>(GetStyleSetName());
+	
+	// Remote Control Expose Menu
+	StyleSet->Set("RemoteControlExposeMenu.WidgetBorder", new FSlateRoundedBoxBrush(FStyleColors::Input, 5.0f));
+
+	StyleSet->Set("RemoteControlExposeMenu.SpinBox", FSpinBoxStyle()
+	.SetBackgroundBrush(FSlateRoundedBoxBrush(FStyleColors::Input, CoreStyleConstants::InputFocusRadius, FStyleColors::Secondary, CoreStyleConstants::InputFocusThickness))
+	.SetHoveredBackgroundBrush(FSlateRoundedBoxBrush(FStyleColors::Input, CoreStyleConstants::InputFocusRadius, FStyleColors::Hover, CoreStyleConstants::InputFocusThickness))
+	.SetActiveFillBrush(FSlateRoundedBoxBrush(FStyleColors::Hover, CoreStyleConstants::InputFocusRadius, FLinearColor::Transparent, CoreStyleConstants::InputFocusThickness))
+	.SetInactiveFillBrush(FSlateRoundedBoxBrush(FStyleColors::Secondary, CoreStyleConstants::InputFocusRadius, FLinearColor::Transparent, CoreStyleConstants::InputFocusThickness))
+	.SetArrowsImage(FSlateNoResource())
+	.SetForegroundColor(FStyleColors::ForegroundHover)
+	.SetTextPadding(FMargin(10.f, 3.5f, 10.f, 4.f)));
+
+	StyleSet->Set("RemoteControlExposeMenu.Background", new FSlateColorBrush(FStyleColors::Dropdown));
+	StyleSet->Set("RemoteControlExposeMenu.Outline", new FSlateRoundedBoxBrush(FStyleColors::Transparent, 0.0f, FStyleColors::DropdownOutline, 1.f));
+
+	const FButtonStyle MenuButtonStyle = FAppStyle::Get().GetWidgetStyle<FButtonStyle>("Menu.Button");
+	StyleSet->Set("RemoteControlExposeMenu.Button", MenuButtonStyle);
+
+	const FTextBlockStyle MenuTextBlockStyle = FAppStyle::Get().GetWidgetStyle<FTextBlockStyle>("NormalText");
+	StyleSet->Set("RemoteControlExposeMenu.Label", MenuTextBlockStyle);
+
+	StyleSet->Set("RemoteControlExposeMenu.MenuBar.Padding", FMargin(12, 4));
+	StyleSet->Set("RemoteControlExposeMenu.MenuIconSize", 16.f);
+
+	// DetailsView icons
+	StyleSet->Set("RemoteControlExposeMenu.NoBrush", new FSlateNoResource());
+	StyleSet->Set("RemoteControlExposeMenu.VisibleAndVisibleChildren", new IMAGE_PLUGIN_BRUSH("Icons/Visible_Children_Exposed", Icon16x16));
+	StyleSet->Set("RemoteControlExposeMenu.Visible", new IMAGE_PLUGIN_BRUSH("Icons/Visible", Icon16x16));
+	StyleSet->Set("RemoteControlExposeMenu.HiddenAndVisibleChildren", new IMAGE_PLUGIN_BRUSH("Icons/Hidden_Children_Exposed", Icon16x16));
+	StyleSet->Set("RemoteControlExposeMenu.Hidden", new IMAGE_PLUGIN_BRUSH("Icons/Hidden", Icon16x16));
+	// StyleSet->Set("RemoteControlExposeMenu.Expand", new IMAGE_PLUGIN_BRUSH("Icons/Expand", Icon3x10));
+	StyleSet->Set("RemoteControlExposeMenu.Expand", new IMAGE_PLUGIN_BRUSH("Icons/ellipsis_12x", Icon12x12));
+
+	FCheckBoxStyle MenuCheckBoxStyle = FAppStyle::Get().GetWidgetStyle<FCheckBoxStyle>("Menu.CheckBox");
+	StyleSet->Set("RemoteControlExposeMenu.CheckBox", MenuCheckBoxStyle);
+
+	FCheckBoxStyle MenuCheckStyle = FAppStyle::Get().GetWidgetStyle<FCheckBoxStyle>("ToggleButtonCheckbox");
+	MenuCheckStyle.SetPadding(FMargin(10.f, 0.f));
+	MenuCheckStyle.SetCheckedImage(*FAppStyle::GetBrush("Level.VisibleIcon16x"));
+	MenuCheckStyle.SetUncheckedImage(*FAppStyle::GetBrush("Level.NotVisibleIcon16x"));
+	MenuCheckStyle.SetCheckedHoveredImage(*FAppStyle::GetBrush("Level.VisibleIcon16x"));
+	MenuCheckStyle.SetUncheckedHoveredImage(*FAppStyle::GetBrush("Level.NotVisibleIcon16x"));
+	MenuCheckStyle.SetCheckedPressedImage(*FAppStyle::GetBrush("Level.VisibleIcon16x"));
+	MenuCheckStyle.SetUncheckedPressedImage(*FAppStyle::GetBrush("Level.NotVisibleIcon16x"));
+	MenuCheckStyle.SetHoveredForegroundColor(FStyleColors::ForegroundHover);
+	MenuCheckStyle.SetCheckedHoveredForegroundColor(FStyleColors::ForegroundHover);
+	MenuCheckStyle.SetForegroundColor(FStyleColors::ForegroundHover);
+	StyleSet->Set("RemoteControlExposeMenu.Check", MenuCheckStyle);
+
+	const FCheckBoxStyle MenuRadioButtonStyle = FAppStyle::Get().GetWidgetStyle<FCheckBoxStyle>("Menu.RadioButton");
+	StyleSet->Set("RemoteControlExposeMenu.RadioButton", MenuRadioButtonStyle);
+
+	const FCheckBoxStyle MenuToggleButtonStyle = FAppStyle::Get().GetWidgetStyle<FCheckBoxStyle>("Menu.ToggleButton");
+	StyleSet->Set("RemoteControlExposeMenu.ToggleButton", MenuToggleButtonStyle);
+
+	StyleSet->Set("RemoteControlExposeMenu.Keybinding", FTextBlockStyle(MenuTextBlockStyle).SetFont(FStyleFonts::Get().Small));
+
+	const FEditableTextBoxStyle NormalEditableTextBoxStyle = FAppStyle::GetWidgetStyle<FEditableTextBoxStyle>("NormalEditableTextBox");
+	StyleSet->Set("RemoteControlExposeMenu.EditableText", FEditableTextBoxStyle(NormalEditableTextBoxStyle).SetFont(FStyleFonts::Get().Normal));
+	StyleSet->Set("MultiboxHookColor", FLinearColor(0.f, 1.f, 0.f, 1.f));
+
+	const FMargin MenuBlockPadding(12.0f, 1.0f, 5.0f, 1.0f);
+	StyleSet->Set("RemoteControlExposeMenu.Block.IndentedPadding", MenuBlockPadding + FMargin(18.0f, 0, 0, 0));
+	StyleSet->Set("RemoteControlExposeMenu.Block.Padding", MenuBlockPadding);
+
+	StyleSet->Set("RemoteControlExposeMenu.SubMenuIndicator", new IMAGE_BRUSH_SVG("Starship/Common/chevron-right", Icon16x16, FStyleColors::Foreground));
+	StyleSet->Set("RemoteControlExposeMenu.Separator", new FSlateColorBrush(FStyleColors::White25));
+	StyleSet->Set("RemoteControlExposeMenu.Separator.Padding", FMargin(12.0f, 6.f, 12.0f, 6.f));
+
+	FSlateFontInfo XSFont(CORE_FONT(7, "Bold"));
+	XSFont.LetterSpacing =  250;
+	const FTextBlockStyle SmallButtonText = FAppStyle::Get().GetWidgetStyle<FTextBlockStyle>("SmallButtonText");
+	StyleSet->Set("RemoteControlExposeMenu.Heading", FTextBlockStyle(SmallButtonText)
+	.SetFont(XSFont)
+	.SetColorAndOpacity(FStyleColors::White25));
+	StyleSet->Set("RemoteControlExposeMenu.Heading.Padding", FMargin(12.0f, 6.f, 12.f, 6.f));
+
+	FSlateStyleRegistry::RegisterSlateStyle(*StyleSet.Get());
+}
+
+void FRemoteControlExposeMenuStyle::Shutdown()
+{
+	if (StyleSet.IsValid())
+	{
+		FSlateStyleRegistry::UnRegisterSlateStyle(*StyleSet.Get());
+		ensure(StyleSet.IsUnique());
+		StyleSet.Reset();
+	}
+}
+
+TSharedPtr<ISlateStyle> FRemoteControlExposeMenuStyle::Get()
+{
+	return StyleSet;
+}
+
+FName FRemoteControlExposeMenuStyle::GetStyleSetName()
+{
+	static const FName RemoteControlPanelStyleName(TEXT("RemoteControlExposeMenu"));
+	return RemoteControlPanelStyleName;
+}
+
+FString FRemoteControlExposeMenuStyle::InContent(const FString& RelativePath, const ANSICHAR* Extension)
+{
+	static const FString ContentDir = IPluginManager::Get().FindPlugin(TEXT("RemoteControl"))->GetBaseDir() + TEXT("/Resources");
+	return (ContentDir / RelativePath) + Extension;
+}
+
+#undef IMAGE_PLUGIN_BRUSH
+#undef IMAGE_BRUSH
+#undef IMAGE_BRUSH_SVG
+#undef BOX_BRUSH
+#undef CORE_BOX_BRUSH
+#undef BOX_PLUGIN_BRUSH
+#undef CORE_FONT
diff --git a/Engine/Plugins/VirtualProduction/RemoteControl/Source/RemoteControlUI/Private/UI/RemoteControlExposeMenuStyle.h b/Engine/Plugins/VirtualProduction/RemoteControl/Source/RemoteControlUI/Private/UI/RemoteControlExposeMenuStyle.h
new file mode 100644
--- /dev/null
+++ b/Engine/Plugins/VirtualProduction/RemoteControl/Source/RemoteControlUI/Private/UI/RemoteControlExposeMenuStyle.h
@@ -0,0 +1,40 @@
+﻿// Copyright Epic Games, Inc. All Rights Reserved.
+
+#pragma once
+#include "Templates/SharedPointerFwd.h"
+
+class ISlateStyle;
+class FSlateStyleSet;
+class FName;
+class FString;
+
+class FRemoteControlExposeMenuStyle
+{
+public:
+	/**
+	 * Initialize the style set.
+	 */
+	static void Initialize();
+
+	/**
+	 * Shutdown the style set.
+	 */
+	static void Shutdown();
+
+	/**
+	 * Static getter for the style set.
+	 */
+	static TSharedPtr<ISlateStyle> Get();
+
+	/**
+	 * Static getter for the style set name.
+	 */
+	static FName GetStyleSetName();
+private:
+	/** Get a path to the module's resources folder. */
+	static FString InContent(const FString& RelativePath, const ANSICHAR* Extension);
+
+private:
+	/** Holds the style set. */
+	static TSharedPtr<FSlateStyleSet> StyleSet;
+};
diff --git a/Engine/Plugins/VirtualProduction/RemoteControl/Source/RemoteControlUI/Private/UI/SRCPanelExposedField.cpp b/Engine/Plugins/VirtualProduction/RemoteControl/Source/RemoteControlUI/Private/UI/SRCPanelExposedField.cpp
--- a/Engine/Plugins/VirtualProduction/RemoteControl/Source/RemoteControlUI/Private/UI/SRCPanelExposedField.cpp
+++ b/Engine/Plugins/VirtualProduction/RemoteControl/Source/RemoteControlUI/Private/UI/SRCPanelExposedField.cpp
@@ -4,6 +4,7 @@
 
 #include "Algo/Transform.h"
 #include "Components/MeshComponent.h"
+#include "Components/SceneComponent.h"
 #include "EditorFontGlyphs.h"
 #include "Factories/IRCDefaultValueFactory.h"
 #include "IDetailTreeNode.h"
@@ -22,6 +23,8 @@
 #include "SRCPanelTreeNode.h"
 #include "SResetToDefaultPropertyEditor.h"
 #include "ScopedTransaction.h"
+#include "UnrealEdGlobals.h"
+#include "Editor/UnrealEdEngine.h"
 #include "Styling/AppStyle.h"
 #include "Styling/RemoteControlStyles.h"
 #include "Styling/SlateBrush.h"
@@ -32,7 +35,6 @@
 #include "Widgets/Layout/SBox.h"
 #include "Widgets/Masks/SRCProtocolMask.h"
 #include "Widgets/SNullWidget.h"
-#include "Widgets/Text/SInlineEditableTextBlock.h"
 #include "Widgets/Text/STextBlock.h"
 
 #define LOCTEXT_NAMESPACE "RemoteControlPanel"
@@ -66,6 +68,56 @@ namespace ExposedFieldUtils
 		OutValueWidget = FieldWidget;
 		OutEditConditionWidget = NodeWidgets.EditConditionWidget;
 	}
+
+	bool IsFieldPathMatchingWithNodePath(const TSharedPtr<IDetailTreeNode>& InNode, const TSharedPtr<FRemoteControlField>& InField)
+	{
+		if (!InField.IsValid())
+		{
+			// if not able to check it by default is true
+			return true;
+		}
+		TSharedPtr<IPropertyHandle> ParentHandle = InNode->CreatePropertyHandle();
+		const FStringView ParentPath = ParentHandle->GetPropertyPath();
+		FString PropPath = InField->FieldPathInfo.ToPathPropertyString();
+		PropPath = PropPath.Replace(TEXT("."), TEXT("->"));
+		const FName ParentPathName = FName(ParentPath);
+		return ParentPathName == FName(PropPath);
+	}
+
+	bool TryCreateWidgetForChildProperty(const TSharedPtr<IDetailTreeNode>& InNode, const TSharedPtr<FRemoteControlField>& InField, TSharedPtr<IPropertyHandle>& OutChildHandle, TSharedPtr<SWidget>& OutValueWidget, TSharedPtr<SWidget>& OutEditConditionWidget)
+	{
+		const TSharedPtr<IPropertyHandle> ParentHandle = InNode->CreatePropertyHandle();
+		uint32 ChildNum;
+		ParentHandle->GetNumChildren(ChildNum);
+		const FStringView ParentPath = ParentHandle->GetPropertyPath();
+		FString FieldPath = InField->FieldPathInfo.ToPathPropertyString();
+		FieldPath = FieldPath.Replace(TEXT("."), TEXT("->"));
+		const FName FieldPathName = FName(FieldPath);
+		const FName ParentPathName = FName(ParentPath);
+		// Need to check if it is a child of the location to update the arrows widget position in the scene when moving
+		const bool bIsLocation = ParentPathName == USceneComponent::GetRelativeLocationPropertyName() || ParentPathName == USceneComponent::GetAbsoluteLocationPropertyName();
+
+		for (uint32 ChildHandleIndex = 0; ChildHandleIndex < ChildNum; ++ChildHandleIndex)
+		{
+			const TSharedPtr<IPropertyHandle> ChildHandle = ParentHandle->GetChildHandle(ChildHandleIndex);
+			const FName ChildPathName = FName(ChildHandle->GetPropertyPath());
+			if (ChildPathName == FieldPathName)
+			{
+				ExposedFieldUtils::CreateWidgets(InNode, OutValueWidget, OutEditConditionWidget);
+				OutChildHandle = ChildHandle;
+				if (bIsLocation)
+				{
+					OutChildHandle->SetOnPropertyValueChanged(FSimpleDelegate::CreateLambda([]()
+					{
+						GUnrealEd->UpdatePivotLocationForSelection();
+					}));
+				}
+				OutValueWidget = OutChildHandle->CreatePropertyValueWidget();
+				return true;
+			}
+		}
+		return false;
+	}
 }
 
 TSharedPtr<SRCPanelTreeNode> SRCPanelExposedField::MakeInstance(const FGenerateWidgetArgs& Args)
@@ -82,7 +134,7 @@ void SRCPanelExposedField::Construct(const FArguments& InArgs, TWeakPtr<FRemoteC
 
 	ColumnSizeData = MoveTemp(InColumnSizeData);
 	WidgetRegistry = MoveTemp(InWidgetRegistry);
-	
+
 	HighlightText = InArgs._HighlightText;
 
 	ResetButtonWidget = SNullWidget::NullWidget;
@@ -90,10 +142,10 @@ void SRCPanelExposedField::Construct(const FArguments& InArgs, TWeakPtr<FRemoteC
 	if (TSharedPtr<FRemoteControlField> FieldPtr = WeakField.Pin())
 	{
 		Initialize(FieldPtr->GetId(), InArgs._Preset.Get(), InArgs._LiveMode);
-	
+
 		CachedLabel = FieldPtr->GetLabel();
 		EntityId = FieldPtr->GetId();
-		
+
 		if (FieldPtr->FieldType == EExposedFieldType::Property)
 		{
 			ConstructPropertyWidget();
@@ -131,7 +183,7 @@ EExposedFieldType SRCPanelExposedField::GetFieldType() const
 	{
 		return Field->FieldType;
 	}
-	
+
 	return EExposedFieldType::Invalid;
 }
 
@@ -227,7 +279,7 @@ TSharedRef<SWidget> SRCPanelExposedField::GetProtocolWidget(const FName ForColum
 			}
 		}
 	}
-	
+
 	return SRCPanelTreeNode::GetProtocolWidget(ForColumnName, InProtocolName);
 }
 
@@ -277,11 +329,11 @@ const bool SRCPanelExposedField::SupportsProtocol(const FName& InProtocolName) c
 void SRCPanelExposedField::Refresh()
 {
 	SRCPanelExposedEntity::Refresh();
-	
+
 	if (TSharedPtr<FRemoteControlField> Field = WeakField.Pin())
 	{
 		CachedLabel = Field->GetLabel();
-		
+
 		if (Field->FieldType == EExposedFieldType::Property)
 		{
 			ConstructPropertyWidget();
@@ -342,27 +394,37 @@ TSharedRef<SWidget> SRCPanelExposedField::ConstructWidget()
 
 				if (TSharedPtr<IDetailTreeNode> Node = Registry->GetObjectTreeNode(Object, Field->FieldPathInfo.ToPathPropertyString(), ERCFindNodeMethod::Path))
 				{
-					TArray<TSharedRef<IDetailTreeNode>> ChildNodes;
-					Node->GetChildren(ChildNodes);
-					ChildWidgets.Reset(ChildNodes.Num());
+					TSharedPtr<SWidget> ValueWidget = SNullWidget::NullWidget;
+					TSharedPtr<SWidget> EditConditionWidget = SNullWidget::NullWidget;
+					TSharedPtr<IPropertyHandle> PropertyHandle;
+					bool bIsChildPropertyWidgetCreated = false;
 
-					for (const TSharedRef<IDetailTreeNode>& ChildNode : ChildNodes)
+					if (!ExposedFieldUtils::IsFieldPathMatchingWithNodePath(Node, Field))
 					{
-						ChildWidgets.Add(SNew(SRCPanelFieldChildNode, ChildNode, ColumnSizeData));
+						bIsChildPropertyWidgetCreated = ExposedFieldUtils::TryCreateWidgetForChildProperty(Node, Field, PropertyHandle, ValueWidget, EditConditionWidget);
 					}
 
-					TSharedPtr<SHorizontalBox> FieldWidget = SNew(SHorizontalBox);
-					TSharedPtr<SWidget> ValueWidget = SNullWidget::NullWidget;
-					TSharedPtr<SWidget> EditConditionWidget = SNullWidget::NullWidget;
+					if (!bIsChildPropertyWidgetCreated)
+					{
+						TArray<TSharedRef<IDetailTreeNode>> ChildNodes;
+						Node->GetChildren(ChildNodes);
+						ChildWidgets.Reset(ChildNodes.Num());
 
-					ExposedFieldUtils::CreateWidgets(Node, ValueWidget, EditConditionWidget);
+						for (const TSharedRef<IDetailTreeNode>& ChildNode : ChildNodes)
+						{
+							ChildWidgets.Add(SNew(SRCPanelFieldChildNode, ChildNode, ColumnSizeData));
+						}
+
+						ExposedFieldUtils::CreateWidgets(Node, ValueWidget, EditConditionWidget);
+
+						PropertyHandle = Node->CreatePropertyHandle();
+					}
 
+					TSharedPtr<SHorizontalBox> FieldWidget = SNew(SHorizontalBox);
 					FieldWidget->AddSlot()
 					[
 						ValueWidget.ToSharedRef()
 					];
-
-					TSharedPtr<IPropertyHandle> PropertyHandle = Node->CreatePropertyHandle();
 					if (PropertyHandle && PropertyHandle->IsValidHandle() && (PropertyHandle->GetParentHandle()->GetPropertyDisplayName().ToString() == FString("Transform")
 						|| PropertyHandle->GetOuterBaseClass() == UMaterialInstanceDynamic::StaticClass()))
 					{
@@ -376,10 +438,9 @@ TSharedRef<SWidget> SRCPanelExposedField::ConstructWidget()
 						}
 
 						TWeakPtr<SRCPanelExposedField> WeakFieldPtr = SharedThis(this);
-						auto IsVisible = [WeakFieldPtr, Node]()
+						auto IsVisible = [WeakFieldPtr, PropertyHandle]()
 						{
 							TSharedPtr<SRCPanelExposedField> FieldPtr = WeakFieldPtr.Pin();
-							TSharedPtr<IPropertyHandle> PropertyHandle = Node->CreatePropertyHandle();
 							if (FieldPtr && PropertyHandle && PropertyHandle->IsValidHandle())
 							{
 								void* DataPtr;
@@ -392,7 +453,7 @@ TSharedRef<SWidget> SRCPanelExposedField::ConstructWidget()
 							}
 							return EVisibility::Hidden;
 						};
-						
+
 						ResetButtonWidget = SNew(SBox)
 							.HAlign(HAlign_Center)
 							.VAlign(VAlign_Center)
@@ -461,10 +522,10 @@ TSharedRef<SWidget> SRCPanelExposedField::ConstructWidget()
 
 					FText ErrorText = FText::Format(LOCTEXT("ExposedPropertyInvalidErrorMessage", "Could not find property {0} on object {1}"), FText::FromString(PropertyName), FText::FromString(TargetName));
 					IRemoteControlModule::BroadcastError(ErrorText.ToString());
-					
+
 					return MakeFieldWidget(CreateInvalidWidget(ErrorText));
 				}
-			} 
+			}
 		}
 
 		return MakeFieldWidget(CreateInvalidWidget());
@@ -484,7 +545,7 @@ TSharedRef<SWidget> SRCPanelExposedField::MakeFieldWidget(const TSharedRef<SWidg
 		bool bIsEditorOnly = false,
             bIsEditableInPackaged = true,
             bIsCallableInPackaged = true;
-		
+
 		if (const TSharedPtr<FRemoteControlField> RCField = GetRemoteControlField().Pin())
 		{
 			bIsEditorOnly = RCField->IsEditorOnly();
@@ -513,10 +574,10 @@ TSharedRef<SWidget> SRCPanelExposedField::MakeFieldWidget(const TSharedRef<SWidg
 		{
 			Builder.AppendLine(LOCTEXT("NotCallableInPackagedWarning", "This function will not be callable in packaged projects."));
 		}
-		
+
 		WarningMessage = Builder.ToText();
 	}
-	
+
 	return CreateEntityWidget(InWidget, ResetButtonWidget.ToSharedRef(), WarningMessage, EditConditionWidget);
 }
 
@@ -564,7 +625,7 @@ void SRCPanelExposedField::ConstructFunctionWidget()
 				ChildSlot.AttachWidget(MakeFieldWidget(ConstructCallFunctionButton()));
 				ChildWidgets = ChildNodes;
 			}
-			
+
 
 			return;
 		}
@@ -576,16 +637,16 @@ void SRCPanelExposedField::ConstructFunctionWidget()
 TSharedRef<SWidget> SRCPanelExposedField::ConstructResetToDefaultWidget(UObject* InObject, TSharedPtr<IPropertyHandle> InPropertyHandle)
 {
 	TSharedPtr<SResetToDefaultPropertyEditor> ResetDefaultProperty = SNew(SResetToDefaultPropertyEditor, InPropertyHandle);
-	
+
 	const TAttribute<bool> CanResetToDefaultValue = TAttribute<bool>::Create(TAttribute<bool>::FGetter::CreateLambda([this, InObject, InPropertyHandle]()
 		{
 			IRemoteControlModule& RemoteControlModule = IRemoteControlModule::Get();
-		
+
 			if (RemoteControlModule.HasDefaultValueCustomization(InObject, InPropertyHandle->GetProperty()))
 			{
 				FRCResetToDefaultArgs Args;
 				Args.Property = InPropertyHandle->GetProperty();
-				
+
 				if (InPropertyHandle->IsValidHandle())
 				{
 					Args.Path = InPropertyHandle->GetPropertyPath();
@@ -608,23 +669,23 @@ TSharedRef<SWidget> SRCPanelExposedField::ConstructResetToDefaultWidget(UObject*
 			{
 				FRCResetToDefaultArgs Args;
 				Args.Property = InPropertyHandle->GetProperty();
-				
+
 				if (InPropertyHandle->IsValidHandle())
 				{
 					Args.Path = InPropertyHandle->GetPropertyPath();
 					Args.ArrayIndex = InPropertyHandle->GetArrayIndex();
 				}
-				
+
 				Args.bCreateTransaction = true;
 				RemoteControlModule.ResetToDefaultValue(InObject, Args);
-			
+
 				return;
 			}
 
 			InPropertyHandle->ResetToDefault();
 		}
 	);
-	
+
 	const FResetToDefaultOverride ResetValueOverride = FResetToDefaultOverride::Create(CanResetToDefaultValue, PerformValueReset);
 
 	if (UMeshComponent* MeshComponent = Cast<UMeshComponent>(InObject))
@@ -642,11 +703,11 @@ TSharedRef<SWidget> SRCPanelExposedField::ConstructResetToDefaultWidget(UObject*
 			Args.bCreateTransaction = false;
 			RemoteControlModule.ResetToDefaultValue(InObject, Args);
 		}
-		
+
 	}
-	
+
 	ResetDefaultProperty.Reset();
-	
+
 	ResetDefaultProperty = SNew(SResetToDefaultPropertyEditor, InPropertyHandle)
 		.CustomResetToDefault(ResetValueOverride);
 
diff --git a/Engine/Plugins/VirtualProduction/RemoteControl/Source/RemoteControlUI/Private/UI/SRemoteControlPanel.cpp b/Engine/Plugins/VirtualProduction/RemoteControl/Source/RemoteControlUI/Private/UI/SRemoteControlPanel.cpp
--- a/Engine/Plugins/VirtualProduction/RemoteControl/Source/RemoteControlUI/Private/UI/SRemoteControlPanel.cpp
+++ b/Engine/Plugins/VirtualProduction/RemoteControl/Source/RemoteControlUI/Private/UI/SRemoteControlPanel.cpp
@@ -152,7 +152,7 @@ public:
 			{
 				return false;
 			}
-			
+
 			if (InKeyEvent.GetKey() == EKeys::Delete)
 			{
 				PinnedOwner->DeleteEntity();
@@ -165,7 +165,7 @@ public:
 				return true;
 			}
 		}
-		
+
 		return false;
 	}
 
@@ -240,7 +240,7 @@ public:
 					[
 						IconWidget
 					]
-					
+
 					// Label Text
 					+ SHorizontalBox::Slot()
 					.AutoWidth()
@@ -304,7 +304,7 @@ private:
 
 		return ECheckBoxState::Undetermined;
 	}
-	
+
 	/**
 	 * Called by Slate when this tool bar button's button is clicked
 	 */
@@ -390,7 +390,7 @@ void SRemoteControlPanel::Construct(const FArguments& InArgs, URemoteControlPres
 	FilterComboButton->AddMetadata(FilterComboButtonMetaData.ToSharedRef());
 
 	EntityProtocolDetails = SNew(SBox);
-	
+
 	EntityList = SNew(SRCPanelExposedEntitiesList, Preset.Get(), WidgetRegistry)
 		.ExposeActorsComboButton(CreateExposeActorsButton())
 		.ExposeFunctionsComboButton(CreateExposeFunctionsButton())
@@ -545,7 +545,7 @@ void SRemoteControlPanel::Construct(const FArguments& InArgs, URemoteControlPres
 	TSharedRef<SRCMajorPanel> LogicPanel = SNew(SRCMajorPanel)
 		.HeaderLabel(LOCTEXT("LogicPanelHeader", "LOGIC"))
 		.EnableFooter(false);
-	
+
 	// Controllers with Behaviours panel
 	TSharedRef<SRCMajorPanel> ControllersAndBehavioursPanel = SNew(SRCMajorPanel)
 		.EnableHeader(false)
@@ -559,7 +559,7 @@ void SRemoteControlPanel::Construct(const FArguments& InArgs, URemoteControlPres
 		.Visibility_Lambda([this] { return !bIsInLiveMode && bIsLogicPanelEnabled && (ActivePanel == ERCPanels::RCP_Properties || ActivePanel == ERCPanels::RCP_None) ? EVisibility::Visible : EVisibility::Collapsed; });
 
 	ControllersAndBehavioursPanel->AddPanel(ControllerPanel.ToSharedRef(), 0.8f);
-	
+
 	constexpr bool bResizable = false;
 	ControllersAndBehavioursPanel->AddPanel(BehaviourPanel.ToSharedRef(), 0.f, bResizable);
 
@@ -575,7 +575,7 @@ void SRemoteControlPanel::Construct(const FArguments& InArgs, URemoteControlPres
 		.Orientation(Orient_Horizontal)
 		.HitDetectionSplitterHandleSize(RCPanelStyle->SplitterHandleSize)
 		.PhysicalSplitterHandleSize(RCPanelStyle->SplitterHandleSize)
-	
+
 		// Panel Drawer
 		+SSplitter::Slot()
 		.Value(0.05f)
@@ -607,9 +607,9 @@ void SRemoteControlPanel::Construct(const FArguments& InArgs, URemoteControlPres
 				+ SWidgetSwitcher::Slot()
 				[
 					LogicPanel
-				]	
+				]
 			]
-		
+
 			// 2. Exposed entities List with Logic Editor (Index : 1).
 			+ SWidgetSwitcher::Slot()
 			[
@@ -631,7 +631,7 @@ void SRemoteControlPanel::Construct(const FArguments& InArgs, URemoteControlPres
 					[
 						EntityList.ToSharedRef()
 					]
-			
+
 					+SSplitter::Slot()
 					.Value(0.6)
 					[
@@ -644,8 +644,8 @@ void SRemoteControlPanel::Construct(const FArguments& InArgs, URemoteControlPres
 			+ SWidgetSwitcher::Slot()
 			[
 				SNew(SSplitter)
-				.Orientation(EOrientation::Orient_Vertical)						
-				
+				.Orientation(EOrientation::Orient_Vertical)
+
 				// Exposed entities List
 				+ SSplitter::Slot()
 				.Value(TreeBindingSplitRatioTop)
@@ -659,7 +659,7 @@ void SRemoteControlPanel::Construct(const FArguments& InArgs, URemoteControlPres
 				[
 					EntityList.ToSharedRef()
 				]
-					
+
 				// Entity Details
 				+ SSplitter::Slot()
 				.Value(TreeBindingSplitRatioBottom)
@@ -667,13 +667,13 @@ void SRemoteControlPanel::Construct(const FArguments& InArgs, URemoteControlPres
 					CreateEntityDetailsView()
 				]
 			]
-				
+
 			// 4. Properties with Protocols (Index : 3).
 			+ SWidgetSwitcher::Slot()
 			[
 				SNew(SSplitter)
-				.Orientation(EOrientation::Orient_Vertical)						
-				
+				.Orientation(EOrientation::Orient_Vertical)
+
 				// Exposed entities List
 				+ SSplitter::Slot()
 				.Value(TreeBindingSplitRatioTop)
@@ -687,7 +687,7 @@ void SRemoteControlPanel::Construct(const FArguments& InArgs, URemoteControlPres
 				[
 					EntityList.ToSharedRef()
 				]
-			
+
 				// Protocol Details
 				+ SSplitter::Slot()
 				.Value(TreeBindingSplitRatioBottom)
@@ -695,7 +695,7 @@ void SRemoteControlPanel::Construct(const FArguments& InArgs, URemoteControlPres
 					EntityProtocolDetails.ToSharedRef()
 				]
 			]
-					
+
 			// 5. Output Log (Index : 4).
 			+ SWidgetSwitcher::Slot()
 			[
@@ -706,15 +706,15 @@ void SRemoteControlPanel::Construct(const FArguments& InArgs, URemoteControlPres
 			+ SWidgetSwitcher::Slot()
 			[
 				SNew(SSplitter)
-				.Orientation(EOrientation::Orient_Horizontal)						
-				
+				.Orientation(EOrientation::Orient_Horizontal)
+
 				// Exposed entities List
 				+ SSplitter::Slot()
 				.Value(0.5f)
 				[
 					EntityList.ToSharedRef()
 				]
-			
+
 				// Logic Panel
 				+ SSplitter::Slot()
 				.Value(0.5f)
@@ -792,13 +792,13 @@ SRemoteControlPanel::~SRemoteControlPanel()
 
 	// Remove protocol bindings
 	IRemoteControlProtocolWidgetsModule& ProtocolWidgetsModule = FModuleManager::LoadModuleChecked<IRemoteControlProtocolWidgetsModule>("RemoteControlProtocolWidgets");
-	ProtocolWidgetsModule.ResetProtocolBindingList();	
+	ProtocolWidgetsModule.ResetProtocolBindingList();
 
 	// Unregister with UI module
-	if (FModuleManager::Get().IsModuleLoaded("RemoteControlUIModule"))
+	if (FModuleManager::Get().IsModuleLoaded("RemoteControlUI"))
 	{
 		FRemoteControlUIModule::Get().UnregisterRemoteControlPanel(this);
-	}	
+	}
 }
 
 void SRemoteControlPanel::Shutdown()
@@ -901,7 +901,7 @@ bool SRemoteControlPanel::IsAllObjectsExposed(TArray<UObject*> InOuterObjects, c
 	return bAllObjectsExposed;
 }
 
-void SRemoteControlPanel::ToggleProperty(const FRCExposesPropertyArgs& InPropertyArgs)
+void SRemoteControlPanel::ToggleProperty(const FRCExposesPropertyArgs& InPropertyArgs, FString InDesiredName)
 {
 	if (!ensure(InPropertyArgs.IsValid()))
 	{
@@ -942,7 +942,7 @@ void SRemoteControlPanel::ToggleProperty(const FRCExposesPropertyArgs& InPropert
 				if (Object)
 				{
 					constexpr bool bCleanDuplicates = true; // GeneratePathToProperty duplicates container name (Array.Array[1], Set.Set[1], etc...)
-					ExposeProperty(Object, FRCFieldPathInfo{ InPropertyArgs.PropertyHandle->GeneratePathToProperty(), bCleanDuplicates });
+					ExposeProperty(Object, FRCFieldPathInfo{ InPropertyArgs.PropertyHandle->GeneratePathToProperty(), bCleanDuplicates }, InDesiredName);
 				}
 			}
 
@@ -1057,9 +1057,9 @@ TSharedRef<SWidget> SRemoteControlPanel::CreateGetterSetterIgnoredWarning() cons
 }
 
 TSharedRef<SWidget> SRemoteControlPanel::CreateExposeFunctionsButton()
-{	
+{
 	FMenuBuilder MenuBuilder(true, nullptr);
-	
+
 	SAssignNew(BlueprintPicker, SRCPanelFunctionPicker)
 		.RemoteControlPanel(SharedThis(this))
 		.AllowDefaultObjects(true)
@@ -1078,7 +1078,7 @@ TSharedRef<SWidget> SRemoteControlPanel::CreateExposeFunctionsButton()
 		.Label(LOCTEXT("ActorFunctionsLabel", "Actor Functions"))
 		.ObjectClass(AActor::StaticClass())
 		.OnSelectFunction(this, &SRemoteControlPanel::ExposeFunction);
-	
+
 	MenuBuilder.BeginSection(NAME_None, LOCTEXT("ExposeHeader", "Expose"));
 	{
 		constexpr bool bNoIndent = true;
@@ -1104,13 +1104,13 @@ TSharedRef<SWidget> SRemoteControlPanel::CreateExposeFunctionsButton()
 			LOCTEXT("FunctionLibraryFunctionSubMenuToolTip", "Expose a function from a blueprint function library."),
 			BlueprintPicker.ToSharedRef()
 		);
-		
+
 		CreatePickerSubMenu(
 			LOCTEXT("SubsystemFunctionSubMenu", "Subsystem Function"),
 			LOCTEXT("SubsystemFunctionSubMenuToolTip", "Expose a function from a subsytem."),
 			SubsystemFunctionPicker.ToSharedRef()
 		);
-		
+
 		CreatePickerSubMenu(
 			LOCTEXT("ActorFunctionSubMenu", "Actor Function"),
 			LOCTEXT("ActorFunctionSubMenuToolTip", "Expose an actor's function."),
@@ -1119,7 +1119,7 @@ TSharedRef<SWidget> SRemoteControlPanel::CreateExposeFunctionsButton()
 	}
 
 	MenuBuilder.EndSection();
-	
+
 	return SAssignNew(ExposeFunctionsComboButton, SComboButton)
 		.AddMetaData<FTagMetaData>(FTagMetaData(TEXT("Expose Functions")))
 		.IsEnabled_Lambda([this]() { return !this->bIsInLiveMode; })
@@ -1148,9 +1148,9 @@ TSharedRef<SWidget> SRemoteControlPanel::CreateExposeFunctionsButton()
 }
 
 TSharedRef<SWidget> SRemoteControlPanel::CreateExposeActorsButton()
-{	
+{
 	FMenuBuilder MenuBuilder(true, nullptr);
-	
+
 	MenuBuilder.BeginSection(NAME_None, LOCTEXT("ExposeHeader", "Expose"));
 	{
 		constexpr bool bNoIndent = true;
@@ -1218,7 +1218,7 @@ TSharedRef<SWidget> SRemoteControlPanel::CreateExposeActorsButton()
 	}
 
 	MenuBuilder.EndSection();
-	
+
 	return SAssignNew(ExposeActorsComboButton, SComboButton)
 		.AddMetaData<FTagMetaData>(FTagMetaData(TEXT("Expose Actors")))
 		.IsEnabled_Lambda([this]() { return !this->bIsInLiveMode; })
@@ -1255,7 +1255,7 @@ TSharedRef<SWidget> SRemoteControlPanel::CreateExposeByClassWidget()
 			: Classes(InClasses)
 		{
 		}
-		
+
 		virtual bool IsClassAllowed(const FClassViewerInitializationOptions& InInitOptions, const UClass* InClass, TSharedRef<FClassViewerFilterFuncs> InFilterFuncs) override
 		{
 			return Classes.Contains(TWeakObjectPtr<const UClass>{InClass});
@@ -1266,12 +1266,12 @@ TSharedRef<SWidget> SRemoteControlPanel::CreateExposeByClassWidget()
 			return false;
 		}
 
-	public:		
+	public:
 		const TSet<TWeakObjectPtr<const UClass>>& Classes;
 	};
 
 	TSharedPtr<FActorClassInLevelFilter> Filter = MakeShared<FActorClassInLevelFilter>(CachedClassesInLevel);
-	
+
 	FClassViewerInitializationOptions Options;
 	{
 		Options.ClassFilters.Add(Filter.ToSharedRef());
@@ -1282,7 +1282,7 @@ TSharedRef<SWidget> SRemoteControlPanel::CreateExposeByClassWidget()
 		Options.bShowNoneOption = false;
 		Options.bShowUnloadedBlueprints = false;
 	}
-	
+
 	TSharedRef<SWidget> Widget = FModuleManager::LoadModuleChecked<FClassViewerModule>("ClassViewer").CreateClassViewer(Options, FOnClassPicked::CreateLambda(
 		[this](UClass* ChosenClass)
 		{
@@ -1325,7 +1325,7 @@ void SRemoteControlPanel::CacheLevelClasses()
 		{
 			CacheActorClass(*It);
 		}
-		
+
 		if (ClassPicker)
 		{
 			ClassPicker->Refresh();
@@ -1382,10 +1382,10 @@ void SRemoteControlPanel::CacheActorClass(AActor* Actor)
 void SRemoteControlPanel::OnMapChange(uint32)
 {
 	CacheLevelClasses();
-	
+
 	if (ClassPicker)
 	{
-		ClassPicker->Refresh();	
+		ClassPicker->Refresh();
 	}
 
 	UpdateRebindButtonVisibility();
@@ -1456,7 +1456,7 @@ void SRemoteControlPanel::BindRemoteControlCommands()
 void SRemoteControlPanel::RegisterEvents()
 {
 	FEditorDelegates::MapChange.AddSP(this, &SRemoteControlPanel::OnMapChange);
-	
+
 	if (GEditor)
 	{
 		GEditor->OnBlueprintReinstanced().AddSP(this, &SRemoteControlPanel::OnBlueprintReinstanced);
@@ -1479,14 +1479,14 @@ void SRemoteControlPanel::UnregisterEvents()
 {
 	Preset->OnEntityExposed().RemoveAll(this);
 	Preset->OnEntityUnexposed().RemoveAll(this);
-	
+
 	if (GEngine)
 	{
 		GEngine->OnLevelActorDeleted().RemoveAll(this);
 		GEngine->OnLevelActorListChanged().RemoveAll(this);
 		GEngine->OnLevelActorAdded().RemoveAll(this);
 	}
-	
+
 	if (GEditor)
 	{
 		GEditor->OnBlueprintReinstanced().RemoveAll(this);
@@ -1502,7 +1502,7 @@ void SRemoteControlPanel::RegisterPanels()
 	if (PanelDrawer.IsValid())
 	{
 		PanelDrawer->OnRCPanelToggled().BindSP(this, &SRemoteControlPanel::OnRCPanelToggled);
-		
+
 		PanelDrawer->CanToggleRCPanel().BindLambda([this]()
 			{
 				return !bIsInLiveMode;
@@ -1522,7 +1522,7 @@ void SRemoteControlPanel::RegisterPanels()
 
 			RegisteredDrawers.Add(PropertiesPanel->GetPanelID(), PropertiesPanel);
 		}
-		
+
 		{// Properties With Details Panel
 			TSharedRef<FRCPanelDrawerArgs> EntityDetailsPanel = MakeShared<FRCPanelDrawerArgs>(ERCPanels::RCP_EntityDetails);
 
@@ -1687,7 +1687,7 @@ void SRemoteControlPanel::OnLogCheckboxToggle(ECheckBoxState State)
 {
 	const bool bIsLogEnabled = (State == ECheckBoxState::Checked) ? true : false;
 	FRemoteControlLogger::Get().EnableLog(bIsLogEnabled);
-	
+
 	if (PanelDrawer.IsValid())
 	{
 		TSharedRef<FRCPanelDrawerArgs> OutputLogPanel = RegisteredDrawers.FindChecked(ERCPanels::RCP_OutputLog);
@@ -1711,11 +1711,12 @@ void SRemoteControlPanel::OnBlueprintReinstanced()
 	Refresh();
 }
 
-void SRemoteControlPanel::ExposeProperty(UObject* Object, FRCFieldPathInfo Path)
+void SRemoteControlPanel::ExposeProperty(UObject* Object, FRCFieldPathInfo Path, FString InDesiredName)
 {
 	if (Path.Resolve(Object))
 	{
 		FRemoteControlPresetExposeArgs Args;
+		Args.Label = InDesiredName;
 		Args.GroupId = GetSelectedGroup();
 		Preset->ExposeProperty(Object, MoveTemp(Path), MoveTemp(Args));
 	}
@@ -1727,7 +1728,7 @@ void SRemoteControlPanel::ExposeFunction(UObject* Object, UFunction* Function)
 	{
 		ExposeFunctionsComboButton->SetIsOpen(false);
 	}
-	
+
 	FScopedTransaction Transaction(LOCTEXT("ExposeFunction", "ExposeFunction"));
 	Preset->Modify();
 
@@ -1747,10 +1748,10 @@ void SRemoteControlPanel::ExposeActor(AActor* Actor)
 	{
 		FScopedTransaction Transaction(LOCTEXT("ExposeActor", "Expose Actor"));
 		Preset->Modify();
-		
+
 		FRemoteControlPresetExposeArgs Args;
 		Args.GroupId = GetSelectedGroup();
-		
+
 		Preset->ExposeActor(Actor, Args);
 	}
 }
@@ -1784,7 +1785,7 @@ TSharedRef<SWidget> SRemoteControlPanel::CreateEntityDetailsView()
 	EntityDetailsView = FModuleManager::GetModuleChecked<FPropertyEditorModule>("PropertyEditor").CreateStructureDetailView(MoveTemp(Args), FStructureDetailsViewArgs(), nullptr);
 
 	UpdateEntityDetailsView(EntityList->GetSelectedEntity());
-	
+
 	const bool bCanShowDetailsView = SelectedEntity.IsValid() && (SelectedEntity->GetRCType() != SRCPanelTreeNode::Group) && (SelectedEntity->GetRCType() != SRCPanelTreeNode::FieldChild);
 
 	if (bCanShowDetailsView)
@@ -1807,7 +1808,7 @@ void SRemoteControlPanel::UpdateEntityDetailsView(const TSharedPtr<SRCPanelTreeN
 	TSharedPtr<FStructOnScope> SelectedEntityPtr;
 
 	SelectedEntity = SelectedNode;
-	
+
 	if (SelectedEntity)
 	{
 		if (SelectedEntity->GetRCType() != SRCPanelTreeNode::Group &&
@@ -1834,7 +1835,7 @@ void SRemoteControlPanel::UpdateEntityDetailsView(const TSharedPtr<SRCPanelTreeN
 		WrappedEntityDetailsView->SetContent(NoneSelectedWidget.ToSharedRef());
 	}
 
-	static const FName ProtocolWidgetsModuleName = "RemoteControlProtocolWidgets";	
+	static const FName ProtocolWidgetsModuleName = "RemoteControlProtocolWidgets";
 	if(SelectedEntity.IsValid() && FModuleManager::Get().IsModuleLoaded(ProtocolWidgetsModuleName) && ensure(Preset.IsValid()))
 	{
 		if (const TSharedPtr<FRemoteControlEntity> RCEntity = Preset->GetExposedEntity(SelectedEntity->GetRCId()).Pin())
@@ -1937,8 +1938,8 @@ void SRemoteControlPanel::TriggerMaterialCompiledRefresh()
 	if (TSharedPtr<SWindow> Window = FSlateApplication::Get().FindWidgetWindow(SharedThis(this)))
 	{
 		const TArray<TSharedRef<SWindow>>& ChildWindows = Window->GetChildWindows();
-		// Selecting a material from the RC window can trigger material recompiles, 
-		// If we refresh the widgets right away we might end up closing the material selection window while the user is 
+		// Selecting a material from the RC window can trigger material recompiles,
+		// If we refresh the widgets right away we might end up closing the material selection window while the user is
 		// still selecting a material. Therefore we must wait until that window is closed before refreshing the panel.
 		if (ChildWindows.Num())
 		{
@@ -1999,9 +2000,9 @@ void SRemoteControlPanel::GenerateToolbar()
 	IMainFrameModule& MainFrame = FModuleManager::Get().LoadModuleChecked<IMainFrameModule>("MainFrame");
 
 	UToolMenu* GeneratedToolbar = ToolMenus->FindMenu(DefaultRemoteControlPanelToolBarName);
-	
+
 	GeneratedToolbar->Context = FToolMenuContext(MainFrame.GetMainFrameCommandBindings());
-	
+
 	TSharedRef<class SWidget> ToolBarWidget = ToolMenus->GenerateWidget(GeneratedToolbar);
 
 	TSharedRef<SWidget> MiscWidgets = SNullWidget::NullWidget;
@@ -2100,7 +2101,7 @@ void SRemoteControlPanel::GenerateToolbar()
 					OnLiveModeChange.ExecuteIfBound(SharedThis(this), bIsInLiveMode);
 				}
 			)
-		
+
 			+ SRCModeSwitcher::Mode("Setup")
 			.DefaultLabel(FText::Format(LOCTEXT("SetupModeLabel", "{0}"), { FText::FromString("Setup") }))
 			.DefaultTooltip(FText::Format(LOCTEXT("SetupModeTooltip", "Switch to {0} mode."), { FText::FromString("Setup") }))
@@ -2166,7 +2167,7 @@ void SRemoteControlPanel::GenerateToolbar()
 				.VAlign(VAlign_Center)
 				.Padding(4.f, 2.f)
 				.FillWidth(1.f)
-				[	
+				[
 					SNew(STextBlock)
 					.ToolTipText(LOCTEXT("RebindButtonToolTip", "Attempt to rebind all unbound entites of the preset."))
 					.Text(LOCTEXT("RebindButtonText", "Rebind All"))
@@ -2210,7 +2211,7 @@ void SRemoteControlPanel::RegisterAuxiliaryToolBar()
 				, Commands.ToggleProtocolMappings->GetLabel()
 			)
 			);
-			
+
 			ToolsSection.AddEntry(FToolMenuEntry::InitWidget("Logic"
 				, SNew(SAutoResizeButton)
 					.UICommand(FRemoteControlCommands::Get().ToggleLogicEditor)
@@ -2262,7 +2263,7 @@ FText SRemoteControlPanel::HandlePresetName() const
 	{
 		return FText::FromString(Preset->GetName());
 	}
-	
+
 	return FText::GetEmpty();
 }
 
@@ -2338,7 +2339,7 @@ void SRemoteControlPanel::ToggleProtocolMappings_Execute()
 	if (PanelDrawer.IsValid())
 	{
 		TSharedRef<FRCPanelDrawerArgs> ProtocolsPanel = RegisteredDrawers.FindChecked(ERCPanels::RCP_Protocols);
-		
+
 		ProtocolsPanel->DrawerVisibility = bIsInProtocolsMode ? EVisibility::Visible : EVisibility::Collapsed;
 
 		// When we are not in Protocols Mode collapse the drawer.
@@ -2397,7 +2398,7 @@ void SRemoteControlPanel::OnRCPanelToggled(ERCPanels InPanelID)
 	if (InPanelID != ActivePanel)
 	{
 		ActivePanel = InPanelID;
-		
+
 		if (EntityList.IsValid() && !bIsInLiveMode)
 		{
 			const bool bToggleProtcolMode = IsInProtocolsMode() && ActivePanel == ERCPanels::RCP_Protocols;
@@ -2421,7 +2422,7 @@ TSharedPtr<class SRCLogicPanelBase> SRemoteControlPanel::GetActiveLogicPanel() c
 	{
 		return ActionPanel;
 	}
-	
+
 	return nullptr;
 }
 
@@ -2429,9 +2430,9 @@ void SRemoteControlPanel::DeleteEntity_Execute()
 {
 	// Currently used  as common entry point of Delete UI command for both RC Entity and Logic Items.
 	// This could potentially be moved out if the Logic panels are moved to a separate tab.
-	
+
 	// ~ Delete Logic Item ~
-	// 
+	//
 	// If the user focus is currently active on a Logic panel then route the Delete command to it and return.
 	if (TSharedPtr<SRCLogicPanelBase> ActiveLogicPanel = GetActiveLogicPanel())
 	{
diff --git a/Engine/Plugins/VirtualProduction/RemoteControl/Source/RemoteControlUI/Private/UI/SRemoteControlPanel.h b/Engine/Plugins/VirtualProduction/RemoteControl/Source/RemoteControlUI/Private/UI/SRemoteControlPanel.h
--- a/Engine/Plugins/VirtualProduction/RemoteControl/Source/RemoteControlUI/Private/UI/SRemoteControlPanel.h
+++ b/Engine/Plugins/VirtualProduction/RemoteControl/Source/RemoteControlUI/Private/UI/SRemoteControlPanel.h
@@ -72,7 +72,7 @@ public:
 	DECLARE_MULTICAST_DELEGATE(FOnEmptyControllers);
 	DECLARE_MULTICAST_DELEGATE(FOnEmptyBehaviours);
 	DECLARE_MULTICAST_DELEGATE(FOnEmptyActions);
-	
+
 	void Construct(const FArguments& InArgs, URemoteControlPreset* InPreset, TSharedPtr<IToolkitHost> InToolkitHost);
 	~SRemoteControlPanel();
 	static void Shutdown();
@@ -108,8 +108,9 @@ public:
 	/**
 	 * Exposes or unexposes a property.
 	 * @param InArgs The extension arguments of the property to toggle.
+	 * @param InDesiredName Desired Name for this property if leaved empty RC will deduce it itself
 	 */
-	void ToggleProperty(const FRCExposesPropertyArgs& InArgs);
+	void ToggleProperty(const FRCExposesPropertyArgs& InArgs, FString InDesiredName = TEXT(""));
 
 	/**
 	 * @return Whether or not the panel is in live mode.
@@ -183,7 +184,7 @@ private:
 	void OnBlueprintReinstanced();
 
 	/** Expose a property using its handle. */
-	void ExposeProperty(UObject* Object, FRCFieldPathInfo FieldPath);
+	void ExposeProperty(UObject* Object, FRCFieldPathInfo FieldPath, FString InDesiredName = TEXT(""));
 
 	/** Exposes a function.  */
 	void ExposeFunction(UObject* Object, UFunction* Function);
@@ -205,13 +206,13 @@ private:
 
 	/** Creates a widget that warns the user when the settings are currently ignoring the Protected flag is enabled. */
 	TSharedRef<SWidget> CreateProtectedIgnoredWarning() const;
-	
+
 	/** Creates a widget that warns the user when the settings are currently ignoring the Getter/Setter is enabled. */
 	TSharedRef<SWidget> CreateGetterSetterIgnoredWarning() const;
 
 	/** Create expose button, allowing to expose blueprints and actor functions. */
 	TSharedRef<SWidget> CreateExposeFunctionsButton();
-	
+
 	/** Create expose button, allowing to expose actors. */
 	TSharedRef<SWidget> CreateExposeActorsButton();
 
@@ -220,7 +221,7 @@ private:
 
 	/** Cache the classes (and parent classes) of all actors in the level. */
 	void CacheLevelClasses();
-	
+
 	//~ Handlers for various level actor events.
 	void OnActorAddedToLevel(AActor* Actor);
 	void OnLevelActorsRemoved(AActor* Actor);
@@ -234,7 +235,7 @@ private:
 
 	/** Create the details view for the entity currently selected. */
 	TSharedRef<SWidget> CreateEntityDetailsView();
-	
+
 	/** Update the details view following entity selection change.  */
 	void UpdateEntityDetailsView(const TSharedPtr<SRCPanelTreeNode>& SelectedNode);
 
@@ -247,13 +248,13 @@ private:
 	//~ Handlers called in order to clear the exposed property cache.
 	void OnEntityExposed(URemoteControlPreset* InPreset, const FGuid& InEntityId);
 	void OnEntityUnexposed(URemoteControlPreset* InPreset, const FGuid& InEntityId);
-	
+
 	/** Toggles the logging part of UI */
 	void OnLogCheckboxToggle(ECheckBoxState State);
 
 	/** Triggers a next frame update of the actor function picker to ensure that added actors are valid. */
 	void UpdateActorFunctionPicker();
-	
+
 	/** Handle clicking on the setting button. */
 	FReply OnClickSettingsButton();
 
@@ -262,13 +263,13 @@ private:
 
 	/* Refreshed called at max once per frame when a material is compiled. */
 	void TriggerMaterialCompiledRefresh();
-	
+
 	/** Registers default tool bar */
 	static void RegisterDefaultToolBar();
 
 	/** Makes a default asset editing toolbar */
 	void GenerateToolbar();
-	
+
 	/** Registers Auxiliary tool bar */
 	static void RegisterAuxiliaryToolBar();
 
@@ -306,7 +307,7 @@ private:
 
 	/** Called to test if user is able to delete a group/exposed entity. */
 	bool CanDeleteEntity() const;
-	
+
 	/** Called when user attempts to rename a group/exposed entity. */
 	void RenameEntity_Execute() const;
 
@@ -330,7 +331,7 @@ private:
 
 	/** Called to test if user is able to Duplicate a logic UI item. */
 	bool CanDuplicateItem() const;
-	
+
 	// Exposed Entities filtering. (Filters the Exposed Entities view)
 	void OnSearchTextChanged(const FText& InFilterText);
 	void OnSearchTextCommitted(const FText& InFilterText, ETextCommit::Type InCommitType);
@@ -341,7 +342,7 @@ private:
 
 	/** Loads settings from config based on the preset identifier. */
 	void LoadSettings(const FGuid& InInstanceId);
-	
+
 	/** Saves settings from config based on the preset identifier. */
 	void SaveSettings();
 
@@ -436,7 +437,7 @@ private:
 	TSharedPtr<class SRCActionPanel> ActionPanel;
 
 	/** LogicClipboardItem - Holds the latest item copied from a Logic panel
-	* 
+	*
 	* Note: We track UObjects (Data Model) here rather than the UI Models as the latter are swept away the moment the user navigates to a different Controller.
 	* For example if the user copies an action from a behaviour in a given Controller but then navigates to another Controller, we can no longer rely on the previous UI objects
 	* as they would have been discarded in favor of a new data set for the actively selected Controller */
@@ -450,7 +451,7 @@ public:
 
 	static const float MinimumPanelWidth;
 	// Global Delegates for Remote Control Logic
-	FOnControllerAdded OnControllerAdded;	
+	FOnControllerAdded OnControllerAdded;
 	FOnBehaviourAdded OnBehaviourAdded;
 	FOnActionAdded OnActionAdded;
 	FOnControllerSelectionChanged OnControllerSelectionChanged;
@@ -461,4 +462,4 @@ public:
 
 	/** The panel from which the latest Logic UI item was copied*/
 	TSharedPtr<SRCLogicPanelBase> LogicClipboardItemSource;
-};
\ No newline at end of file
+};
