[Remote Control] QAVirtualProduction Compile Win64 linking error

Functions were wrongly wrapped in WITH_EDITOR define

diff --git a/Engine/Plugins/VirtualProduction/RemoteControl/Source/RemoteControlLogic/Public/Action/RCActionContainer.h b/Engine/Plugins/VirtualProduction/RemoteControl/Source/RemoteControlLogic/Public/Action/RCActionContainer.h
--- a/Engine/Plugins/VirtualProduction/RemoteControl/Source/RemoteControlLogic/Public/Action/RCActionContainer.h
+++ b/Engine/Plugins/VirtualProduction/RemoteControl/Source/RemoteControlLogic/Public/Action/RCActionContainer.h
@@ -70,10 +70,9 @@ public:
 	//~ Begin UObject
 #if WITH_EDITOR
 	/** Called after applying a transaction to the object. Used to broadcast Undo related container changes to UI */
-	virtual void PostEditUndo();
-	void ExecuteActionsOnLoad();
-	virtual void PostLoad() override;
+	virtual void PostEditUndo() override;
 #endif
+	virtual void PostLoad() override;
 	//~ End UObject
 
 	/** Set of child action container */
@@ -97,6 +96,8 @@ private:
 	/** Add remote control property function */
 	URCFunctionAction* AddFunctionAction(const TSharedRef<const FRemoteControlFunction> InRemoteControlFunction);
 
+	void ExecuteActionsOnLoad();
+	
 	/** The list of Actions present in this container */
 	UPROPERTY()
 	TSet<TObjectPtr<URCAction>> Actions;
