Fix duplicated LOCTEXT macro keys in RemoteControlUI causing CIS error

diff --git a/Engine/Plugins/VirtualProduction/RemoteControl/Source/RemoteControlUI/Private/RemoteControlUIModule.cpp b/Engine/Plugins/VirtualProduction/RemoteControl/Source/RemoteControlUI/Private/RemoteControlUIModule.cpp
--- a/Engine/Plugins/VirtualProduction/RemoteControl/Source/RemoteControlUI/Private/RemoteControlUIModule.cpp
+++ b/Engine/Plugins/VirtualProduction/RemoteControl/Source/RemoteControlUI/Private/RemoteControlUIModule.cpp
@@ -1240,7 +1240,7 @@ void FRemoteControlUIModule::CreateSubMenuForChildProperties(const FRCExposesPro
 	// Expose and Unexpose ALL section
 	MenuBuilder.AddMenuEntry(
 	LOCTEXT("RCExposeAll", "Expose All"),
-	LOCTEXT("RCExposeAll", "Expose all sub-properties"),
+	LOCTEXT("RCExposeAllTooltip", "Expose all sub-properties"),
 	FSlateIcon(),
 	FUIAction(FExecuteAction::CreateRaw(this, &FRemoteControlUIModule::OnExposeAll, ExposeAllArgs)),
 	NAME_None,
@@ -1248,7 +1248,7 @@ void FRemoteControlUIModule::CreateSubMenuForChildProperties(const FRCExposesPro
 
 	MenuBuilder.AddMenuEntry(
 		LOCTEXT("RCUnexposeAll", "Unexpose All"),
-		LOCTEXT("RCUnexposeAll", "Unexpose all sub-properties"),
+		LOCTEXT("RCUnexposeAllToolTip", "Unexpose all sub-properties"),
 		FSlateIcon(),
 		FUIAction(FExecuteAction::CreateRaw(this, &FRemoteControlUIModule::OnUnexposeAll, ExposeAllArgs)),
 		NAME_None,
@@ -1257,7 +1257,7 @@ void FRemoteControlUIModule::CreateSubMenuForChildProperties(const FRCExposesPro
 	MenuBuilder.EndSection();
 
 	// Sub-Properties section
-	MenuBuilder.BeginSection(NAME_None, LOCTEXT("ExposeUnexposeExpansion", "EXPOSE"));
+	MenuBuilder.BeginSection(NAME_None, LOCTEXT("ExposeUnexposeExpansionSection", "EXPOSE"));
 
 	for (FRCExposesAllPropertiesArgs PropArgs : ExposeAllArgs)
 	{
