[Remote Control] Fix non-unity compile issue #316763

diff --git a/Engine/Plugins/VirtualProduction/RemoteControl/Source/RemoteControlLogic/Private/Action/RCActionContainer.cpp b/Engine/Plugins/VirtualProduction/RemoteControl/Source/RemoteControlLogic/Private/Action/RCActionContainer.cpp
--- a/Engine/Plugins/VirtualProduction/RemoteControl/Source/RemoteControlLogic/Private/Action/RCActionContainer.cpp
+++ b/Engine/Plugins/VirtualProduction/RemoteControl/Source/RemoteControlLogic/Private/Action/RCActionContainer.cpp
@@ -7,6 +7,7 @@
 #include "Action/RCPropertyAction.h"
 #include "Behaviour/Builtin/Path/RCSetAssetByPathBehaviour.h"
 #include "Behaviour/RCBehaviour.h"
+#include "Containers/Ticker.h"
 #include "Controller/RCController.h"
 #include "Controller/RCCustomControllerUtilities.h"
 #include "IRemoteControlModule.h"
