[Remote Control] Crash in Virtual Production template project

diff --git a/Engine/Plugins/VirtualProduction/RemoteControl/Source/RemoteControlUI/Private/RemoteControlUIModule.cpp b/Engine/Plugins/VirtualProduction/RemoteControl/Source/RemoteControlUI/Private/RemoteControlUIModule.cpp
--- a/Engine/Plugins/VirtualProduction/RemoteControl/Source/RemoteControlUI/Private/RemoteControlUIModule.cpp
+++ b/Engine/Plugins/VirtualProduction/RemoteControl/Source/RemoteControlUI/Private/RemoteControlUIModule.cpp
@@ -5,10 +5,10 @@
 #include "AssetEditor/RemoteControlPresetEditorToolkit.h"
 #include "AssetTools/RemoteControlPresetActions.h"
 #include "AssetToolsModule.h"
-#include "DetailTreeNode.h"
 #include "Commands/RemoteControlCommands.h"
 #include "Components/SkeletalMeshComponent.h"
 #include "Components/StaticMeshComponent.h"
+#include "DetailTreeNode.h"
 #include "Elements/Framework/TypedElementRegistry.h"
 #include "Elements/Framework/TypedElementSelectionSet.h"
 #include "Engine/SkeletalMesh.h"
@@ -19,8 +19,8 @@
 #include "Interfaces/IMainFrameModule.h"
 #include "Kismet2/ComponentEditorUtils.h"
 #include "LevelEditorSubsystem.h"
-#include "MaterialEditor/DEditorParameterValue.h"
 #include "MaterialDomain.h"
+#include "MaterialEditor/DEditorParameterValue.h"
 #include "Materials/Material.h"
 #include "PropertyHandle.h"
 #include "RemoteControlActor.h"
@@ -32,13 +32,13 @@
 #include "Styling/AppStyle.h"
 #include "Styling/AppStyle.h"
 #include "Textures/SlateIcon.h"
-#include "UI/RemoteControlExposeMenuStyle.h"
 #include "UI/Action/SRCActionPanel.h"
 #include "UI/Behaviour/Builtin/Conditional/SRCBehaviourConditional.h"
 #include "UI/Behaviour/SRCBehaviourPanel.h"
 #include "UI/Customizations/FPassphraseCustomization.h"
 #include "UI/Customizations/NetworkAddressCustomization.h"
 #include "UI/Customizations/RemoteControlEntityCustomization.h"
+#include "UI/RemoteControlExposeMenuStyle.h"
 #include "UI/RemoteControlPanelStyle.h"
 #include "UI/SRCPanelExposedActor.h"
 #include "UI/SRCPanelExposedEntitiesList.h"
@@ -1222,8 +1222,26 @@ bool FRemoteControlUIModule::HasChildProperties(const FProperty* InProperty) con
 {
 	if (const FStructProperty* StructProperty = CastField<FStructProperty>(InProperty))
 	{
-		// if at least it contains one child than it is consider valid to expose child properties
-		return StructProperty->Struct->Children->IsValidLowLevel() || StructProperty->Struct->ChildProperties->IsValidLowLevel();
+		if (const TObjectPtr<UScriptStruct>& ScriptStruct = StructProperty->Struct)
+		{
+			// If property contains at least one child, than it is considered valid to expose its child properties
+			
+			if (const TObjectPtr<UField>& StructChildren = ScriptStruct->Children)
+			{
+				if (StructChildren->IsValidLowLevel())
+				{
+					return true;
+				}
+			}
+
+			if (const FField* ChildProperties = StructProperty->Struct->ChildProperties)
+			{
+				if (ChildProperties->IsValidLowLevel())
+				{
+					return true;
+				}
+			}
+		}
 	}
 	return false;
 }
@@ -1284,10 +1302,40 @@ void FRemoteControlUIModule::CreateSubMenuForChildProperties(const FRCExposesPro
 
 bool FRemoteControlUIModule::HasChildPropertiesExposed(const FRCExposesPropertyArgs& InPropertyArgs) const
 {
-	if (InPropertyArgs.PropertyHandle.IsValid() && InPropertyArgs.PropertyHandle->GetProperty()->IsA<FStructProperty>())
+	if (InPropertyArgs.PropertyHandle.IsValid())
 	{
-		// Check if it has sub-properties exposed
-		FStructProperty* StructProperty = CastFieldChecked<FStructProperty>(InPropertyArgs.PropertyHandle->GetProperty());
+		if (FProperty* Property = InPropertyArgs.PropertyHandle->GetProperty())
+		{
+			if (Property->IsA<FStructProperty>())
+			{
+				// Check if it has sub-properties exposed
+				const FStructProperty* StructProperty = CastFieldChecked<FStructProperty>(Property);
+				int32 ChildHandleIndex = 0;
+				for (TFieldIterator<FProperty> It(StructProperty->Struct); It; ++It)
+				{
+					FRCExposesPropertyArgs ChildArgs;
+					ChildArgs.OwnerObject = InPropertyArgs.OwnerObject.Get();
+					ChildArgs.Property = *It;
+					ChildArgs.PropertyPath = (*It)->GetPathName(nullptr);
+					ChildArgs.PropertyHandle = InPropertyArgs.PropertyHandle->GetChildHandle(ChildHandleIndex++);
+					const EPropertyExposeStatus ChildStatus = GetPropertyExposeStatus(ChildArgs);
+
+					if (ChildStatus == EPropertyExposeStatus::Exposed)
+					{
+						return true;
+					}
+				}
+			}
+		}
+	}
+	return false;
+}
+
+void FRemoteControlUIModule::GetAllExposableSubPropertyFromStruct(const FRCExposesPropertyArgs InPropertyArgs, TArray<FRCExposesAllPropertiesArgs>& OutAllProperty) const
+{
+	if (FProperty* PropertyHandleProperty = InPropertyArgs.PropertyHandle->GetProperty())
+	{
+		FStructProperty* StructProperty = CastFieldChecked<FStructProperty>(PropertyHandleProperty);
 		int32 ChildHandleIndex = 0;
 		for (TFieldIterator<FProperty> It(StructProperty->Struct); It; ++It)
 		{
@@ -1296,56 +1344,54 @@ bool FRemoteControlUIModule::HasChildPropertiesExposed(const FRCExposesPropertyA
 			ChildArgs.Property = *It;
 			ChildArgs.PropertyPath = (*It)->GetPathName(nullptr);
 			ChildArgs.PropertyHandle = InPropertyArgs.PropertyHandle->GetChildHandle(ChildHandleIndex++);
-			EPropertyExposeStatus ChildStatus = GetPropertyExposeStatus(ChildArgs);
 
-			if (ChildStatus == EPropertyExposeStatus::Exposed)
+			// Skip if not valid
+			if (!ChildArgs.IsValid())
 			{
-				return true;
+				continue;
 			}
-		}
-	}
-	return false;
-}
 
-void FRemoteControlUIModule::GetAllExposableSubPropertyFromStruct(const FRCExposesPropertyArgs InPropertyArgs, TArray<FRCExposesAllPropertiesArgs>& OutAllProperty) const
-{
-	FStructProperty* StructProperty = CastFieldChecked<FStructProperty>(InPropertyArgs.PropertyHandle->GetProperty());
-	int32 ChildHandleIndex = 0;
-	for (TFieldIterator<FProperty> It(StructProperty->Struct); It; ++It)
-	{
-		FRCExposesPropertyArgs ChildArgs;
-		ChildArgs.OwnerObject = InPropertyArgs.OwnerObject.Get();
-		ChildArgs.Property = *It;
-		ChildArgs.PropertyPath = (*It)->GetPathName(nullptr);
-		ChildArgs.PropertyHandle = InPropertyArgs.PropertyHandle->GetChildHandle(ChildHandleIndex++);
-		// Skip if not valid or special case for FColor and FLinearColor since they can decide to show the Alpha channel or not
-		if (!ChildArgs.IsValid() || (InPropertyArgs.GetProperty()->HasMetaData("HideAlphaChannel") && ChildArgs.Property->GetFName() == FName("A")))
-		{
-			continue;
-		}
-		FText PropertyName = It->GetDisplayNameText();
-		if (StructProperty->GetFName() == USceneComponent::GetRelativeRotationPropertyName() || StructProperty->GetFName() == USceneComponent::GetAbsoluteRotationPropertyName())
-		{
-			PropertyName = FText::FromString(It->GetName());
-		}
-		FText ExposedPropertyLabel = FText::Format(LOCTEXT("RC_ChildProperty", "{0}"), PropertyName);
-		TAttribute<FText> ToolTip = TAttribute<FText>::CreateLambda([this, ChildArgs, PropertyName]()
-		{
-			const EPropertyExposeStatus ChildStatus = GetPropertyExposeStatus(ChildArgs);
-			switch (ChildStatus)
+			// Skip if FColor and FLinearColor since they can decide to show the Alpha channel or not
+			if (const FProperty* Property = InPropertyArgs.GetProperty())
 			{
-				case EPropertyExposeStatus::Exposed:
-					return FText::Format(LOCTEXT("RCUnexpose", "Unexpose {0}"), PropertyName);
+				if (Property->HasMetaData("HideAlphaChannel"))
+				{
+					TWeakFieldPtr<FProperty> ChildProperty = ChildArgs.Property;
 
-				case EPropertyExposeStatus::Unexposed:
-					return FText::Format(LOCTEXT("RCExpose", "Expose {0}"), PropertyName);
+					if (ChildProperty.IsValid())
+					{
+						if (ChildProperty->GetFName() == FName("A"))
+						{
+							continue;
+						}
+					}
+				}
+			}
 
-				case EPropertyExposeStatus::Unexposable:
+			FText PropertyName = It->GetDisplayNameText();
+			if (StructProperty->GetFName() == USceneComponent::GetRelativeRotationPropertyName() || StructProperty->GetFName() == USceneComponent::GetAbsoluteRotationPropertyName())
+			{
+				PropertyName = FText::FromString(It->GetName());
+			}
+			FText ExposedPropertyLabel = FText::Format(LOCTEXT("RC_ChildProperty", "{0}"), PropertyName);
+			TAttribute<FText> ToolTip = TAttribute<FText>::CreateLambda([this, ChildArgs, PropertyName]()
+			{
+				const EPropertyExposeStatus ChildStatus = GetPropertyExposeStatus(ChildArgs);
+				switch (ChildStatus)
+				{
+					case EPropertyExposeStatus::Exposed:
+						return FText::Format(LOCTEXT("RCUnexpose", "Unexpose {0}"), PropertyName);
+
+					case EPropertyExposeStatus::Unexposed:
+						return FText::Format(LOCTEXT("RCExpose", "Expose {0}"), PropertyName);
+
+					case EPropertyExposeStatus::Unexposable:
 				default:
 					return FText::Format(LOCTEXT("RCUnexposable", "Property {0} cannot be exposed"), PropertyName);
-			}
-		});
-		OutAllProperty.Add(FRCExposesAllPropertiesArgs({ChildArgs, PropertyName, PropertyName.ToString(), ExposedPropertyLabel, ToolTip}));
+				}
+			});
+			OutAllProperty.Add(FRCExposesAllPropertiesArgs({ChildArgs, PropertyName, PropertyName.ToString(), ExposedPropertyLabel, ToolTip}));
+		}
 	}
 }
 
