[Remote Control] Add UI Module method to retrieve a custom controller widget from outside the RC Module

diff --git a/Engine/Plugins/VirtualProduction/RemoteControl/Source/RemoteControlUI/Private/RemoteControlUIModule.cpp b/Engine/Plugins/VirtualProduction/RemoteControl/Source/RemoteControlUI/Private/RemoteControlUIModule.cpp
--- a/Engine/Plugins/VirtualProduction/RemoteControl/Source/RemoteControlUI/Private/RemoteControlUIModule.cpp
+++ b/Engine/Plugins/VirtualProduction/RemoteControl/Source/RemoteControlUI/Private/RemoteControlUIModule.cpp
@@ -8,6 +8,7 @@
 #include "Commands/RemoteControlCommands.h"
 #include "Components/SkeletalMeshComponent.h"
 #include "Components/StaticMeshComponent.h"
+#include "Controller/RCCustomControllerUtilities.h"
 #include "DetailTreeNode.h"
 #include "Elements/Framework/TypedElementRegistry.h"
 #include "Elements/Framework/TypedElementSelectionSet.h"
@@ -35,6 +36,7 @@
 #include "UI/Action/SRCActionPanel.h"
 #include "UI/Behaviour/Builtin/Conditional/SRCBehaviourConditional.h"
 #include "UI/Behaviour/SRCBehaviourPanel.h"
+#include "UI/Controller/CustomControllers/SCustomTextureControllerWidget.h"
 #include "UI/Customizations/FPassphraseCustomization.h"
 #include "UI/Customizations/NetworkAddressCustomization.h"
 #include "UI/Customizations/RemoteControlEntityCustomization.h"
@@ -1012,6 +1014,18 @@ void FRemoteControlUIModule::SelectObjects(const TArray<UObject*>& Objects) cons
 	}
 }
 
+TSharedPtr<SWidget> FRemoteControlUIModule::CreateCustomControllerWidget(URCVirtualPropertyBase* InController) const
+{
+	TSharedPtr<SWidget> CustomControllerWidget;
+	const FString& CustomName = UE::RCCustomControllers::GetCustomControllerTypeName(InController);
+	if (CustomName == UE::RCCustomControllers::CustomTextureControllerName)
+	{
+		CustomControllerWidget = SNew(SCustomTextureControllerWidget, InController);
+	}
+
+	return CustomControllerWidget;
+}
+
 void FRemoteControlUIModule::RegisterWidgetFactories()
 {
 	RegisterWidgetFactoryForType(FRemoteControlActor::StaticStruct(), FOnGenerateRCWidget::CreateStatic(&SRCPanelExposedActor::MakeInstance));
diff --git a/Engine/Plugins/VirtualProduction/RemoteControl/Source/RemoteControlUI/Private/RemoteControlUIModule.h b/Engine/Plugins/VirtualProduction/RemoteControl/Source/RemoteControlUI/Private/RemoteControlUIModule.h
--- a/Engine/Plugins/VirtualProduction/RemoteControl/Source/RemoteControlUI/Private/RemoteControlUIModule.h
+++ b/Engine/Plugins/VirtualProduction/RemoteControl/Source/RemoteControlUI/Private/RemoteControlUIModule.h
@@ -58,6 +58,7 @@ public:
 	virtual void UnregisterWidgetFactoryForType(UScriptStruct* RemoteControlEntityType) override;
 	virtual void HighlightPropertyInDetailsPanel(const FPropertyPath& Path) const override;
 	virtual void SelectObjects(const TArray<UObject*>& Objects) const override;
+	virtual TSharedPtr<SWidget> CreateCustomControllerWidget(URCVirtualPropertyBase* InController) const override;
 	//~ End IRemoteControlUIModule interface
 
 	/**
diff --git a/Engine/Plugins/VirtualProduction/RemoteControl/Source/RemoteControlUI/Private/UI/Controller/RCControllerModel.cpp b/Engine/Plugins/VirtualProduction/RemoteControl/Source/RemoteControlUI/Private/UI/Controller/RCControllerModel.cpp
--- a/Engine/Plugins/VirtualProduction/RemoteControl/Source/RemoteControlUI/Private/UI/Controller/RCControllerModel.cpp
+++ b/Engine/Plugins/VirtualProduction/RemoteControl/Source/RemoteControlUI/Private/UI/Controller/RCControllerModel.cpp
@@ -48,22 +48,22 @@ TSharedRef<SWidget> FRCControllerModel::GetWidget() const
 
 	const TSharedRef<SHorizontalBox> FieldWidget = SNew(SHorizontalBox);
 	if (VirtualPropertyWeakPtr.IsValid())
-	{		
-		const FString& CustomName = UE::RCCustomControllers::GetCustomControllerTypeName(VirtualPropertyWeakPtr.Get());
-
+	{
+		static const FMargin SlotMargin(10.0f, 2.0f);
+		
 		// If this is a custom controller, we will use its custom widget
-		if (!CustomName.IsEmpty())
+		if (const TSharedPtr<SWidget>& CustomControllerWidget = IRemoteControlUIModule::Get().CreateCustomControllerWidget(VirtualPropertyWeakPtr.Get()))
 		{
 			FieldWidget->AddSlot()
-				.Padding(FMargin(10.0f, 2.0f))
+				.Padding(SlotMargin)
 				[
-					GetCustomControllerWidget(CustomName)
+					CustomControllerWidget.ToSharedRef()
 				];
 		}
 		else if (NodeWidgets.ValueWidget)
 		{
 			FieldWidget->AddSlot()
-				.Padding(FMargin(10.0f, 2.0f))
+				.Padding(SlotMargin)
 				.HAlign(HAlign_Left)
 				[
 					NodeWidgets.ValueWidget.ToSharedRef()
@@ -72,7 +72,7 @@ TSharedRef<SWidget> FRCControllerModel::GetWidget() const
 		else if (NodeWidgets.WholeRowWidget)
 		{
 			FieldWidget->AddSlot()
-				.Padding(FMargin(10.0f, 2.0f))
+				.Padding(SlotMargin)
 				[
 					NodeWidgets.WholeRowWidget.ToSharedRef()
 				];
diff --git a/Engine/Plugins/VirtualProduction/RemoteControl/Source/RemoteControlUI/Public/IRemoteControlUIModule.h b/Engine/Plugins/VirtualProduction/RemoteControl/Source/RemoteControlUI/Public/IRemoteControlUIModule.h
--- a/Engine/Plugins/VirtualProduction/RemoteControl/Source/RemoteControlUI/Public/IRemoteControlUIModule.h
+++ b/Engine/Plugins/VirtualProduction/RemoteControl/Source/RemoteControlUI/Public/IRemoteControlUIModule.h
@@ -11,15 +11,16 @@
 #include "Widgets/Layout/SSplitter.h"
 
 class FDelegateHandle;
+class FPropertyPath;
 class FRCPanelWidgetRegistry;
-struct FOnGenerateGlobalRowExtensionArgs;
-struct FRemoteControlEntity;
-class IPropertyHandle;
 class IDetailCategoryBuilder;
 class IDetailLayoutBuilder;
-struct SRCPanelTreeNode;
+class IPropertyHandle;
+class URCVirtualPropertyBase;
 class URemoteControlPreset;
-class FPropertyPath;
+struct FOnGenerateGlobalRowExtensionArgs;
+struct FRemoteControlEntity;
+struct SRCPanelTreeNode;
 
 struct FRCColumnSizeData
 {
@@ -210,4 +211,9 @@ public:
 	 * Set the list of selected objects.
 	 */
 	virtual void SelectObjects(const TArray<UObject*>& Objects) const = 0;
+
+	/**
+	 * Tries to retrieve a Custom Controller Widget for the specified Controller
+	 */
+	virtual TSharedPtr<SWidget> CreateCustomControllerWidget(URCVirtualPropertyBase* InController) const = 0;
 };
