[RemoteControl] Fix HTTP query parameters not being parsed when sent via WebSocket http route

diff --git a/Engine/Plugins/VirtualProduction/RemoteControl/Source/WebRemoteControl/Private/WebRemoteControlInternalUtils.cpp b/Engine/Plugins/VirtualProduction/RemoteControl/Source/WebRemoteControl/Private/WebRemoteControlInternalUtils.cpp
--- a/Engine/Plugins/VirtualProduction/RemoteControl/Source/WebRemoteControl/Private/WebRemoteControlInternalUtils.cpp
+++ b/Engine/Plugins/VirtualProduction/RemoteControl/Source/WebRemoteControl/Private/WebRemoteControlInternalUtils.cpp
@@ -5,6 +5,7 @@
 #include "HttpServerRequest.h"
 #include "Kismet/KismetSystemLibrary.h"
 #include "Misc/Base64.h"
+#include "PlatformHttp.h"
 #include "RemoteControlSettings.h"
 #include "Serialization/JsonReader.h"
 #include "UObject/StructOnScope.h"
@@ -84,11 +85,40 @@ TSharedRef<FHttpServerRequest> UnwrapHttpRequest(const FRCRequestWrapper& Wrappe
 	}
 
 	WebRemoteControlInternalUtils::AddWrappedRequestHeader(*WrappedHttpRequest);
-	WrappedHttpRequest->RelativePath = Wrapper.URL;
 	WrappedHttpRequest->Verb = ParseHttpVerb(Wrapper.Verb);
 	WrappedHttpRequest->Body = Wrapper.TCHARBody;
 	WrappedHttpRequest->Headers.Add(WebRemoteControlInternalUtils::PassphraseHeader, { Wrapper.Passphrase });
 
+	// Parse query parameters from URL, if any
+	int32 QueryParamsIndex = 0;
+	if (Wrapper.URL.FindChar(TCHAR('?'), QueryParamsIndex))
+	{
+		FString PathWithoutParams = Wrapper.URL;
+
+		FString QueryParamsStr = PathWithoutParams.Mid(QueryParamsIndex + 1);
+		PathWithoutParams.MidInline(0, QueryParamsIndex, false);
+
+		// Split query params
+		TArray<FString> QueryParamPairs;
+		QueryParamsStr.ParseIntoArray(QueryParamPairs, TEXT("&"), true);
+		for (const FString& QueryParamPair : QueryParamPairs)
+		{
+			int32 Equalsindex = 0;
+			if (QueryParamPair.FindChar(TCHAR('='), Equalsindex))
+			{
+				FString QueryParamKey = FPlatformHttp::UrlDecode(QueryParamPair.Mid(0, Equalsindex));
+				FString QueryParamValue = FPlatformHttp::UrlDecode(QueryParamPair.Mid(Equalsindex + 1));
+				WrappedHttpRequest->QueryParams.Emplace(MoveTemp(QueryParamKey), MoveTemp(QueryParamValue));
+			}
+		}
+
+		WrappedHttpRequest->RelativePath = PathWithoutParams;
+	}
+	else
+	{
+		WrappedHttpRequest->RelativePath = Wrapper.URL;
+	}
+
 	return WrappedHttpRequest;
 }
 
diff --git a/Engine/Plugins/VirtualProduction/RemoteControl/Source/WebRemoteControl/WebRemoteControl.Build.cs b/Engine/Plugins/VirtualProduction/RemoteControl/Source/WebRemoteControl/WebRemoteControl.Build.cs
--- a/Engine/Plugins/VirtualProduction/RemoteControl/Source/WebRemoteControl/WebRemoteControl.Build.cs
+++ b/Engine/Plugins/VirtualProduction/RemoteControl/Source/WebRemoteControl/WebRemoteControl.Build.cs
@@ -18,6 +18,7 @@ public class WebRemoteControl : ModuleRules
         PrivateDependencyModuleNames.AddRange(
 			new string[] {
 				"AssetRegistry",
+				"HTTP",
 				"Networking",
 				"RemoteControl",
 				"RemoteControlCommon",
