Remote Control Logic:
- Update Target Entity when target entity is invalid or out of date

diff --git a/Engine/Plugins/VirtualProduction/RemoteControl/Source/RemoteControlLogic/Private/Behaviour/Builtin/Path/RCSetAssetByPathBehaviour.cpp b/Engine/Plugins/VirtualProduction/RemoteControl/Source/RemoteControlLogic/Private/Behaviour/Builtin/Path/RCSetAssetByPathBehaviour.cpp
--- a/Engine/Plugins/VirtualProduction/RemoteControl/Source/RemoteControlLogic/Private/Behaviour/Builtin/Path/RCSetAssetByPathBehaviour.cpp
+++ b/Engine/Plugins/VirtualProduction/RemoteControl/Source/RemoteControlLogic/Private/Behaviour/Builtin/Path/RCSetAssetByPathBehaviour.cpp
@@ -165,7 +165,12 @@ bool URCSetAssetByPathBehaviour::SetAssetByPath(const FString& AssetPath, const
 bool URCSetAssetByPathBehaviour::SetInternalAsset(UObject* SetterObject)
 {
 	const URCController* Controller = ControllerWeakPtr.Get();
-	
+
+	if (!TargetEntity || TargetEntity->GetId() != TargetEntityId)
+	{
+		UpdateTargetEntity();
+	}
+
 	if (!SetterObject || !Controller || !TargetEntity)
 	{
 		FRemoteControlLogger::Get().Log(SetAssetByPathBehaviourHelpers::SetAssetByPathBehaviour, [SetterObject]
@@ -362,6 +367,12 @@ TWeakPtr<const FRemoteControlEntity> URCSetAssetByPathBehaviour::GetTargetEntity
 bool URCSetAssetByPathBehaviour::SetExternalAsset(FString InExternalPath)
 {
 	URCController* Controller = ControllerWeakPtr.Get();
+
+	if (!TargetEntity || TargetEntity->GetId() != TargetEntityId)
+	{
+		UpdateTargetEntity();
+	}
+
 	if (!Controller || !TargetEntity)
 	{
 		return false;
