[Remote Control] Texture Controller not working with dynamic material

Code was not properly taking into account TObjectPtr vs. raw pointer

diff --git a/Engine/Plugins/VirtualProduction/RemoteControl/Source/RemoteControlLogic/Private/Behaviour/Builtin/Bind/RCBehaviourBind.cpp b/Engine/Plugins/VirtualProduction/RemoteControl/Source/RemoteControlLogic/Private/Behaviour/Builtin/Bind/RCBehaviourBind.cpp
--- a/Engine/Plugins/VirtualProduction/RemoteControl/Source/RemoteControlLogic/Private/Behaviour/Builtin/Bind/RCBehaviourBind.cpp
+++ b/Engine/Plugins/VirtualProduction/RemoteControl/Source/RemoteControlLogic/Private/Behaviour/Builtin/Bind/RCBehaviourBind.cpp
@@ -118,7 +118,7 @@ static bool EvaluateCustomControllerBindCompatibility(const URCController* InCon
 		return false;
 	}
 	
-	const FString CustomName = UE::RCCustomControllers::GetCustomControllerTypeName(InController);
+	const FString& CustomName = UE::RCCustomControllers::GetCustomControllerTypeName(InController);
 
 	if (CustomName.IsEmpty())
 	{
@@ -126,14 +126,16 @@ static bool EvaluateCustomControllerBindCompatibility(const URCController* InCon
 	}
 
 	// External Texture for image files (controller string, rc exposed property is object
-	if (CustomName == UE::RCCustomControllers::CustomTextureControllerName && InControllerPropertyClass == FStrProperty::StaticClass() && InRemoteControlPropertyClass == FObjectProperty::StaticClass())
+	if (CustomName == UE::RCCustomControllers::CustomTextureControllerName && InControllerPropertyClass == FStrProperty::StaticClass())
 	{
-		const FObjectProperty* TargetObjectProperty = CastField<FObjectProperty>(InRemoteControlProperty);
-
-		// Target type is texture - we're good to go
-		if (TargetObjectProperty->PropertyClass == UTexture::StaticClass())
+		// This cast makes sure we catch both UObject* and TObjectPtr
+		if (const FObjectProperty* TargetObjectProperty = CastField<FObjectProperty>(InRemoteControlProperty))
 		{
-			return true;
+			// Target type is texture - we're good to go
+			if (TargetObjectProperty->PropertyClass == UTexture::StaticClass())
+			{
+				return true;
+			}
 		}
 	}
 
