[UnrealStage] Fix color grading outliner entries becoming grey when level is reloaded

diff --git a/Engine/Plugins/Runtime/nDisplay/Source/DisplayCluster/Public/DisplayClusterLightCardActor.h b/Engine/Plugins/Runtime/nDisplay/Source/DisplayCluster/Public/DisplayClusterLightCardActor.h
--- a/Engine/Plugins/Runtime/nDisplay/Source/DisplayCluster/Public/DisplayClusterLightCardActor.h
+++ b/Engine/Plugins/Runtime/nDisplay/Source/DisplayCluster/Public/DisplayClusterLightCardActor.h
@@ -341,7 +341,7 @@ protected:
 #endif
 
 	/** Indicates this light card should be considered a flag */
-	UPROPERTY(Getter = IsLightCardFlag, Setter = SetIsLightCardFlag, meta = (AllowPrivateAccess))
+	UPROPERTY(Getter = IsLightCardFlag, Setter = SetIsLightCardFlag, meta = (AllowPrivateAccess = "true"))
 	bool bIsLightCardFlag = false;
 
 private:
diff --git a/Engine/Plugins/VirtualProduction/RemoteControl/Source/WebRemoteControl/Private/WebSocketMessageHandler.cpp b/Engine/Plugins/VirtualProduction/RemoteControl/Source/WebRemoteControl/Private/WebSocketMessageHandler.cpp
--- a/Engine/Plugins/VirtualProduction/RemoteControl/Source/WebRemoteControl/Private/WebSocketMessageHandler.cpp
+++ b/Engine/Plugins/VirtualProduction/RemoteControl/Source/WebRemoteControl/Private/WebSocketMessageHandler.cpp
@@ -1003,15 +1003,18 @@ void FWebSocketMessageHandler::ProcessChangedProperties()
 				const int64 SequenceNumber = GetSequenceNumber(ClientToEventsPair.Key);
 
 				//Check if multiple booleans properties want to be sent and send them since multiple booleans have problem with the common workflow.
-				if(!TrySendMultipleBoolProperties(Preset, ClientToEventsPair.Key, ClassToEventsPair.Value, SequenceNumber))
+				if (ClassToEventsPair.Key == FBoolProperty::StaticClass())
 				{
-					TArray<uint8> WorkingBuffer;
-					if (ClientToEventsPair.Value.Num() && WritePropertyChangeEventPayload(Preset, { ClassToEventsPair.Value }, SequenceNumber, WorkingBuffer))
-					{
-						TArray<uint8> Payload;
-						WebRemoteControlUtils::ConvertToUTF8(WorkingBuffer, Payload);
-						Server->Send(ClientToEventsPair.Key, Payload);
-					}
+					TrySendMultipleBoolProperties(Preset, ClientToEventsPair.Key, ClassToEventsPair.Value, SequenceNumber);
+					continue;
+				}
+
+				TArray<uint8> WorkingBuffer;
+				if (ClientToEventsPair.Value.Num() && WritePropertyChangeEventPayload(Preset, { ClassToEventsPair.Value }, SequenceNumber, WorkingBuffer))
+				{
+					TArray<uint8> Payload;
+					WebRemoteControlUtils::ConvertToUTF8(WorkingBuffer, Payload);
+					Server->Send(ClientToEventsPair.Key, Payload);
 				}
 			}
 		}
