Fix property watchers not working in PIE

diff --git a/Engine/Plugins/VirtualProduction/RemoteControl/Source/RemoteControl/Private/RemoteControlPreset.cpp b/Engine/Plugins/VirtualProduction/RemoteControl/Source/RemoteControl/Private/RemoteControlPreset.cpp
--- a/Engine/Plugins/VirtualProduction/RemoteControl/Source/RemoteControl/Private/RemoteControlPreset.cpp
+++ b/Engine/Plugins/VirtualProduction/RemoteControl/Source/RemoteControl/Private/RemoteControlPreset.cpp
@@ -2178,6 +2178,11 @@ void URemoteControlPreset::OnPieEvent(bool)
 		{
 			for (TSharedPtr<FRemoteControlEntity> Entity : PresetPtr->Registry->GetExposedEntities<FRemoteControlEntity>())
 			{
+				if (Entity->GetStruct() == FRemoteControlProperty::StaticStruct())
+				{
+					StaticCastSharedPtr<FRemoteControlProperty>(Entity)->FieldPathInfo.Resolve(Entity->GetBoundObject());
+				}
+
 				PresetPtr->PerFrameUpdatedEntities.Add(Entity->GetId());
 			}
 		}
