[RemoteControl] Remove unnecessary ensures

diff --git a/Engine/Plugins/VirtualProduction/RemoteControl/Source/WebRemoteControl/Private/WebSocketMessageHandler.cpp b/Engine/Plugins/VirtualProduction/RemoteControl/Source/WebRemoteControl/Private/WebSocketMessageHandler.cpp
--- a/Engine/Plugins/VirtualProduction/RemoteControl/Source/WebRemoteControl/Private/WebSocketMessageHandler.cpp
+++ b/Engine/Plugins/VirtualProduction/RemoteControl/Source/WebRemoteControl/Private/WebSocketMessageHandler.cpp
@@ -1640,8 +1640,10 @@ void FWebSocketMessageHandler::GatherActorChangesForClass(UClass* ActorClass, TM
 {
 	const FWatchedClassData* WatchedData = ActorNotificationMap.Find(ActorClass);
 
-	if (!ensureMsgf(WatchedData != nullptr, TEXT("An actor was still being watched for a class that is no longer in ActorNotificationMap")))
+	if (WatchedData == nullptr)
 	{
+		// We may have stopped watching the class on the same frame that an actor of that class changed, in which case
+		// nobody is expecting an update and we can ignore it
 		return;
 	}
 
@@ -1708,7 +1710,8 @@ void FWebSocketMessageHandler::GatherActorChangesForDeletedClass(const FWebSocke
 
 	if (!WatchingClients)
 	{
-		ensureMsgf(false, TEXT("An actor was still being watched for a deleted class that is no longer in ActorNotificationMap"));
+		// We may have stopped watching the class on the same frame that an actor of that class changed, in which case
+		// nobody is expecting an update and we can ignore it
 		return;
 	}
 
