[Remote Control] Controllers List Items: add column extensibility

By registering to the new IRemoteControlUIModule delegates:
- OnAddControllerExtensionColumn
- OnGenerateControllerExtensionsWidgets

it is possible to add a custom column to the controllers list, and customize its widget, e.g. to perform custom actions on the controller

diff --git a/Engine/Plugins/VirtualProduction/RemoteControl/Source/RemoteControlUI/Private/RemoteControlUIModule.h b/Engine/Plugins/VirtualProduction/RemoteControl/Source/RemoteControlUI/Private/RemoteControlUIModule.h
--- a/Engine/Plugins/VirtualProduction/RemoteControl/Source/RemoteControlUI/Private/RemoteControlUIModule.h
+++ b/Engine/Plugins/VirtualProduction/RemoteControl/Source/RemoteControlUI/Private/RemoteControlUIModule.h
@@ -48,6 +48,8 @@ public:
 	//~ Begin IRemoteControlUIModule interface
 	virtual FOnGenerateExtensions& GetExtensionGenerators() override { return ExtensionsGenerator; }
 	virtual FOnRemoteControlPresetOpened& OnRemoteControlPresetOpened() override { return RemoteControlPresetOpenedDelegate; }
+	virtual FOnGenerateControllerExtensionsWidgets& OnGenerateControllerExtensionsWidgets() override { return GenerateControllerExtensionsWidgetsDelegate; }
+	virtual FOnAddControllerExtensionColumn& OnAddControllerExtensionColumn() override { return OnAddControllerExtensionColumnDelegate; }
 	virtual FDelegateHandle AddPropertyFilter(FOnDisplayExposeIcon OnDisplayExposeIcon) override;
 	virtual void RemovePropertyFilter(const FDelegateHandle& FilterDelegateHandle) override;
 	virtual void RegisterMetadataCustomization(FName MetadataKey, FOnCustomizeMetadataEntry OnCustomizeCallback) override;
@@ -245,6 +247,12 @@ private:
 	/** Delegate called when a Remote Control Preset panel is created for the specified Remote Control Preset */
 	FOnRemoteControlPresetOpened RemoteControlPresetOpenedDelegate;
 
+	/** Called when setting up/resetting the Controller panel list. Register to add custom column. */
+	FOnAddControllerExtensionColumn OnAddControllerExtensionColumnDelegate;
+	
+	/** Delegate called when a Remote Control Preset panel is created/opened for the specified Remote Control Preset */
+	FOnGenerateControllerExtensionsWidgets GenerateControllerExtensionsWidgetsDelegate;
+
 	/** Filters added by other plugins queried to determine if a property should display an expose icon. */
 	TMap<FDelegateHandle, FOnDisplayExposeIcon> ExternalFilterDelegates;
 
diff --git a/Engine/Plugins/VirtualProduction/RemoteControl/Source/RemoteControlUI/Private/UI/Controller/RCControllerModel.cpp b/Engine/Plugins/VirtualProduction/RemoteControl/Source/RemoteControlUI/Private/UI/Controller/RCControllerModel.cpp
--- a/Engine/Plugins/VirtualProduction/RemoteControl/Source/RemoteControlUI/Private/UI/Controller/RCControllerModel.cpp
+++ b/Engine/Plugins/VirtualProduction/RemoteControl/Source/RemoteControlUI/Private/UI/Controller/RCControllerModel.cpp
@@ -133,6 +133,21 @@ TSharedRef<SWidget> FRCControllerModel::GetCustomControllerWidget(const FString&
 	return SNullWidget::NullWidget;
 }
 
+TSharedRef<SWidget> FRCControllerModel::GetControllerExtensionWidget(const FName& InColumnName) const
+{
+	FRCControllerExtensionWidgetsInfo ExtensionWidgetsInfo(GetVirtualProperty());
+	
+	IRemoteControlUIModule::Get().OnGenerateControllerExtensionsWidgets().Broadcast(ExtensionWidgetsInfo);
+
+	const TMap<FName, TSharedRef<SWidget>>& WidgetsMap = ExtensionWidgetsInfo.CustomWidgetsMap;
+	if (WidgetsMap.Contains(InColumnName))
+	{
+		return WidgetsMap[InColumnName];
+	}
+	
+	return SNullWidget::NullWidget;
+}
+
 void FRCControllerModel::SetMultiController(bool bInIsMultiController)
 {
 	bIsMultiController = bInIsMultiController;
diff --git a/Engine/Plugins/VirtualProduction/RemoteControl/Source/RemoteControlUI/Private/UI/Controller/RCControllerModel.h b/Engine/Plugins/VirtualProduction/RemoteControl/Source/RemoteControlUI/Private/UI/Controller/RCControllerModel.h
--- a/Engine/Plugins/VirtualProduction/RemoteControl/Source/RemoteControlUI/Private/UI/Controller/RCControllerModel.h
+++ b/Engine/Plugins/VirtualProduction/RemoteControl/Source/RemoteControlUI/Private/UI/Controller/RCControllerModel.h
@@ -51,6 +51,9 @@ public:
 	/** Retrieve the widget for a custom controller */
 	TSharedRef<SWidget> GetCustomControllerWidget(const FString& InCustomControllerTypeName) const;
 
+	/** The widget allowing controller columns customization */
+	TSharedRef<SWidget> GetControllerExtensionWidget(const FName& InColumnName) const;
+	
 	/**
 	 * Specify if this is a MultiController handling other hidden controllers.
 	 * MultiControllers will setup and show a Type Selection Widget.
diff --git a/Engine/Plugins/VirtualProduction/RemoteControl/Source/RemoteControlUI/Private/UI/Controller/SRCControllerPanelList.cpp b/Engine/Plugins/VirtualProduction/RemoteControl/Source/RemoteControlUI/Private/UI/Controller/SRCControllerPanelList.cpp
--- a/Engine/Plugins/VirtualProduction/RemoteControl/Source/RemoteControlUI/Private/UI/Controller/SRCControllerPanelList.cpp
+++ b/Engine/Plugins/VirtualProduction/RemoteControl/Source/RemoteControlUI/Private/UI/Controller/SRCControllerPanelList.cpp
@@ -114,6 +114,13 @@ namespace UE::RCControllerPanelList
 
 				return WrapWithDropTarget(DragHandleWidget, EDragDropSupportedModes::ReorderOnly /*Allow reorder operation only*/);
 			}
+			else if (const TSharedPtr<SRCControllerPanelList> ControlPanelList = ControllerPanelListWeakPtr.Pin())
+			{
+				if (ControlPanelList->GetCustomColumns().Contains(ColumnName))
+				{
+					return ControllerItem->GetControllerExtensionWidget(ColumnName);
+				}
+			}
 
 			return SNullWidget::NullWidget;
 		}
@@ -473,6 +480,27 @@ void SRCControllerPanelList::Reset()
 
 	ShowFieldIdHeaderColumn(bShowFieldIdsColumn);
 	ShowValueTypeHeaderColumn(bIsInMultiControllerMode);
+
+	// Handle custom additional columns
+	CustomColumns.Empty();
+	IRemoteControlUIModule::Get().OnAddControllerExtensionColumn().Broadcast(CustomColumns);
+	for (const FName& ColumnName : CustomColumns)
+	{
+		if (ControllersHeaderRow.IsValid())
+		{
+			const bool bColumnIsGenerated = ControllersHeaderRow->IsColumnGenerated(ColumnName);
+			if (!bColumnIsGenerated)
+			{
+				ControllersHeaderRow->AddColumn(
+					SHeaderRow::FColumn::FArguments()
+					.ColumnId(ColumnName)
+					.DefaultLabel(FText::FromName(ColumnName))
+					.FillWidth(0.2f)
+					.HeaderContentPadding(RCPanelStyle->HeaderRowPadding)
+				);
+			}			
+		}
+	}
 	
 	ListView->RebuildList();
 }
@@ -961,6 +989,11 @@ void SRCControllerPanelList::ShowFieldIdHeaderColumn(bool bInShowColumn)
 	}
 }
 
+void SRCControllerPanelList::AddColumn(const FName& InColumnName)
+{
+	CustomColumns.AddUnique(InColumnName);
+}
+
 void SRCControllerPanelList::SetMultiControllerMode(bool bIsUniqueModeOn)
 {
 	bIsInMultiControllerMode = bIsUniqueModeOn;
diff --git a/Engine/Plugins/VirtualProduction/RemoteControl/Source/RemoteControlUI/Private/UI/Controller/SRCControllerPanelList.h b/Engine/Plugins/VirtualProduction/RemoteControl/Source/RemoteControlUI/Private/UI/Controller/SRCControllerPanelList.h
--- a/Engine/Plugins/VirtualProduction/RemoteControl/Source/RemoteControlUI/Private/UI/Controller/SRCControllerPanelList.h
+++ b/Engine/Plugins/VirtualProduction/RemoteControl/Source/RemoteControlUI/Private/UI/Controller/SRCControllerPanelList.h
@@ -142,6 +142,9 @@ public:
 	/** Enable/disable MultiController Mode */
 	void SetMultiControllerMode(bool bIsUniqueModeOn);
 
+	/** Return the current list of custom columns names, if any */
+	TConstArrayView<FName> GetCustomColumns() const { return CustomColumns; } 
+
 	/** Flag that facilitates usage of two mutually exclusive drag-drop zones within a single panel 
 	* 
 	* The first drag-drop zone is empty panel space for "Bind To New Controller"
@@ -221,6 +224,9 @@ private:
 	/** Set the visibility of the List header the Field Id Column */
 	void ShowFieldIdHeaderColumn(bool bInShowColumn);
 
+	/** Add a custom column */
+	void AddColumn(const FName& InColumnName);
+
 	/** Panel Style reference. */
 	const FRCPanelStyle* RCPanelStyle;
 
@@ -231,5 +237,8 @@ private:
 
 	/** Storing the header so we can add/remove columns */
 	TSharedPtr<SHeaderRow> ControllersHeaderRow;
+
+	/** List of custom columns names */
+	TArray<FName> CustomColumns;
 };
 
diff --git a/Engine/Plugins/VirtualProduction/RemoteControl/Source/RemoteControlUI/Public/IRemoteControlUIModule.h b/Engine/Plugins/VirtualProduction/RemoteControl/Source/RemoteControlUI/Public/IRemoteControlUIModule.h
--- a/Engine/Plugins/VirtualProduction/RemoteControl/Source/RemoteControlUI/Public/IRemoteControlUIModule.h
+++ b/Engine/Plugins/VirtualProduction/RemoteControl/Source/RemoteControlUI/Public/IRemoteControlUIModule.h
@@ -123,6 +123,44 @@ DECLARE_DELEGATE_RetVal_OneParam(bool, FOnDisplayExposeIcon, const FRCExposesPro
  */
 DECLARE_DELEGATE_FourParams(FOnCustomizeMetadataEntry, URemoteControlPreset* /*Preset*/, const FGuid& /*DisplayedEntityId*/, IDetailLayoutBuilder& /*LayoutBuilder*/, IDetailCategoryBuilder& /*CategoryBuilder*/);
 
+/** A struct used for controllers list items columns customization */
+struct FRCControllerExtensionWidgetsInfo
+{
+	FRCControllerExtensionWidgetsInfo(URCVirtualPropertyBase* InController)
+		: Controller(InController)
+	{
+	}
+
+	/** Add a new widget, and associate it to the specified custom column */
+	bool AddColumnWidget(const FName& InColumnName, const TSharedRef<SWidget>& InWidget)
+	{
+		if (CustomWidgetsMap.Contains(InColumnName))
+		{
+			return false;
+		}
+		
+		CustomWidgetsMap.Add(InColumnName, InWidget);
+
+		return true;
+	}
+
+	/** The controller being customized */
+	URCVirtualPropertyBase* Controller;
+
+	/** A map associating each column with a widget */
+	TMap<FName, TSharedRef<SWidget>> CustomWidgetsMap;
+};
+
+/**
+ * Callback called to customize the extensible widget for controllers in controllers list
+ */
+DECLARE_MULTICAST_DELEGATE_OneParam(FOnGenerateControllerExtensionsWidgets, FRCControllerExtensionWidgetsInfo& /*OutWidgetsInfo*/);
+
+/**
+ * Called when setting up / resetting the Controller panel list. Register to add custom column.
+ */
+DECLARE_MULTICAST_DELEGATE_OneParam(FOnAddControllerExtensionColumn, TArray<FName>& /*OutColumnNames*/);
+
 /**
  * A Remote Control module that allows exposing objects and properties from the editor.
  */
@@ -145,6 +183,12 @@ public:
 	/** Delegate called when a Remote Control Preset panel is created/opened for the specified Remote Control Preset */
     DECLARE_MULTICAST_DELEGATE_OneParam(FOnRemoteControlPresetOpened, URemoteControlPreset* /*RemoteControlPreset*/);
     virtual FOnRemoteControlPresetOpened& OnRemoteControlPresetOpened() = 0;
+
+	/** Delegate called when a Remote Control Preset panel is created/opened for the specified Remote Control Preset */
+	virtual FOnGenerateControllerExtensionsWidgets& OnGenerateControllerExtensionsWidgets() = 0;
+
+	/** Delegate called when a Remote Control Preset panel is created/opened for the specified Remote Control Preset */
+	virtual FOnAddControllerExtensionColumn& OnAddControllerExtensionColumn() = 0;
 	
 	/** 
 	 * Get the toolbar extension generators.
