[Remote Control] Controllers List Items: add column extensibility

By registering to the new IRemoteControlUIModule delegates:
- OnAddControllerExtensionColumn
- OnGenerateControllerExtensionsWidgets

it is possible to add a custom column to the controllers list, and customize its widget, e.g. to perform custom actions on the controller

diff --git a/Engine/Plugins/VirtualProduction/RemoteControl/Source/RemoteControlUI/Private/UI/Controller/SRCControllerPanelList.cpp b/Engine/Plugins/VirtualProduction/RemoteControl/Source/RemoteControlUI/Private/UI/Controller/SRCControllerPanelList.cpp
--- a/Engine/Plugins/VirtualProduction/RemoteControl/Source/RemoteControlUI/Private/UI/Controller/SRCControllerPanelList.cpp
+++ b/Engine/Plugins/VirtualProduction/RemoteControl/Source/RemoteControlUI/Private/UI/Controller/SRCControllerPanelList.cpp
@@ -14,6 +14,7 @@
 #include "IPropertyRowGenerator.h"
 #include "IRemoteControlModule.h"
 #include "Interfaces/IMainFrameModule.h"
+#include "Materials/MaterialInterface.h"
 #include "RCControllerModel.h"
 #include "RCMultiController.h"
 #include "RCVirtualProperty.h"
