[Remote Control] Enums not properly supported in RCP

C++ and BP enum properties are handled differently, with the former not properly working with RC Controllers.

diff --git a/Engine/Plugins/VirtualProduction/RemoteControl/Source/RemoteControlLogic/Private/Action/Bind/RCPropertyBindAction.cpp b/Engine/Plugins/VirtualProduction/RemoteControl/Source/RemoteControlLogic/Private/Action/Bind/RCPropertyBindAction.cpp
--- a/Engine/Plugins/VirtualProduction/RemoteControl/Source/RemoteControlLogic/Private/Action/Bind/RCPropertyBindAction.cpp
+++ b/Engine/Plugins/VirtualProduction/RemoteControl/Source/RemoteControlLogic/Private/Action/Bind/RCPropertyBindAction.cpp
@@ -252,10 +252,12 @@ void URCPropertyBindAction::Execute() const
 			}
 
 			// Numeric to Byte/Enum
-			if (RemoteControlProperty->IsA(FByteProperty::StaticClass()))
+			if (RemoteControlProperty->IsA(FByteProperty::StaticClass()) ||
+				RemoteControlProperty->IsA(FEnumProperty::StaticClass()))
 			{
 				Handle->SetValue((uint8)NumericValue);
 			}
+			
 			// Numeric To Numeric
 			if (RemoteControlProperty->IsA(FNumericProperty::StaticClass()))
 			{
diff --git a/Engine/Plugins/VirtualProduction/RemoteControl/Source/RemoteControlLogic/Private/Behaviour/Builtin/Bind/RCBehaviourBind.cpp b/Engine/Plugins/VirtualProduction/RemoteControl/Source/RemoteControlLogic/Private/Behaviour/Builtin/Bind/RCBehaviourBind.cpp
--- a/Engine/Plugins/VirtualProduction/RemoteControl/Source/RemoteControlLogic/Private/Behaviour/Builtin/Bind/RCBehaviourBind.cpp
+++ b/Engine/Plugins/VirtualProduction/RemoteControl/Source/RemoteControlLogic/Private/Behaviour/Builtin/Bind/RCBehaviourBind.cpp
@@ -161,6 +161,19 @@ bool URCBehaviourBind::CanHaveActionForField(URCController* Controller, TSharedR
 		return false; // Unbound Property
 	}
 
+	// Enums with base other than 8bits should be ignored, so let's make sure of that
+	if (const FEnumProperty* EnumProperty = CastField<FEnumProperty>(RemoteControlProperty))
+	{
+		if (const UEnum* Enum = EnumProperty->GetEnum())
+		{
+			const int64 MaxEnumValue = Enum->GetMaxEnumValue();
+			const uint32 NeededBits = FMath::RoundUpToPowerOfTwo(MaxEnumValue);
+
+			// 8 bits enums only
+			return NeededBits <= 256;
+		}
+	}
+	
 	// Indirect Binding (related types)
 	//
 	const static TMap<FFieldClass*, TArray<FFieldClass*>> SupportedIndirectBindsMap =
@@ -168,7 +181,7 @@ bool URCBehaviourBind::CanHaveActionForField(URCController* Controller, TSharedR
 		/* Controller Type */                           /* Supported Remote Control Property Types */
 
 		{ FStrProperty::StaticClass(),      /* --> */   { FTextProperty::StaticClass(),     FNameProperty::StaticClass()}},
-		{ FNumericProperty::StaticClass(),  /* --> */   { FNumericProperty::StaticClass(),  FBoolProperty::StaticClass(),  FByteProperty::StaticClass() } },
+		{ FNumericProperty::StaticClass(),  /* --> */   { FNumericProperty::StaticClass(),  FBoolProperty::StaticClass(),  FByteProperty::StaticClass(), FEnumProperty::StaticClass() } },
 		{ FBoolProperty::StaticClass(),     /* --> */   { FFloatProperty::StaticClass(),    FIntProperty::StaticClass(),   FBoolProperty::StaticClass() } }
 	};
 
diff --git a/Engine/Plugins/VirtualProduction/RemoteControl/Source/RemoteControlUI/Private/UI/Controller/SRCControllerPanelList.cpp b/Engine/Plugins/VirtualProduction/RemoteControl/Source/RemoteControlUI/Private/UI/Controller/SRCControllerPanelList.cpp
--- a/Engine/Plugins/VirtualProduction/RemoteControl/Source/RemoteControlUI/Private/UI/Controller/SRCControllerPanelList.cpp
+++ b/Engine/Plugins/VirtualProduction/RemoteControl/Source/RemoteControlUI/Private/UI/Controller/SRCControllerPanelList.cpp
@@ -59,6 +59,48 @@ namespace UE::RCControllerPanelList
 		const FName ValueTypeSelection = TEXT("Value Type Selection");
 	}
 
+	const TSet<UScriptStruct*>& GetSupportedStructs()
+	{
+		static const TSet<UScriptStruct*> SupportedStructs = {
+			TBaseStructure<FVector>::Get(),
+			TBaseStructure<FRotator>::Get(),
+			TBaseStructure<FColor>::Get()
+		};
+		
+		return SupportedStructs;
+	}
+
+	const TSet<UClass*>& GetSupportedObjects()
+	{
+		static const TSet<UClass*> SupportedObjects = {
+			UTexture::StaticClass(),
+			UStaticMesh::StaticClass(),
+			UMaterialInterface::StaticClass(),
+		};
+		
+		return SupportedObjects;
+	}
+	
+	bool IsStructPropertyTypeSupported(const FStructProperty* InStructProperty)
+	{
+		if (InStructProperty)
+		{
+			return GetSupportedStructs().Contains(InStructProperty->Struct);
+		}
+
+		return false;
+	}
+
+	bool IsObjectPropertyTypeSupported(const FObjectProperty* InObjectProperty)
+	{
+		if (InObjectProperty)
+		{
+			return GetSupportedObjects().Contains(InObjectProperty->PropertyClass);
+		}
+
+		return false;
+	}
+
 	class SControllerItemListRow : public SMultiColumnTableRow<TSharedRef<FRCControllerModel>>
 	{
 	public:
@@ -756,6 +798,44 @@ static TSharedPtr<FExposedEntityDragDrop> GetExposedEntityDragDrop(TSharedPtr<FD
 	return nullptr;
 }
 
+bool SRCControllerPanelList::IsEntitySupported(const FGuid ExposedEntityId)
+{
+	if (URemoteControlPreset* Preset = GetPreset())
+	{
+		if (const TSharedPtr<const FRemoteControlProperty>& RemoteControlProperty = Preset->GetExposedEntity<FRemoteControlProperty>(ExposedEntityId).Pin())
+		{
+			if (RemoteControlProperty->FieldType == EExposedFieldType::Property)
+			{
+				const FProperty* Property = RemoteControlProperty->GetProperty();
+				
+				if (const FEnumProperty* EnumProperty = CastField<FEnumProperty>(Property))
+				{
+					if (const UEnum* Enum = EnumProperty->GetEnum())
+					{
+						const int64 MaxEnumValue = Enum->GetMaxEnumValue();
+						const uint32 NeededBits = FMath::RoundUpToPowerOfTwo(MaxEnumValue);
+
+						// 8 bits enums only
+						return NeededBits <= 256;
+					}
+				}
+				else if (const FStructProperty* StructProperty = CastField<FStructProperty>(Property))
+				{
+					return UE::RCControllerPanelList::IsStructPropertyTypeSupported(StructProperty);
+				}
+				else if (const FObjectProperty* ObjectProperty = CastField<FObjectProperty>(Property))
+				{
+					return UE::RCControllerPanelList::IsObjectPropertyTypeSupported(ObjectProperty);
+				}
+
+				return true;
+			}
+		}
+	}
+
+	return false;
+}
+
 bool SRCControllerPanelList::OnAllowDrop(TSharedPtr<FDragDropOperation> DragDropOperation)
 {
 	if (IsListViewHovered())
@@ -773,13 +853,8 @@ bool SRCControllerPanelList::OnAllowDrop(TSharedPtr<FDragDropOperation> DragDrop
 		// Fetch the Exposed Entity
 		const FGuid ExposedEntityId = DragDropOp->GetId();
 
-		if (URemoteControlPreset* Preset = GetPreset())
-		{
-			if (TSharedPtr<const FRemoteControlField> RemoteControlField = Preset->GetExposedEntity<FRemoteControlField>(ExposedEntityId).Pin())
-			{
-				return RemoteControlField->FieldType == EExposedFieldType::Property;
-			}
-		}
+		// Check if Entity is supported by controllers
+		return IsEntitySupported(ExposedEntityId);
 	}
 
 	return false;
diff --git a/Engine/Plugins/VirtualProduction/RemoteControl/Source/RemoteControlUI/Private/UI/Controller/SRCControllerPanelList.h b/Engine/Plugins/VirtualProduction/RemoteControl/Source/RemoteControlUI/Private/UI/Controller/SRCControllerPanelList.h
--- a/Engine/Plugins/VirtualProduction/RemoteControl/Source/RemoteControlUI/Private/UI/Controller/SRCControllerPanelList.h
+++ b/Engine/Plugins/VirtualProduction/RemoteControl/Source/RemoteControlUI/Private/UI/Controller/SRCControllerPanelList.h
@@ -226,7 +226,9 @@ private:
 
 	/** Add a custom column */
 	void AddColumn(const FName& InColumnName);
-
+	
+	/** Checks whether the provided Entity can be used to create a Controller*/
+	bool IsEntitySupported(FGuid ExposedEntityId);
 
 	/** Panel Style reference. */
 	const FRCPanelStyle* RCPanelStyle;
