[RemoteControl] RemoteControlUI
Array entries not updating/rebinding correctly when modifying container in RC preset.

diff --git a/Engine/Plugins/VirtualProduction/RemoteControl/Source/RemoteControlUI/Private/UI/RCPanelWidgetRegistry.cpp b/Engine/Plugins/VirtualProduction/RemoteControl/Source/RemoteControlUI/Private/UI/RCPanelWidgetRegistry.cpp
--- a/Engine/Plugins/VirtualProduction/RemoteControl/Source/RemoteControlUI/Private/UI/RCPanelWidgetRegistry.cpp
+++ b/Engine/Plugins/VirtualProduction/RemoteControl/Source/RemoteControlUI/Private/UI/RCPanelWidgetRegistry.cpp
@@ -184,12 +184,12 @@ TSharedPtr<IDetailTreeNode> FRCPanelWidgetRegistry::GetStructTreeNode(const TSha
 	return WidgetRegistryUtils::FindNode(Generator->GetRootTreeNodes(), InField, InFindMethod);
 }
 
-void FRCPanelWidgetRegistry::Refresh(UObject* InObject)
+void FRCPanelWidgetRegistry::Refresh(UObject* InObject, bool bForceRefresh)
 {
 	if (TSharedPtr<IPropertyRowGenerator>* Generator = ObjectToRowGenerator.Find({InObject}))
 	{
 		const TArray<TWeakObjectPtr<UObject>>& SelectedObjects = (*Generator)->GetSelectedObjects();
-		if (SelectedObjects.Num() != 1 || SelectedObjects[0] != InObject)
+		if (bForceRefresh || SelectedObjects.Num() != 1 || SelectedObjects[0] != InObject)
 		{
 			(*Generator)->SetObjects({InObject});
 		}
diff --git a/Engine/Plugins/VirtualProduction/RemoteControl/Source/RemoteControlUI/Private/UI/RCPanelWidgetRegistry.h b/Engine/Plugins/VirtualProduction/RemoteControl/Source/RemoteControlUI/Private/UI/RCPanelWidgetRegistry.h
--- a/Engine/Plugins/VirtualProduction/RemoteControl/Source/RemoteControlUI/Private/UI/RCPanelWidgetRegistry.h
+++ b/Engine/Plugins/VirtualProduction/RemoteControl/Source/RemoteControlUI/Private/UI/RCPanelWidgetRegistry.h
@@ -54,7 +54,7 @@ public:
 	/**
 	 * Update the generator for a given object.
 	 */
-	void Refresh(UObject* InObject);
+	void Refresh(UObject* InObject, bool bForceRefresh = false);
 
 	/**
 	 * Update the generator for a given struct.
diff --git a/Engine/Plugins/VirtualProduction/RemoteControl/Source/RemoteControlUI/Private/UI/SRCPanelExposedEntitiesList.cpp b/Engine/Plugins/VirtualProduction/RemoteControl/Source/RemoteControlUI/Private/UI/SRCPanelExposedEntitiesList.cpp
--- a/Engine/Plugins/VirtualProduction/RemoteControl/Source/RemoteControlUI/Private/UI/SRCPanelExposedEntitiesList.cpp
+++ b/Engine/Plugins/VirtualProduction/RemoteControl/Source/RemoteControlUI/Private/UI/SRCPanelExposedEntitiesList.cpp
@@ -630,7 +630,7 @@ FText SRCPanelExposedEntitiesList::HandleEntityListHeaderLabel() const
 
 void SRCPanelExposedEntitiesList::OnObjectPropertyChange(UObject* InObject, FPropertyChangedEvent& InChangeEvent)
 {
-	EPropertyChangeType::Type TypesNeedingRefresh = EPropertyChangeType::ArrayAdd | EPropertyChangeType::ArrayClear | EPropertyChangeType::ArrayRemove | EPropertyChangeType::ValueSet;
+	EPropertyChangeType::Type TypesNeedingRefresh = EPropertyChangeType::ArrayAdd | EPropertyChangeType::ArrayClear | EPropertyChangeType::ArrayRemove | EPropertyChangeType::ValueSet | EPropertyChangeType::Duplicate;
 	auto IsRelevantProperty = [](FFieldClass* PropertyClass)
 	{
 		return PropertyClass && (PropertyClass == FArrayProperty::StaticClass() || PropertyClass == FSetProperty::StaticClass() || PropertyClass == FMapProperty::StaticClass());
@@ -642,7 +642,15 @@ void SRCPanelExposedEntitiesList::OnObjectPropertyChange(UObject* InObject, FPro
 	{
 		if (TSharedPtr<FRCPanelWidgetRegistry> Registry = WidgetRegistry.Pin())
 		{
-			Registry->Refresh(InObject);
+			if (InChangeEvent.Property)
+			{
+				// Force the refresh only if the property itself pass IsRelevantProperty check
+				Registry->Refresh(InObject, IsRelevantProperty(InChangeEvent.Property->GetClass()));
+			}
+			else
+			{
+				Registry->Refresh(InObject);
+			}
 		}
 
 		if (Preset)
