PR #10691: Fix clang-cl compilation of RCVirtualProperty.cpp

diff --git a/Engine/Plugins/VirtualProduction/RemoteControl/Source/RemoteControl/Private/RCVirtualProperty.cpp b/Engine/Plugins/VirtualProduction/RemoteControl/Source/RemoteControl/Private/RCVirtualProperty.cpp
--- a/Engine/Plugins/VirtualProduction/RemoteControl/Source/RemoteControl/Private/RCVirtualProperty.cpp
+++ b/Engine/Plugins/VirtualProduction/RemoteControl/Source/RemoteControl/Private/RCVirtualProperty.cpp
@@ -88,7 +88,7 @@ const FPropertyBagPropertyDesc* URCVirtualPropertyBase::GetBagPropertyDesc() con
 const FInstancedPropertyBag* URCVirtualPropertyBase::GetPropertyBagInstance() const
 {
 	// That is should be implemented in child classes
-	check(nullptr);
+	check(false);
 	return nullptr;
 }
 
