[RemoteControl] RemoteControlUI
Behaviour add button change, adding all behaviour and grey out the ones that can't be used with the specific controller

diff --git a/Engine/Plugins/VirtualProduction/RemoteControl/Source/RemoteControlLogic/Private/Controller/RCController.cpp b/Engine/Plugins/VirtualProduction/RemoteControl/Source/RemoteControlLogic/Private/Controller/RCController.cpp
--- a/Engine/Plugins/VirtualProduction/RemoteControl/Source/RemoteControlLogic/Private/Controller/RCController.cpp
+++ b/Engine/Plugins/VirtualProduction/RemoteControl/Source/RemoteControlLogic/Private/Controller/RCController.cpp
@@ -52,6 +52,19 @@ URCBehaviour* URCController::CreateBehaviour(TSubclassOf<URCBehaviourNode> InBeh
 	return NewBehaviour;
 }
 
+URCBehaviour* URCController::CreateBehaviourWithoutCheck(TSubclassOf<URCBehaviourNode> InBehaviourNodeClass)
+{
+	const URCBehaviourNode* DefaultBehaviourNode = Cast<URCBehaviourNode>(InBehaviourNodeClass->GetDefaultObject());
+	
+	URCBehaviour* NewBehaviour = NewObject<URCBehaviour>(this, DefaultBehaviourNode->GetBehaviourClass(), NAME_None, RF_Transactional);
+	NewBehaviour->BehaviourNodeClass = InBehaviourNodeClass;
+	NewBehaviour->Id = FGuid::NewGuid();
+	NewBehaviour->ActionContainer->PresetWeakPtr = PresetWeakPtr;
+	NewBehaviour->ControllerWeakPtr = this;
+	
+	return NewBehaviour;
+}
+
 int32 URCController::RemoveBehaviour(URCBehaviour* InBehaviour)
 {
 	return Behaviours.Remove(InBehaviour);
diff --git a/Engine/Plugins/VirtualProduction/RemoteControl/Source/RemoteControlLogic/Public/Controller/RCController.h b/Engine/Plugins/VirtualProduction/RemoteControl/Source/RemoteControlLogic/Public/Controller/RCController.h
--- a/Engine/Plugins/VirtualProduction/RemoteControl/Source/RemoteControlLogic/Public/Controller/RCController.h
+++ b/Engine/Plugins/VirtualProduction/RemoteControl/Source/RemoteControlLogic/Public/Controller/RCController.h
@@ -31,6 +31,9 @@ public:
 	/** Create new behaviour */
 	virtual URCBehaviour* CreateBehaviour(TSubclassOf<URCBehaviourNode> InBehaviourNodeClass);
 
+	/** Create new behaviour without checking if supported */
+	virtual URCBehaviour* CreateBehaviourWithoutCheck(TSubclassOf<URCBehaviourNode> InBehaviourNodeClass);
+
 	/** Remove the behaviour by behaviour UObject pointer */
 	virtual int32 RemoveBehaviour(URCBehaviour* InBehaviour);
 
diff --git a/Engine/Plugins/VirtualProduction/RemoteControl/Source/RemoteControlUI/Private/UI/Behaviour/SRCBehaviourPanel.cpp b/Engine/Plugins/VirtualProduction/RemoteControl/Source/RemoteControlUI/Private/UI/Behaviour/SRCBehaviourPanel.cpp
--- a/Engine/Plugins/VirtualProduction/RemoteControl/Source/RemoteControlUI/Private/UI/Behaviour/SRCBehaviourPanel.cpp
+++ b/Engine/Plugins/VirtualProduction/RemoteControl/Source/RemoteControlUI/Private/UI/Behaviour/SRCBehaviourPanel.cpp
@@ -171,13 +171,15 @@ TSharedRef<SWidget> SRCBehaviourPanel::GetBehaviourMenuContentWidget()
 						}
 					}
 
-					URCBehaviour* Behaviour = Controller->CreateBehaviour(Class);
+					URCBehaviour* Behaviour = Controller->CreateBehaviourWithoutCheck(Class);
 					if (!Behaviour)
 					{
 						continue;
 					}
 
 					FUIAction Action(FExecuteAction::CreateSP(this, &SRCBehaviourPanel::OnAddBehaviourClicked, Class));
+					Action.CanExecuteAction.BindSP(this, &SRCBehaviourPanel::CanExecuteAddBehaviour, Class, Behaviour);
+
 					MenuBuilder.AddMenuEntry(
 						FText::Format(LOCTEXT("AddBehaviourNode", "{0}"), Behaviour->GetDisplayName()),
 						FText::Format(LOCTEXT("AddBehaviourNodeTooltip", "{0}"), Behaviour->GetBehaviorDescription()),
@@ -210,6 +212,12 @@ void SRCBehaviourPanel::OnAddBehaviourClicked(UClass* InClass)
 	}
 }
 
+bool SRCBehaviourPanel::CanExecuteAddBehaviour(UClass* InClass, URCBehaviour* InBehaviour) const
+{
+	const URCBehaviourNode* DefaultBehaviourNode = Cast<URCBehaviourNode>(InClass->GetDefaultObject());
+	return DefaultBehaviourNode->IsSupported(InBehaviour);
+}
+
 FReply SRCBehaviourPanel::OnClickEmptyButton()
 {
 	if (const TSharedPtr<FRCControllerModel> ControllerItem = SelectedControllerItemWeakPtr.Pin())
diff --git a/Engine/Plugins/VirtualProduction/RemoteControl/Source/RemoteControlUI/Private/UI/Behaviour/SRCBehaviourPanel.h b/Engine/Plugins/VirtualProduction/RemoteControl/Source/RemoteControlUI/Private/UI/Behaviour/SRCBehaviourPanel.h
--- a/Engine/Plugins/VirtualProduction/RemoteControl/Source/RemoteControlUI/Private/UI/Behaviour/SRCBehaviourPanel.h
+++ b/Engine/Plugins/VirtualProduction/RemoteControl/Source/RemoteControlUI/Private/UI/Behaviour/SRCBehaviourPanel.h
@@ -87,6 +87,9 @@ private:
 	/** Handles click event for Add Behaviour button*/
 	void OnAddBehaviourClicked(UClass* InClass);
 
+	/** Check if the controller can add the behaviour requested */
+	bool CanExecuteAddBehaviour(UClass* InClass, URCBehaviour* InBehaviour) const;
+
 	/** Handles click event for "Empty" button; clears all Behaviours from the panel*/
 	FReply OnClickEmptyButton();
 
