[RemoteControl] RemoteControlUI
Changed color of group panel and action panel based on feedback

diff --git a/Engine/Plugins/VirtualProduction/RemoteControl/Source/RemoteControlUI/Private/UI/Action/SRCActionPanelList.h b/Engine/Plugins/VirtualProduction/RemoteControl/Source/RemoteControlUI/Private/UI/Action/SRCActionPanelList.h
--- a/Engine/Plugins/VirtualProduction/RemoteControl/Source/RemoteControlUI/Private/UI/Action/SRCActionPanelList.h
+++ b/Engine/Plugins/VirtualProduction/RemoteControl/Source/RemoteControlUI/Private/UI/Action/SRCActionPanelList.h
@@ -64,11 +64,10 @@ public:
 
 		ListView = SNew(SListView<TSharedPtr<ActionType>>)
 			.ListItemsSource(&ActionItems)
+			.HeaderRow(ActionType::GetHeaderRow())
 			.OnGenerateRow(this, &SRCActionPanelList::OnGenerateWidgetForList)
 			.OnSelectionChanged(this, &SRCActionPanelList::OnSelectionChanged)
-			.ListViewStyle(&RCPanelStyle->TableViewStyle)
-			.OnContextMenuOpening(this, &SRCLogicPanelListBase::GetContextMenuWidget)
-			.HeaderRow(ActionType::GetHeaderRow());
+			.OnContextMenuOpening(this, &SRCLogicPanelListBase::GetContextMenuWidget);
 
 		ChildSlot
 			[
diff --git a/Engine/Plugins/VirtualProduction/RemoteControl/Source/RemoteControlUI/Private/UI/SRCPanelExposedEntitiesList.cpp b/Engine/Plugins/VirtualProduction/RemoteControl/Source/RemoteControlUI/Private/UI/SRCPanelExposedEntitiesList.cpp
--- a/Engine/Plugins/VirtualProduction/RemoteControl/Source/RemoteControlUI/Private/UI/SRCPanelExposedEntitiesList.cpp
+++ b/Engine/Plugins/VirtualProduction/RemoteControl/Source/RemoteControlUI/Private/UI/SRCPanelExposedEntitiesList.cpp
@@ -206,13 +206,12 @@ void SRCPanelExposedEntitiesList::Construct(const FArguments& InArgs, URemoteCon
 	// Groups List
 	SAssignNew(GroupsListView, SListView<TSharedPtr<SRCPanelTreeNode>>)
 		.ItemHeight(24.f)
-		.OnGenerateRow(this, &SRCPanelExposedEntitiesList::OnGenerateRow)
-		.OnSelectionChanged(this, &SRCPanelExposedEntitiesList::OnSelectionChanged)
 		.SelectionMode(ESelectionMode::Single)
 		.ListItemsSource(reinterpret_cast<TArray<TSharedPtr<SRCPanelTreeNode>>*>(&FieldGroups))
-		.OnContextMenuOpening(this, &SRCPanelExposedEntitiesList::OnContextMenuOpening, SRCPanelTreeNode::Group)
-		.ListViewStyle(&RCPanelStyle->TableViewStyle)
-		.ClearSelectionOnClick(true);
+		.ClearSelectionOnClick(true)
+		.OnGenerateRow(this, &SRCPanelExposedEntitiesList::OnGenerateRow)
+		.OnSelectionChanged(this, &SRCPanelExposedEntitiesList::OnSelectionChanged)
+		.OnContextMenuOpening(this, &SRCPanelExposedEntitiesList::OnContextMenuOpening, SRCPanelTreeNode::Group);
 
 	// Group Dock Panel
 	TSharedPtr<SRCMinorPanel> GroupDockPanel = SNew(SRCMinorPanel)
