[RemoteControl] Added redirector for ERCMask::None -> ERCMask::NoMask

diff --git a/Engine/Plugins/VirtualProduction/RemoteControl/Config/BaseRemoteControl.ini b/Engine/Plugins/VirtualProduction/RemoteControl/Config/BaseRemoteControl.ini
--- a/Engine/Plugins/VirtualProduction/RemoteControl/Config/BaseRemoteControl.ini
+++ b/Engine/Plugins/VirtualProduction/RemoteControl/Config/BaseRemoteControl.ini
@@ -2,6 +2,7 @@
 +PropertyRedirects=(OldName="RemoteControlField.Label",NewName="RemoteControlEntity.Label")
 +PropertyRedirects=(OldName="RemoteControlField.Id",NewName="RemoteControlEntity.Id")
 +PropertyRedirects=(OldName="/Script/RemoteControlLogic.RemoteControlLogicConfig.SupportedControllerObjectTypes",NewName="/Script/RemoteControlLogic.RemoteControlLogicConfig.SupportedControllerObjectClassPaths")
++EnumRedirects=(OldName="ERCMask",ValueChanges=(("None","NoMask")))
 
 [/Script/RemoteControlLogic.RemoteControlLogicConfig]
 +SupportedControllerTypes=EPropertyBagPropertyType::String
