[RemoteControl] RemoteControlUI
Remove Expose/Logic panel bar

diff --git a/Engine/Plugins/VirtualProduction/RemoteControl/Source/RemoteControlUI/Private/UI/SRCPanelExposedEntitiesList.cpp b/Engine/Plugins/VirtualProduction/RemoteControl/Source/RemoteControlUI/Private/UI/SRCPanelExposedEntitiesList.cpp
--- a/Engine/Plugins/VirtualProduction/RemoteControl/Source/RemoteControlUI/Private/UI/SRCPanelExposedEntitiesList.cpp
+++ b/Engine/Plugins/VirtualProduction/RemoteControl/Source/RemoteControlUI/Private/UI/SRCPanelExposedEntitiesList.cpp
@@ -200,7 +200,7 @@ void SRCPanelExposedEntitiesList::Construct(const FArguments& InArgs, URemoteCon
 
 	// Major Panel
 	TSharedPtr<SRCMajorPanel> ExposePanel = SNew(SRCMajorPanel)
-		.HeaderLabel(LOCTEXT("ExposePanelHeader", "EXPOSE"))
+		.EnableHeader(false)
 		.EnableFooter(false);
 
 	// Groups List
diff --git a/Engine/Plugins/VirtualProduction/RemoteControl/Source/RemoteControlUI/Private/UI/SRemoteControlPanel.cpp b/Engine/Plugins/VirtualProduction/RemoteControl/Source/RemoteControlUI/Private/UI/SRemoteControlPanel.cpp
--- a/Engine/Plugins/VirtualProduction/RemoteControl/Source/RemoteControlUI/Private/UI/SRemoteControlPanel.cpp
+++ b/Engine/Plugins/VirtualProduction/RemoteControl/Source/RemoteControlUI/Private/UI/SRemoteControlPanel.cpp
@@ -543,7 +543,7 @@ void SRemoteControlPanel::Construct(const FArguments& InArgs, URemoteControlPres
 
 	// Make 3 Columns with Controllers + Behaviours + Actions
 	TSharedRef<SRCMajorPanel> LogicPanel = SNew(SRCMajorPanel)
-		.HeaderLabel(LOCTEXT("LogicPanelHeader", "LOGIC"))
+		.EnableHeader(false)
 		.EnableFooter(false);
 
 	// Controllers with Behaviours panel
