[RemoteControl] Add zlib compression option for WebSocket traffic

diff --git a/Engine/Plugins/VirtualProduction/EpicStageApp/Source/EpicStageApp/Public/StageAppVersion.h b/Engine/Plugins/VirtualProduction/EpicStageApp/Source/EpicStageApp/Public/StageAppVersion.h
--- a/Engine/Plugins/VirtualProduction/EpicStageApp/Source/EpicStageApp/Public/StageAppVersion.h
+++ b/Engine/Plugins/VirtualProduction/EpicStageApp/Source/EpicStageApp/Public/StageAppVersion.h
@@ -11,7 +11,7 @@ struct EPICSTAGEAPP_API FEpicStageAppAPIVersion
 	static const uint16 Major = 1;
 
 	/** The minor component of the current version. */
-	static const uint16 Minor = 8;
+	static const uint16 Minor = 9;
 	
 	/** The patch component of the current version. */
 	static const uint16 Patch = 0;
diff --git a/Engine/Plugins/VirtualProduction/RemoteControl/Source/RemoteControl/Public/IRemoteControlModule.h b/Engine/Plugins/VirtualProduction/RemoteControl/Source/RemoteControl/Public/IRemoteControlModule.h
--- a/Engine/Plugins/VirtualProduction/RemoteControl/Source/RemoteControl/Public/IRemoteControlModule.h
+++ b/Engine/Plugins/VirtualProduction/RemoteControl/Source/RemoteControl/Public/IRemoteControlModule.h
@@ -152,6 +152,16 @@ enum class ERCModifyOperation : uint8
 	DIVIDE
 };
 
+/**
+ * Type of compression applied to WebSocket traffic
+ */
+UENUM()
+enum class ERCWebSocketCompressionMode : uint8
+{
+	NONE,
+	ZLIB
+};
+
 /**
  * Reference to a UObject or one of its properties
  */
diff --git a/Engine/Plugins/VirtualProduction/RemoteControl/Source/WebRemoteControl/Private/RemoteControlRequest.h b/Engine/Plugins/VirtualProduction/RemoteControl/Source/WebRemoteControl/Private/RemoteControlRequest.h
--- a/Engine/Plugins/VirtualProduction/RemoteControl/Source/WebRemoteControl/Private/RemoteControlRequest.h
+++ b/Engine/Plugins/VirtualProduction/RemoteControl/Source/WebRemoteControl/Private/RemoteControlRequest.h
@@ -787,6 +787,31 @@ struct FRCWebSocketTransactionEndBody : public FRCRequest
 	int32 TransactionId = -1;
 };
 
+/**
+ * Holds a request made via websocket to change the compression mode.
+ */
+USTRUCT()
+struct FRCWebSocketCompressionChangeBody : public FRCRequest
+{
+	GENERATED_BODY()
+
+	FRCWebSocketCompressionChangeBody()
+	{
+		AddStructParameter(ParametersFieldLabel());
+	}
+
+	/**
+	 * Get the label for the property value struct.
+	 */
+	static FString ParametersFieldLabel() { return TEXT("Parameters"); }
+
+	/**
+	 * The compression mode to use.
+	 */
+	UPROPERTY()
+	ERCWebSocketCompressionMode Mode = ERCWebSocketCompressionMode::NONE;
+};
+
 /**
  * Struct representation of SetPresetController HTTP request
  */
diff --git a/Engine/Plugins/VirtualProduction/RemoteControl/Source/WebRemoteControl/Private/RemoteControlResponse.h b/Engine/Plugins/VirtualProduction/RemoteControl/Source/WebRemoteControl/Private/RemoteControlResponse.h
--- a/Engine/Plugins/VirtualProduction/RemoteControl/Source/WebRemoteControl/Private/RemoteControlResponse.h
+++ b/Engine/Plugins/VirtualProduction/RemoteControl/Source/WebRemoteControl/Private/RemoteControlResponse.h
@@ -609,3 +609,30 @@ struct FRCTransactionEndedEvent
 	int64 SequenceNumber = -1;
 };
 
+/**
+ * Event sent to confirm that the WebSocket compression mode has changed.
+ * This will be sent using the previous compression mode, and messages after this point will use the new one indicated in this event.
+ */
+USTRUCT()
+struct FRCCompressionChangedEvent
+{
+	GENERATED_BODY()
+
+	FRCCompressionChangedEvent()
+	: Type(TEXT("CompressionChanged"))
+	{
+	}
+
+	/**
+	 * Type of the event.
+	 */
+	UPROPERTY()
+	FString Type;
+
+	/**
+	 * The new compression mode which will be used by future messages.
+	 */
+	UPROPERTY()
+	ERCWebSocketCompressionMode Mode = ERCWebSocketCompressionMode::NONE;
+};
+
diff --git a/Engine/Plugins/VirtualProduction/RemoteControl/Source/WebRemoteControl/Private/RemoteControlWebSocketServer.cpp b/Engine/Plugins/VirtualProduction/RemoteControl/Source/WebRemoteControl/Private/RemoteControlWebSocketServer.cpp
--- a/Engine/Plugins/VirtualProduction/RemoteControl/Source/WebRemoteControl/Private/RemoteControlWebSocketServer.cpp
+++ b/Engine/Plugins/VirtualProduction/RemoteControl/Source/WebRemoteControl/Private/RemoteControlWebSocketServer.cpp
@@ -6,6 +6,7 @@
 #include "IPAddress.h"
 #include "IRemoteControlModule.h"
 #include "IWebSocketNetworkingModule.h"
+#include "Misc/Compression.h"
 #include "Misc/WildcardString.h"
 #include "RemoteControlRequest.h"
 #include "RemoteControlSettings.h"
@@ -16,6 +17,12 @@
 
 #define LOCTEXT_NAMESPACE "RCWebSocketServer"
 
+static TAutoConsoleVariable<int32> CVarWebSocketMaxUncompressedMessageSize(
+	TEXT("WebSocket.MaxUncompressedMessageSize"),
+	1024 * 1024 * 256,
+	TEXT("If a compressed WebSocket message reports an uncompressed size larger than this, reject the message.")
+);
+
 namespace RemoteControlWebSocketServer
 {
 	static const FString MessageNameFieldName = TEXT("MessageName");
@@ -180,19 +187,16 @@ void FRCWebSocketServer::Broadcast(const TArray<uint8>& InUTF8Payload)
 {
 	for (FWebSocketConnection& Connection : Connections)
 	{
-		if (Connection.Socket)
-		{
-			Connection.Socket->Send(InUTF8Payload.GetData(), InUTF8Payload.Num(), /*PrependSize=*/false);
-		}
+		SendOnConnection(Connection, InUTF8Payload);
 	}
 }
 
 void FRCWebSocketServer::Send(const FGuid& InTargetClientId, const TArray<uint8>& InUTF8Payload)
 {
 	TRACE_CPUPROFILER_EVENT_SCOPE(FRCWebSocketServer::Send);
-	if (FWebSocketConnection* Connection = Connections.FindByPredicate([&InTargetClientId](const FWebSocketConnection& InConnection) { return InConnection.Id == InTargetClientId; }))
+	if (FWebSocketConnection* Connection = GetClientById(InTargetClientId))
 	{
-		Connection->Socket->Send(InUTF8Payload.GetData(), InUTF8Payload.Num(), /*PrependSize=*/false);
+		SendOnConnection(*Connection, InUTF8Payload);
 	}
 }
 
@@ -201,6 +205,20 @@ bool FRCWebSocketServer::IsRunning() const
 	return !!Server;
 }
 
+void FRCWebSocketServer::SetClientCompressionMode(const FGuid& ClientId, ERCWebSocketCompressionMode Mode)
+{
+	for (FWebSocketConnection& Connection : Connections)
+	{
+		if (ClientId != Connection.Id)
+		{
+			continue;
+		}
+
+		Connection.CompressionMode = Mode;
+		break;
+	}
+}
+
 bool FRCWebSocketServer::Tick(float DeltaTime)
 {
 	Server->Tick();
@@ -234,9 +252,65 @@ void FRCWebSocketServer::ReceivedRawPacket(void* Data, int32 Size, FGuid ClientI
 	}
 	
 	TRACE_CPUPROFILER_EVENT_SCOPE(FRCWebSocketServer::ReceivedRawPacket);
-	
+
 	TArray<uint8> Payload;
-	WebRemoteControlUtils::ConvertToTCHAR(MakeArrayView(static_cast<uint8*>(Data), Size), Payload);
+
+	if (FWebSocketConnection* Connection = GetClientById(ClientId))
+	{
+		switch (Connection->CompressionMode)
+		{
+			case ERCWebSocketCompressionMode::ZLIB:
+			{
+				TRACE_CPUPROFILER_EVENT_SCOPE(TEXT("DecompressWebSocketData"));
+				
+				// Read the header containing the message's uncompressed size
+				const TArrayView<uint8> CompressedData = MakeArrayView(static_cast<uint8*>(Data), Size);
+				FMemoryReaderView Archive(CompressedData);
+
+#if !PLATFORM_LITTLE_ENDIAN
+				// Use canonical little-endian ordering
+				Archive.SetByteSwapping(true);
+#endif
+
+				int64 HeaderSize;
+				int32 UncompressedSize;
+				Archive << UncompressedSize;
+				HeaderSize = Archive.Tell();
+
+				if (UncompressedSize <= 0)
+				{
+					// Invalid size; don't try to allocate this or we'll crash
+					return;
+				}
+
+				if (UncompressedSize > CVarWebSocketMaxUncompressedMessageSize.GetValueOnGameThread())
+				{
+					return;
+				}
+
+				TArray<uint8> UncompressedData;
+				UncompressedData.SetNumUninitialized(UncompressedSize);
+
+				const bool bUncompressOk = FCompression::UncompressMemory(
+					NAME_Zlib,
+					UncompressedData.GetData(), UncompressedSize,
+					CompressedData.GetData() + HeaderSize, CompressedData.Num() - HeaderSize
+				);
+
+				if (!bUncompressOk)
+				{
+					return;
+				}
+
+				WebRemoteControlUtils::ConvertToTCHAR(UncompressedData, Payload);
+			}
+			break;
+
+			default:
+				WebRemoteControlUtils::ConvertToTCHAR(MakeArrayView(static_cast<uint8*>(Data), Size), Payload);
+				break;
+		}
+	}
 
 	if (TOptional<FRemoteControlWebSocketMessage> Message = RemoteControlWebSocketServer::ParseWebsocketMessage(Payload))
 	{
@@ -257,6 +331,48 @@ void FRCWebSocketServer::OnSocketClose(INetworkingWebSocket* Socket)
 	}
 }
 
+FRCWebSocketServer::FWebSocketConnection* FRCWebSocketServer::GetClientById(const FGuid& Id)
+{
+	return Connections.FindByPredicate([&Id](const FWebSocketConnection& InConnection) { return InConnection.Id == Id; });
+}
+
+void FRCWebSocketServer::SendOnConnection(FWebSocketConnection& Connection, const TArray<uint8>& InUTF8Payload)
+{
+	if (!Connection.Socket)
+	{
+		return;
+	}
+
+	switch (Connection.CompressionMode)
+	{
+	case ERCWebSocketCompressionMode::ZLIB:
+		{
+			TRACE_CPUPROFILER_EVENT_SCOPE(TEXT("CompressWebSocketData"));
+
+			int32 CompressedSize = InUTF8Payload.Num();
+
+			TArray<uint8> CompressedPayload;
+			CompressedPayload.SetNumUninitialized(CompressedSize);
+
+			const bool bCompressOk = FCompression::CompressMemory(
+				NAME_Zlib,
+				CompressedPayload.GetData(), CompressedSize,
+				InUTF8Payload.GetData(), InUTF8Payload.Num()
+			);
+
+			if (bCompressOk && CompressedSize < InUTF8Payload.Num())
+			{
+				Connection.Socket->Send(CompressedPayload.GetData(), CompressedSize, /*PrependSize=*/false);
+				return;
+			}
+		}
+		break;
+	}
+
+	// Send uncompressed data
+	Connection.Socket->Send(InUTF8Payload.GetData(), InUTF8Payload.Num(), /*PrependSize=*/false);
+}
+
 EWebsocketConnectionFilterResult FRCWebSocketServer::FilterConnection(FString OriginHeader, FString ClientIP) const
 {
 	const URemoteControlSettings* Settings = GetDefault<URemoteControlSettings>();
diff --git a/Engine/Plugins/VirtualProduction/RemoteControl/Source/WebRemoteControl/Private/RemoteControlWebSocketServer.h b/Engine/Plugins/VirtualProduction/RemoteControl/Source/WebRemoteControl/Private/RemoteControlWebSocketServer.h
--- a/Engine/Plugins/VirtualProduction/RemoteControl/Source/WebRemoteControl/Private/RemoteControlWebSocketServer.h
+++ b/Engine/Plugins/VirtualProduction/RemoteControl/Source/WebRemoteControl/Private/RemoteControlWebSocketServer.h
@@ -6,6 +6,7 @@
 
 #include "Containers/Ticker.h"
 #include "INetworkingWebSocket.h"
+#include "IRemoteControlModule.h"
 #include "IWebRemoteControlModule.h"
 #include "IWebSocketServer.h"
 #include "RemoteControlRoute.h"
@@ -114,7 +115,12 @@ public:
 	/** Callback when a socket is closed */
 	FOnWebSocketConnectionClosed& OnConnectionClosed() { return OnConnectionClosedDelegate; }
 
+	/** Set the compression mode for a client by its GUID. */
+	void SetClientCompressionMode(const FGuid& ClientId, ERCWebSocketCompressionMode Mode);
+
 private:
+	class FWebSocketConnection;
+
 	bool Tick(float DeltaTime);
 
 	/** Handles a new client connecting. */
@@ -124,6 +130,12 @@ private:
 	void ReceivedRawPacket(void* Data, int32 Size, FGuid ClientId, TSharedPtr<FInternetAddr> PeerAddress);
 
 	void OnSocketClose(INetworkingWebSocket* Socket);
+
+	/** Given a client ID, find the corresponding client, or null if it's not connected to this server. */
+	FWebSocketConnection* GetClientById(const FGuid& Id);
+
+	/** Send a message on a specific connection */
+	void SendOnConnection(FWebSocketConnection& Connection, const TArray<uint8>& InUTF8Payload);
 	
 	/** Handle rejecting the websocket connection if it doesn't respect the user's CORS policy. */
 	EWebsocketConnectionFilterResult FilterConnection(FString OriginHeader, FString ClientIP) const;
@@ -177,10 +189,11 @@ private:
 
 		/** IP Address of the client. */
 		TSharedPtr<FInternetAddr> PeerAddress;
+		
+		/** Compression mode to use for this client. */
+		ERCWebSocketCompressionMode CompressionMode = ERCWebSocketCompressionMode::NONE;
 	};
 
-
-
 private:
 	/** Handle to the tick delegate. */
 	FTSTicker::FDelegateHandle TickerHandle;
diff --git a/Engine/Plugins/VirtualProduction/RemoteControl/Source/WebRemoteControl/Private/WebSocketMessageHandler.cpp b/Engine/Plugins/VirtualProduction/RemoteControl/Source/WebRemoteControl/Private/WebSocketMessageHandler.cpp
--- a/Engine/Plugins/VirtualProduction/RemoteControl/Source/WebRemoteControl/Private/WebSocketMessageHandler.cpp
+++ b/Engine/Plugins/VirtualProduction/RemoteControl/Source/WebRemoteControl/Private/WebSocketMessageHandler.cpp
@@ -12,7 +12,6 @@
 #include "RCVirtualProperty.h"
 #include "GameFramework/Actor.h"
 #include "RemoteControlPreset.h"
-#include "RemoteControlRequest.h"
 #include "RemoteControlRoute.h"
 #include "RemoteControlReflectionUtils.h"
 #include "RemoteControlWebsocketRoute.h"
@@ -373,6 +372,12 @@ void FWebSocketMessageHandler::RegisterRoutes(FWebRemoteControlModule* WebRemote
 		TEXT("transaction.end"),
 		FWebSocketMessageDelegate::CreateRaw(this, &FWebSocketMessageHandler::HandleWebSocketEndEditorTransaction)
 	));
+
+	RegisterRoute(WebRemoteControl, MakeUnique<FRemoteControlWebsocketRoute>(
+		TEXT("Change the compression method this client will use to communicate with Unreal Engine"),
+		TEXT("compression.change"),
+		FWebSocketMessageDelegate::CreateRaw(this, &FWebSocketMessageHandler::HandleWebSocketCompressionChange)
+	));
 }
 
 void FWebSocketMessageHandler::UnregisterRoutes(FWebRemoteControlModule* WebRemoteControl)
@@ -837,6 +842,26 @@ void FWebSocketMessageHandler::HandleWebSocketEndEditorTransaction(const FRemote
 #endif
 }
 
+void FWebSocketMessageHandler::HandleWebSocketCompressionChange(const FRemoteControlWebSocketMessage& WebSocketMessage)
+{
+	FRCWebSocketCompressionChangeBody Body;
+	if (!WebRemoteControlInternalUtils::DeserializeRequestPayload(WebSocketMessage.RequestPayload, nullptr, Body))
+	{
+		return;
+	}
+
+	// First, reply to the client confirming the new mode so it can continue with its old decompression method until it receives this message
+	FRCCompressionChangedEvent Event;
+	Event.Mode = Body.Mode;
+
+	TArray<uint8> Payload;
+	WebRemoteControlUtils::SerializeMessage(Event, Payload);
+	Server->Send(WebSocketMessage.ClientId, Payload);
+
+	// Now update the compression method for all future messages
+	Server->SetClientCompressionMode(WebSocketMessage.ClientId, Body.Mode);
+}
+
 void FWebSocketMessageHandler::ProcessChangedControllers()
 {
 	// Go over each controller that was changed for each preset
@@ -1417,7 +1442,7 @@ void FWebSocketMessageHandler::TimeOutTransactions()
 	}
 
 	// Do this as a separate step since it may remove entries from the map during iteration
-	for (auto TransactionIterator : TransactionsToEnd)
+	for (const auto& TransactionIterator : TransactionsToEnd)
 	{
 		EndClientTransaction(TransactionIterator.Key, TransactionIterator.Value);
 	}
diff --git a/Engine/Plugins/VirtualProduction/RemoteControl/Source/WebRemoteControl/Private/WebSocketMessageHandler.h b/Engine/Plugins/VirtualProduction/RemoteControl/Source/WebRemoteControl/Private/WebSocketMessageHandler.h
--- a/Engine/Plugins/VirtualProduction/RemoteControl/Source/WebRemoteControl/Private/WebSocketMessageHandler.h
+++ b/Engine/Plugins/VirtualProduction/RemoteControl/Source/WebRemoteControl/Private/WebSocketMessageHandler.h
@@ -9,6 +9,7 @@
 #include "RemoteControlField.h"
 #include "RemoteControlModels.h"
 #include "RemoteControlPreset.h"
+#include "RemoteControlRequest.h"
 #include "RemoteControlResponse.h"
 
 
@@ -112,6 +113,9 @@ private:
 	/** Handles ending a manual editor transaction. */
 	void HandleWebSocketEndEditorTransaction(const FRemoteControlWebSocketMessage& WebSocketMessage);
 
+	/** Handles changing the compression mode of WebSocket data */
+	void HandleWebSocketCompressionChange(const FRemoteControlWebSocketMessage& WebSocketMessage);
+
 	//Preset callbacks
 	void OnPresetExposedPropertiesModified(URemoteControlPreset* Owner, const TSet<FGuid>& ModifiedPropertyIds);
 	void OnPropertyExposed(URemoteControlPreset* Owner,  const FGuid& EntityId);
