[RemoteControl] Add missing include to fix incremental build

diff --git a/Engine/Plugins/VirtualProduction/RemoteControl/Source/WebRemoteControl/Private/RemoteControlResponse.h b/Engine/Plugins/VirtualProduction/RemoteControl/Source/WebRemoteControl/Private/RemoteControlResponse.h
--- a/Engine/Plugins/VirtualProduction/RemoteControl/Source/WebRemoteControl/Private/RemoteControlResponse.h
+++ b/Engine/Plugins/VirtualProduction/RemoteControl/Source/WebRemoteControl/Private/RemoteControlResponse.h
@@ -3,6 +3,7 @@
 
 #include "AssetRegistry/AssetData.h"
 #include "CoreMinimal.h"
+#include "IRemoteControlModule.h"
 #include "RemoteControlField.h"
 #include "RemoteControlModels.h"
 #include "RemoteControlPreset.h"
