[RemoteControl] RemoteControlUI
Remove DeleteAllEntities button from Group in the Header and in the Footer from Properties/Controller/Behaviour and Action.
Move up Properties footer buttons (ExposeActor, ExposeFunction) to the header and remove all footer.

diff --git a/Engine/Plugins/VirtualProduction/RemoteControl/Source/RemoteControlUI/Private/UI/Action/SRCActionPanel.cpp b/Engine/Plugins/VirtualProduction/RemoteControl/Source/RemoteControlUI/Private/UI/Action/SRCActionPanel.cpp
--- a/Engine/Plugins/VirtualProduction/RemoteControl/Source/RemoteControlUI/Private/UI/Action/SRCActionPanel.cpp
+++ b/Engine/Plugins/VirtualProduction/RemoteControl/Source/RemoteControlUI/Private/UI/Action/SRCActionPanel.cpp
@@ -95,7 +95,7 @@ void SRCActionPanel::UpdateWrappedWidget(TSharedPtr<FRCBehaviourModel> InBehavio
 		// Action Dock Panel
 		TSharedPtr<SRCMinorPanel> ActionDockPanel = SNew(SRCMinorPanel)
 			.HeaderLabel(LOCTEXT("ActionsLabel", "Actions"))
-			.EnableFooter(true)
+			.EnableFooter(false)
 			[
 				ActionPanelList.ToSharedRef()
 			];
@@ -145,31 +145,9 @@ void SRCActionPanel::UpdateWrappedWidget(TSharedPtr<FRCBehaviourModel> InBehavio
 					.Image(FAppStyle::GetBrush("Icons.Duplicate"))
 				]
 			];
-		
-	// Delete Selected Action Button
-		TSharedRef<SWidget> DeleteSelectedActionButton = SNew(SButton)
-			.AddMetaData<FTagMetaData>(FTagMetaData(TEXT("Delete Selected Action")))
-			.HAlign(HAlign_Center)
-			.VAlign(VAlign_Center)
-			.ForegroundColor(FSlateColor::UseForeground())
-			.ButtonStyle(&RCPanelStyle->FlatButtonStyle)
-			.ToolTipText(LOCTEXT("DeleteSelectedActionToolTip", "Deletes all selected Action."))
-			.OnClicked(this, &SRCActionPanel::RequestDeleteSelectedItem)
-			.IsEnabled_Lambda([this]() { return ActionPanelList.IsValid() && ActionPanelList->NumSelectedLogicItems() > 0; })
-			[
-				SNew(SBox)
-				.WidthOverride(RCPanelStyle->IconSize.X)
-				.HeightOverride(RCPanelStyle->IconSize.Y)
-				[
-					SNew(SImage)
-					.ColorAndOpacity(FSlateColor::UseForeground())
-					.Image(FAppStyle::GetBrush("Icons.Delete"))
-				]
-			];
 
 		ActionDockPanel->AddHeaderToolbarItem(EToolbar::Left, AddNewActionButton);
 		ActionDockPanel->AddHeaderToolbarItem(EToolbar::Right, AddAllActionsButton);
-		ActionDockPanel->AddFooterToolbarItem(EToolbar::Right, DeleteSelectedActionButton);
 
 		// Header Dock Panel
 		TSharedPtr<SRCMinorPanel> BehaviourDetailsPanel = SNew(SRCMinorPanel)
diff --git a/Engine/Plugins/VirtualProduction/RemoteControl/Source/RemoteControlUI/Private/UI/Behaviour/SRCBehaviourPanel.cpp b/Engine/Plugins/VirtualProduction/RemoteControl/Source/RemoteControlUI/Private/UI/Behaviour/SRCBehaviourPanel.cpp
--- a/Engine/Plugins/VirtualProduction/RemoteControl/Source/RemoteControlUI/Private/UI/Behaviour/SRCBehaviourPanel.cpp
+++ b/Engine/Plugins/VirtualProduction/RemoteControl/Source/RemoteControlUI/Private/UI/Behaviour/SRCBehaviourPanel.cpp
@@ -79,7 +79,7 @@ void SRCBehaviourPanel::UpdateWrappedWidget(TSharedPtr<FRCControllerModel> InCon
 		// Behaviour Dock Panel
 		TSharedPtr<SRCMinorPanel> BehaviourDockPanel = SNew(SRCMinorPanel)
 			.HeaderLabel(LOCTEXT("BehavioursLabel", "Behavior"))
-			.EnableFooter(true)
+			.EnableFooter(false)
 			[
 				SAssignNew(BehaviourPanelList, SRCBehaviourPanelList, SharedThis(this), InControllerItem, RemoteControlPanel.ToSharedRef())
 			];
@@ -110,29 +110,7 @@ void SRCBehaviourPanel::UpdateWrappedWidget(TSharedPtr<FRCControllerModel> InCon
 				GetBehaviourMenuContentWidget()
 			];
 
-		// Delete Selected Behaviour Button
-		TSharedRef<SWidget> DeleteSelectedBehaviourButton = SNew(SButton)
-			.AddMetaData<FTagMetaData>(FTagMetaData(TEXT("Delete Selected Behaviour")))
-			.HAlign(HAlign_Center)
-			.VAlign(VAlign_Center)
-			.ForegroundColor(FSlateColor::UseForeground())
-			.ButtonStyle(&RCPanelStyle->FlatButtonStyle)
-			.ToolTipText(LOCTEXT("DeleteSelectedBehaviourToolTip", "Deletes the selected behaviour."))
-			.OnClicked(this, &SRCBehaviourPanel::RequestDeleteSelectedItem)
-			.IsEnabled_Lambda([this]() { return BehaviourPanelList.IsValid() && BehaviourPanelList->NumSelectedLogicItems(); })
-			[
-				SNew(SBox)
-				.WidthOverride(RCPanelStyle->IconSize.X)
-				.HeightOverride(RCPanelStyle->IconSize.Y)
-				[
-					SNew(SImage)
-					.ColorAndOpacity(FSlateColor::UseForeground())
-					.Image(FAppStyle::GetBrush("Icons.Delete"))
-				]
-			];
-
 		BehaviourDockPanel->AddHeaderToolbarItem(EToolbar::Left, AddNewBehaviourButton);
-		BehaviourDockPanel->AddFooterToolbarItem(EToolbar::Right, DeleteSelectedBehaviourButton);
 
 		WrappedBoxWidget->SetContent(BehaviourDockPanel.ToSharedRef());
 	}
diff --git a/Engine/Plugins/VirtualProduction/RemoteControl/Source/RemoteControlUI/Private/UI/Controller/SRCControllerPanel.cpp b/Engine/Plugins/VirtualProduction/RemoteControl/Source/RemoteControlUI/Private/UI/Controller/SRCControllerPanel.cpp
--- a/Engine/Plugins/VirtualProduction/RemoteControl/Source/RemoteControlUI/Private/UI/Controller/SRCControllerPanel.cpp
+++ b/Engine/Plugins/VirtualProduction/RemoteControl/Source/RemoteControlUI/Private/UI/Controller/SRCControllerPanel.cpp
@@ -43,7 +43,7 @@ void SRCControllerPanel::Construct(const FArguments& InArgs, const TSharedRef<SR
 	// Controller Dock Panel
 	TSharedPtr<SRCMinorPanel> ControllerDockPanel = SNew(SRCMinorPanel)
 		.HeaderLabel(LOCTEXT("ControllersLabel", "Controller"))
-		.EnableFooter(true)
+		.EnableFooter(false)
 		[
 			SAssignNew(ControllerPanelList, SRCControllerPanelList, SharedThis(this), InPanel)
 		];
@@ -75,29 +75,7 @@ void SRCControllerPanel::Construct(const FArguments& InArgs, const TSharedRef<SR
 			GetControllerMenuContentWidget()
 		];
 
-	// Delete Selected Controller Button
-	TSharedRef<SWidget> DeleteSelectedControllerButton = SNew(SButton)
-		.AddMetaData<FTagMetaData>(FTagMetaData(TEXT("Delete Selected Controller")))
-		.HAlign(HAlign_Center)
-		.VAlign(VAlign_Center)
-		.ForegroundColor(FSlateColor::UseForeground())
-		.ButtonStyle(&RCPanelStyle->FlatButtonStyle)
-		.ToolTipText(LOCTEXT("DeleteSelectedControllerToolTip", "Deletes the selected controller."))
-		.OnClicked(this, &SRCControllerPanel::RequestDeleteSelectedItem)
-		.IsEnabled_Lambda([this]() { return ControllerPanelList.IsValid() && !bIsInLiveMode.Get() && ControllerPanelList->NumSelectedLogicItems() > 0; })
-		[
-			SNew(SBox)
-			.WidthOverride(RCPanelStyle->IconSize.X)
-			.HeightOverride(RCPanelStyle->IconSize.Y)
-			[
-				SNew(SImage)
-				.ColorAndOpacity(FSlateColor::UseForeground())
-				.Image(FAppStyle::GetBrush("Icons.Delete"))
-			]
-		];
-
 	ControllerDockPanel->AddHeaderToolbarItem(EToolbar::Left, AddNewControllerButton);
-	ControllerDockPanel->AddFooterToolbarItem(EToolbar::Right, DeleteSelectedControllerButton);
 	ControllerDockPanel->AddHeaderToolbarItem(EToolbar::Right, GetMultiControllerSwitchWidget());
 
 	ChildSlot
diff --git a/Engine/Plugins/VirtualProduction/RemoteControl/Source/RemoteControlUI/Private/UI/SRCPanelExposedEntitiesList.cpp b/Engine/Plugins/VirtualProduction/RemoteControl/Source/RemoteControlUI/Private/UI/SRCPanelExposedEntitiesList.cpp
--- a/Engine/Plugins/VirtualProduction/RemoteControl/Source/RemoteControlUI/Private/UI/SRCPanelExposedEntitiesList.cpp
+++ b/Engine/Plugins/VirtualProduction/RemoteControl/Source/RemoteControlUI/Private/UI/SRCPanelExposedEntitiesList.cpp
@@ -242,30 +242,7 @@ void SRCPanelExposedEntitiesList::Construct(const FArguments& InArgs, URemoteCon
 			]
 		];
 
-	// Delete All Groups Button
-	TSharedPtr<SWidget> DeleteAllGroupsButton = SNew(SButton)
-		.AddMetaData<FTagMetaData>(FTagMetaData(TEXT("Delete All Groups")))
-		.IsEnabled_Lambda([this]() { return !bIsInLiveMode.Get() && FieldGroups.Num() > 0; })
-		.Visibility_Lambda([this]() { return (!bIsInLiveMode.Get() && FieldGroups.Num() > 0) ? EVisibility::Visible : EVisibility::Collapsed; })
-		.HAlign(HAlign_Center)
-		.VAlign(VAlign_Center)
-		.ForegroundColor(FSlateColor::UseForeground())
-		.ButtonStyle(&RCPanelStyle->FlatButtonStyle)
-		.ToolTipText(LOCTEXT("DeleteAllGroupsToolTip", "Delete all groups."))
-		.OnClicked(this, &SRCPanelExposedEntitiesList::RequestDeleteAllGroups)
-		[
-			SNew(SBox)
-			.WidthOverride(RCPanelStyle->IconSize.X)
-			.HeightOverride(RCPanelStyle->IconSize.Y)
-			[
-				SNew(SImage)
-				.ColorAndOpacity(FSlateColor::UseForeground())
-				.Image(FAppStyle::GetBrush("Icons.Delete"))
-			]
-		];
-
 	GroupDockPanel->AddHeaderToolbarItem(EToolbar::Left, NewGroupButton.ToSharedRef());
-	GroupDockPanel->AddHeaderToolbarItem(EToolbar::Right, DeleteAllGroupsButton.ToSharedRef());
 
 	ExposePanel->AddPanel(GroupDockPanel.ToSharedRef(), 0.25f);
 
@@ -328,42 +305,11 @@ void SRCPanelExposedEntitiesList::Construct(const FArguments& InArgs, URemoteCon
 	TSharedPtr<SRCMinorPanel> ExposeDockPanel = SNew(SRCMinorPanel)
 		.HeaderLabel(this, &SRCPanelExposedEntitiesList::HandleEntityListHeaderLabel)
 		.Visibility_Lambda([this]() {return bIsInLiveMode.Get() ? EVisibility::Collapsed : EVisibility::Visible; })
-		.EnableFooter(true)
+		.EnableFooter(false)
 		[
 			FieldsListView.ToSharedRef()
 		];
 
-	// Placeholder Box
-	TSharedPtr<SWidget> PlaceholderBox = SNew(SBox)
-		.AddMetaData<FTagMetaData>(FTagMetaData(TEXT("Placeholder")))
-		.WidthOverride(30.f)
-		.HeightOverride(30.f)
-		[
-			SNullWidget::NullWidget
-		];
-
-	// Delete All Entities Button
-	TSharedPtr<SWidget> DeleteAllEntitiesButton = SNew(SButton)
-		.AddMetaData<FTagMetaData>(FTagMetaData(TEXT("Delete All Entities")))
-		.IsEnabled_Lambda([this]() { return !bIsInLiveMode.Get() && FieldWidgetMap.Num() > 0; })
-		.Visibility_Lambda([this]() { return (!bIsInLiveMode.Get() && FieldWidgetMap.Num() > 0) ? EVisibility::Visible : EVisibility::Collapsed; })
-		.HAlign(HAlign_Center)
-		.VAlign(VAlign_Center)
-		.ForegroundColor(FSlateColor::UseForeground())
-		.ButtonStyle(&RCPanelStyle->FlatButtonStyle)
-		.ToolTipText(LOCTEXT("DeleteAllEntitiesToolTip", "Delete all entities."))
-		.OnClicked(this, &SRCPanelExposedEntitiesList::RequestDeleteAllEntities)
-		[
-			SNew(SBox)
-			.WidthOverride(RCPanelStyle->IconSize.X)
-			.HeightOverride(RCPanelStyle->IconSize.Y)
-			[
-				SNew(SImage)
-				.ColorAndOpacity(FSlateColor::UseForeground())
-				.Image(FAppStyle::GetBrush("Icons.Delete"))
-			]
-		];
-
 	// Mode Switcher
 	TSharedPtr<SRCModeSwitcher> ModeSwitcher = SNew(SRCModeSwitcher)
 		.OnModeSwitched_Lambda([this](const SRCModeSwitcher::FRCMode& NewMode)
@@ -427,11 +373,9 @@ void SRCPanelExposedEntitiesList::Construct(const FArguments& InArgs, URemoteCon
 	}
 
 	// Expose Button
-	ExposeDockPanel->AddHeaderToolbarItem(EToolbar::Left, PlaceholderBox.ToSharedRef());
 	ExposeDockPanel->AddHeaderToolbarItem(EToolbar::Right, ModeSwitcher.ToSharedRef());
-	ExposeDockPanel->AddFooterToolbarItem(EToolbar::Left, InArgs._ExposeActorsComboButton.Get().ToSharedRef());
-	ExposeDockPanel->AddFooterToolbarItem(EToolbar::Left, InArgs._ExposeFunctionsComboButton.Get().ToSharedRef());
-	ExposeDockPanel->AddFooterToolbarItem(EToolbar::Right, DeleteAllEntitiesButton.ToSharedRef());
+	ExposeDockPanel->AddHeaderToolbarItem(EToolbar::Left, InArgs._ExposeActorsComboButton.Get().ToSharedRef());
+	ExposeDockPanel->AddHeaderToolbarItem(EToolbar::Left, InArgs._ExposeFunctionsComboButton.Get().ToSharedRef());
 
 	ExposePanel->AddPanel(ExposeDockPanel.ToSharedRef(), 0.75f);
 
