[Remote Control] Remove Custom Controller execution on post load

Custom controllers added automatic execution on PostLoad (see https://p4-swarm.epicgames.net/changes/25874141).

This CL removes the support for that (given that the execution of controllers is non-deterministic), and gives the control back to the user.
This change also prevents issues related to unreliability, and assets being automatically modified/dirtied on scene load.

diff --git a/Engine/Plugins/VirtualProduction/RemoteControl/Source/RemoteControlLogic/Private/Action/RCActionContainer.cpp b/Engine/Plugins/VirtualProduction/RemoteControl/Source/RemoteControlLogic/Private/Action/RCActionContainer.cpp
--- a/Engine/Plugins/VirtualProduction/RemoteControl/Source/RemoteControlLogic/Private/Action/RCActionContainer.cpp
+++ b/Engine/Plugins/VirtualProduction/RemoteControl/Source/RemoteControlLogic/Private/Action/RCActionContainer.cpp
@@ -172,51 +172,6 @@ void URCActionContainer::PostEditUndo()
 }
 #endif
 
-void URCActionContainer::ExecuteActionsOnLoad()
-{
-	// In some cases, external resources are referenced, so we need to execute Controllers actions while loading
-	if (URCBehaviour* ParentBehaviour = GetParentBehaviour())
-	{
-		if (URCController* Controller = ParentBehaviour->ControllerWeakPtr.Get())
-		{
-			if (UE::RCCustomControllers::CustomControllerExecutesOnLoad(Controller))
-			{
-				Controller->ExecuteBehaviours();
-				return;
-			}
-		}
-
-		if (URCSetAssetByPathBehaviour* AssetPathBehaviour = Cast<URCSetAssetByPathBehaviour>(ParentBehaviour))
-		{
-			if (!AssetPathBehaviour->bInternal)
-			{
-				// Target entity might not be set up yet, we need it to be set in order for execution to work
-				AssetPathBehaviour->UpdateTargetEntity();
-				ParentBehaviour->Execute();
-			}
-		}
-	}
-}
-
-void URCActionContainer::PostLoad()
-{
-	UObject::PostLoad();
-
-	TWeakObjectPtr<URCActionContainer> ThisWeak(this);
-
-	FTSTicker::GetCoreTicker().AddTicker(FTickerDelegate::CreateLambda(
-		[ThisWeak](float InDeltaTime)->bool
-		{
-			if (URCActionContainer* const This = ThisWeak.Get())
-			{
-				This->ExecuteActionsOnLoad();
-			}
-
-			// Return false for one time execution
-			return false;
-		}));
-}
-
 URCAction* URCActionContainer::FindActionByFieldId(const FGuid InId) const
 {
 	for (URCAction* Action :  Actions)
diff --git a/Engine/Plugins/VirtualProduction/RemoteControl/Source/RemoteControlLogic/Private/Controller/RCCustomControllerUtilities.cpp b/Engine/Plugins/VirtualProduction/RemoteControl/Source/RemoteControlLogic/Private/Controller/RCCustomControllerUtilities.cpp
--- a/Engine/Plugins/VirtualProduction/RemoteControl/Source/RemoteControlLogic/Private/Controller/RCCustomControllerUtilities.cpp
+++ b/Engine/Plugins/VirtualProduction/RemoteControl/Source/RemoteControlLogic/Private/Controller/RCCustomControllerUtilities.cpp
@@ -9,10 +9,6 @@ namespace UE::RCCustomControllers
 	// A Map to associate a custom controller name to its underlying type
 	inline static TMap<FString, EPropertyBagPropertyType> CustomControllerTypes =
 		{{CustomTextureControllerName, EPropertyBagPropertyType::String}};
-
-	// A Map to associate a custom controller name to its underlying type
-	inline static TSet<FString> ExecuteOnLoadControllers =
-		{CustomTextureControllerName};
 }
 
 bool UE::RCCustomControllers::IsCustomController(const URCVirtualPropertyBase* InController)
@@ -25,16 +21,6 @@ FString UE::RCCustomControllers::GetCustomControllerTypeName(const URCVirtualPro
 	return InController->GetMetadataValue(CustomControllerNameKey);
 }
 
-bool UE::RCCustomControllers::CustomControllerExecutesOnLoad(const URCVirtualPropertyBase* InController)
-{
-	return InController->GetMetadataValue(CustomControllerExecuteOnLoadNameKey).ToBool();
-}
-
-bool UE::RCCustomControllers::CustomControllerExecutesOnLoad(const FName& InCustomControllerTypeName)
-{
-	return ExecuteOnLoadControllers.Contains(InCustomControllerTypeName.ToString());
-}
-
 bool UE::RCCustomControllers::IsValidCustomController(const FName& InCustomControllerTypeName)
 {
 	return CustomControllerTypes.Contains(InCustomControllerTypeName.ToString());
@@ -60,11 +46,8 @@ TMap<FName, FString> UE::RCCustomControllers::GetCustomControllerMetaData(const
 	TMap<FName, FString> OutMetaData;
 	OutMetaData.Emplace(UE::RCCustomControllers::CustomControllerNameKey, InCustomControllerTypeName);
 
-	if (UE::RCCustomControllers::CustomControllerExecutesOnLoad(FName(*InCustomControllerTypeName)))
-	{
-		OutMetaData.Emplace(UE::RCCustomControllers::CustomControllerExecuteOnLoadNameKey, TEXT("true"));
-	}
-
+	// Note: this is a good place to additionally edit OutMetaData by manually adding elements here if custom logic is needed
+	
 	return OutMetaData;
 }
 
diff --git a/Engine/Plugins/VirtualProduction/RemoteControl/Source/RemoteControlLogic/Public/Action/RCActionContainer.h b/Engine/Plugins/VirtualProduction/RemoteControl/Source/RemoteControlLogic/Public/Action/RCActionContainer.h
--- a/Engine/Plugins/VirtualProduction/RemoteControl/Source/RemoteControlLogic/Public/Action/RCActionContainer.h
+++ b/Engine/Plugins/VirtualProduction/RemoteControl/Source/RemoteControlLogic/Public/Action/RCActionContainer.h
@@ -72,7 +72,6 @@ public:
 	/** Called after applying a transaction to the object. Used to broadcast Undo related container changes to UI */
 	virtual void PostEditUndo() override;
 #endif
-	virtual void PostLoad() override;
 	//~ End UObject
 
 	/** Set of child action container */
@@ -95,8 +94,6 @@ private:
 
 	/** Add remote control property function */
 	URCFunctionAction* AddFunctionAction(const TSharedRef<const FRemoteControlFunction> InRemoteControlFunction);
-
-	void ExecuteActionsOnLoad();
 	
 	/** The list of Actions present in this container */
 	UPROPERTY()
diff --git a/Engine/Plugins/VirtualProduction/RemoteControl/Source/RemoteControlLogic/Public/Controller/RCCustomControllerUtilities.h b/Engine/Plugins/VirtualProduction/RemoteControl/Source/RemoteControlLogic/Public/Controller/RCCustomControllerUtilities.h
--- a/Engine/Plugins/VirtualProduction/RemoteControl/Source/RemoteControlLogic/Public/Controller/RCCustomControllerUtilities.h
+++ b/Engine/Plugins/VirtualProduction/RemoteControl/Source/RemoteControlLogic/Public/Controller/RCCustomControllerUtilities.h
@@ -16,10 +16,12 @@ namespace UE::RCCustomControllers
 	
 	/** Returns the Custom Controller Type Name (e.g. Asset Path). Empty string if Controller is not a custom one. */
 	REMOTECONTROLLOGIC_API FString GetCustomControllerTypeName(const URCVirtualPropertyBase* InController);
+	
+	UE_DEPRECATED(5.4, "Custom Controllers execution on load is deprecated. Function will always return false.")
+	REMOTECONTROLLOGIC_API inline bool CustomControllerExecutesOnLoad(const URCVirtualPropertyBase* InController) { return false; }
 
-	REMOTECONTROLLOGIC_API bool CustomControllerExecutesOnLoad(const URCVirtualPropertyBase* InController);
-
-	REMOTECONTROLLOGIC_API bool CustomControllerExecutesOnLoad(const FName& InCustomControllerTypeName);
+	UE_DEPRECATED(5.4, "Custom Controllers execution on load is deprecated. Function will always return false.")
+	REMOTECONTROLLOGIC_API inline bool CustomControllerExecutesOnLoad(const FName& InCustomControllerTypeName) { return false; }
 	
 	/** Returns the type associated to the specified Custom Controller Type Name */
 	REMOTECONTROLLOGIC_API EPropertyBagPropertyType GetCustomControllerType(const FName& InCustomControllerTypeName);
