[RemoteControl] RemoteControl
New feature PropertyID and PropertyIdAction

diff --git a/Engine/Plugins/VirtualProduction/RemoteControl/Source/RemoteControl/Private/PropertyIdHandler/BasePropertyIdHandler.cpp b/Engine/Plugins/VirtualProduction/RemoteControl/Source/RemoteControl/Private/PropertyIdHandler/BasePropertyIdHandler.cpp
new file mode 100644
--- /dev/null
+++ b/Engine/Plugins/VirtualProduction/RemoteControl/Source/RemoteControl/Private/PropertyIdHandler/BasePropertyIdHandler.cpp
@@ -0,0 +1,106 @@
+// Copyright Epic Games, Inc. All Rights Reserved.
+
+#include "BasePropertyIdHandler.h"
+
+bool FBasePropertyIdHandler::IsPropertySupported(const FProperty* InProperty) const
+{
+	if (InProperty)
+	{
+		const FProperty* Property = TryGetPropertyInContainerOrSelf(InProperty);
+		return !(Property->GetClass()->IsChildOf(FStructProperty::StaticClass()) ||
+				 Property->GetClass()->IsChildOf(FObjectProperty::StaticClass()) ||
+				 Property->GetClass()->IsChildOf(FEnumProperty::StaticClass())	  );
+	}
+	return false;
+}
+
+EPropertyBagPropertyType FBasePropertyIdHandler::GetPropertyType(const FProperty* InProperty) const
+{
+	EPropertyBagPropertyType PropertyType = EPropertyBagPropertyType::None;
+	if (InProperty)
+	{
+		const FProperty* Property = TryGetPropertyInContainerOrSelf(InProperty);
+		const FName PropertyName = Property->GetClass()->GetFName();
+		if (PropertyName == NAME_BoolProperty)
+		{
+			PropertyType = EPropertyBagPropertyType::Bool;
+		}
+		else if (PropertyName == NAME_ByteProperty)
+		{
+			if (const FByteProperty* ByteProperty = CastField<FByteProperty>(InProperty))
+			{
+				if (ByteProperty->Enum)
+				{
+					PropertyType = EPropertyBagPropertyType::Enum;
+				}
+			}
+			else
+			{
+				PropertyType = EPropertyBagPropertyType::Byte;
+			}
+		}
+		else if (PropertyName == NAME_Int32Property)
+		{
+			PropertyType = EPropertyBagPropertyType::Int32;
+		}
+		else if (PropertyName == NAME_Int64Property)
+		{
+			PropertyType = EPropertyBagPropertyType::Int64;
+		}
+		else if (PropertyName == NAME_FloatProperty)
+		{
+			PropertyType = EPropertyBagPropertyType::Float;
+		}
+		else if (PropertyName == NAME_NameProperty)
+		{
+			PropertyType = EPropertyBagPropertyType::Name;
+		}
+		else if (PropertyName == NAME_StrProperty)
+		{
+			PropertyType = EPropertyBagPropertyType::String;
+		}
+		else if (PropertyName == NAME_TextProperty)
+		{
+			PropertyType = EPropertyBagPropertyType::Text;
+		}
+	}
+	return PropertyType;
+}
+
+FName FBasePropertyIdHandler::GetPropertyTypeName(const FProperty* InProperty) const
+{
+	if (InProperty)
+	{
+		const FProperty* Property = TryGetPropertyInContainerOrSelf(InProperty);
+		if (const FByteProperty* ByteProperty = CastField<FByteProperty>(Property))
+		{
+			if (ByteProperty->Enum)
+			{
+				return ByteProperty->Enum->GetFName();
+			}
+		}
+		return Property->GetClass()->GetFName();
+	}
+	return FName(TEXT(""));
+}
+
+UObject* FBasePropertyIdHandler::GetPropertyTypeObject(const FProperty* InProperty) const
+{
+	if (InProperty)
+	{
+		const FProperty* Property = TryGetPropertyInContainerOrSelf(InProperty);
+		if (const FByteProperty* ByteProperty = CastField<FByteProperty>(Property))
+		{
+			if (ByteProperty->Enum)
+			{
+				return ByteProperty->Enum;
+			}
+		}
+	}
+	return nullptr;
+}
+
+TObjectPtr<UObject> FBasePropertyIdHandler::GetObjectPropertyDefaultValue(const FProperty* InProperty, const UClass* InClassToCreate) const
+{
+	return nullptr;
+}
diff --git a/Engine/Plugins/VirtualProduction/RemoteControl/Source/RemoteControl/Private/PropertyIdHandler/BasePropertyIdHandler.h b/Engine/Plugins/VirtualProduction/RemoteControl/Source/RemoteControl/Private/PropertyIdHandler/BasePropertyIdHandler.h
new file mode 100644
--- /dev/null
+++ b/Engine/Plugins/VirtualProduction/RemoteControl/Source/RemoteControl/Private/PropertyIdHandler/BasePropertyIdHandler.h
@@ -0,0 +1,17 @@
+// Copyright Epic Games, Inc. All Rights Reserved.
+
+#pragma once
+
+#include "IPropertyIdHandler.h"
+
+class FBasePropertyIdHandler : public IPropertyIdHandler
+{
+public:
+	//~ Begin IPropertyIdHandler
+	virtual bool IsPropertySupported(const FProperty* InProperty) const override;
+	virtual EPropertyBagPropertyType GetPropertyType(const FProperty* InProperty) const override;
+	virtual FName GetPropertyTypeName(const FProperty* InProperty) const override;
+	virtual UObject* GetPropertyTypeObject(const FProperty* InProperty) const override;
+	virtual TObjectPtr<UObject> GetObjectPropertyDefaultValue(const FProperty* InProperty, const UClass* InClassToCreate) const override;
+	//~ End IPropertyIdHandler
+};
diff --git a/Engine/Plugins/VirtualProduction/RemoteControl/Source/RemoteControl/Private/PropertyIdHandler/EnumPropertyIdHandler.cpp b/Engine/Plugins/VirtualProduction/RemoteControl/Source/RemoteControl/Private/PropertyIdHandler/EnumPropertyIdHandler.cpp
new file mode 100644
--- /dev/null
+++ b/Engine/Plugins/VirtualProduction/RemoteControl/Source/RemoteControl/Private/PropertyIdHandler/EnumPropertyIdHandler.cpp
@@ -0,0 +1,49 @@
+// Copyright Epic Games, Inc. All Rights Reserved.
+
+#include "EnumPropertyIdHandler.h"
+
+bool FEnumPropertyIdHandler::IsPropertySupported(const FProperty* InProperty) const
+{
+	if (InProperty)
+	{
+		const FProperty* Property = TryGetPropertyInContainerOrSelf(InProperty);
+		return Property->GetClass()->IsChildOf(FEnumProperty::StaticClass());
+	}
+	return false;
+}
+
+EPropertyBagPropertyType FEnumPropertyIdHandler::GetPropertyType(const FProperty* InProperty) const
+{
+	return EPropertyBagPropertyType::Enum;
+}
+
+FName FEnumPropertyIdHandler::GetPropertyTypeName(const FProperty* InProperty) const
+{
+	if (InProperty)
+	{
+		const FProperty* Property = TryGetPropertyInContainerOrSelf(InProperty);
+		if (const FEnumProperty* EnumProperty = CastField<FEnumProperty>(Property))
+		{
+			return EnumProperty->GetEnum()->GetFName();
+		}
+	}
+	return FName(TEXT(""));
+}
+
+UObject* FEnumPropertyIdHandler::GetPropertyTypeObject(const FProperty* InProperty) const
+{
+	if (InProperty)
+	{
+		const FProperty* Property = TryGetPropertyInContainerOrSelf(InProperty);
+		if (const FEnumProperty* EnumProperty = CastField<FEnumProperty>(Property))
+		{
+			return EnumProperty->GetEnum();
+		}
+	}
+	return nullptr;
+}
+
+TObjectPtr<UObject> FEnumPropertyIdHandler::GetObjectPropertyDefaultValue(const FProperty* InProperty, const UClass* InClassToCreate) const
+{
+	return nullptr;
+}
diff --git a/Engine/Plugins/VirtualProduction/RemoteControl/Source/RemoteControl/Private/PropertyIdHandler/EnumPropertyIdHandler.h b/Engine/Plugins/VirtualProduction/RemoteControl/Source/RemoteControl/Private/PropertyIdHandler/EnumPropertyIdHandler.h
new file mode 100644
--- /dev/null
+++ b/Engine/Plugins/VirtualProduction/RemoteControl/Source/RemoteControl/Private/PropertyIdHandler/EnumPropertyIdHandler.h
@@ -0,0 +1,17 @@
+// Copyright Epic Games, Inc. All Rights Reserved.
+
+#pragma once
+
+#include "IPropertyIdHandler.h"
+
+class FEnumPropertyIdHandler : public IPropertyIdHandler
+{
+public:
+	//~ Begin IPropertyIdHandler
+	virtual bool IsPropertySupported(const FProperty* InProperty) const override;
+	virtual EPropertyBagPropertyType GetPropertyType(const FProperty* InProperty) const override;
+	virtual FName GetPropertyTypeName(const FProperty* InProperty) const override;
+	virtual UObject* GetPropertyTypeObject(const FProperty* InProperty) const override;
+	virtual TObjectPtr<UObject> GetObjectPropertyDefaultValue(const FProperty* InProperty, const UClass* InClassToCreate) const override;
+	//~ End IPropertyIdHandler
+};
diff --git a/Engine/Plugins/VirtualProduction/RemoteControl/Source/RemoteControl/Private/PropertyIdHandler/ObjectPropertyIdHandler.cpp b/Engine/Plugins/VirtualProduction/RemoteControl/Source/RemoteControl/Private/PropertyIdHandler/ObjectPropertyIdHandler.cpp
new file mode 100644
--- /dev/null
+++ b/Engine/Plugins/VirtualProduction/RemoteControl/Source/RemoteControl/Private/PropertyIdHandler/ObjectPropertyIdHandler.cpp
@@ -0,0 +1,54 @@
+// Copyright Epic Games, Inc. All Rights Reserved.
+
+#include "ObjectPropertyIdHandler.h"
+#include "MaterialDomain.h"
+#include "Materials/Material.h"
+#include "Materials/MaterialInterface.h"
+#include "UObject/Package.h"
+#include "UObject/UObjectGlobals.h"
+
+bool FObjectPropertyIdHandler::IsPropertySupported(const FProperty* InProperty) const
+{
+	if (InProperty)
+	{
+		const FProperty* Property = TryGetPropertyInContainerOrSelf(InProperty);
+		return Property->GetClass()->IsChildOf(FObjectProperty::StaticClass());
+	}
+	return false;
+}
+
+EPropertyBagPropertyType FObjectPropertyIdHandler::GetPropertyType(const FProperty* InProperty) const
+{
+	return EPropertyBagPropertyType::Object;
+}
+
+FName FObjectPropertyIdHandler::GetPropertyTypeName(const FProperty* InProperty) const
+{
+	if (InProperty)
+	{
+		const FProperty* Property = TryGetPropertyInContainerOrSelf(InProperty);
+		if (const FObjectProperty* ObjectProperty = CastField<FObjectProperty>(Property))
+		{
+			return ObjectProperty->PropertyClass->GetFName();
+		}
+	}
+	return FName(TEXT(""));
+}
+
+UObject* FObjectPropertyIdHandler::GetPropertyTypeObject(const FProperty* InProperty) const
+{
+	return nullptr;
+}
+
+TObjectPtr<UObject> FObjectPropertyIdHandler::GetObjectPropertyDefaultValue(const FProperty* InProperty, const UClass* InClassToCreate) const
+{
+	if (InClassToCreate)
+	{
+		if (InClassToCreate->IsChildOf(UMaterialInterface::StaticClass()))
+		{
+			return UMaterial::GetDefaultMaterial(MD_Surface);
+		}
+		return NewObject<UObject>(GetTransientPackage(), InClassToCreate);
+	}
+	return nullptr;
+}
diff --git a/Engine/Plugins/VirtualProduction/RemoteControl/Source/RemoteControl/Private/PropertyIdHandler/ObjectPropertyIdHandler.h b/Engine/Plugins/VirtualProduction/RemoteControl/Source/RemoteControl/Private/PropertyIdHandler/ObjectPropertyIdHandler.h
new file mode 100644
--- /dev/null
+++ b/Engine/Plugins/VirtualProduction/RemoteControl/Source/RemoteControl/Private/PropertyIdHandler/ObjectPropertyIdHandler.h
@@ -0,0 +1,17 @@
+// Copyright Epic Games, Inc. All Rights Reserved.
+
+#pragma once
+
+#include "IPropertyIdHandler.h"
+
+class FObjectPropertyIdHandler : public IPropertyIdHandler
+{
+public:
+	//~ Begin IPropertyIdHandler
+	virtual bool IsPropertySupported(const FProperty* InProperty) const override;
+	virtual EPropertyBagPropertyType GetPropertyType(const FProperty* InProperty) const override;
+	virtual FName GetPropertyTypeName(const FProperty* InProperty) const override;
+	virtual UObject* GetPropertyTypeObject(const FProperty* InProperty) const override;
+	virtual TObjectPtr<UObject> GetObjectPropertyDefaultValue(const FProperty* InProperty, const UClass* InClassToCreate) const override;
+	//~ End IPropertyIdHandler
+};
diff --git a/Engine/Plugins/VirtualProduction/RemoteControl/Source/RemoteControl/Private/PropertyIdHandler/StructPropertyIdHandler.cpp b/Engine/Plugins/VirtualProduction/RemoteControl/Source/RemoteControl/Private/PropertyIdHandler/StructPropertyIdHandler.cpp
new file mode 100644
--- /dev/null
+++ b/Engine/Plugins/VirtualProduction/RemoteControl/Source/RemoteControl/Private/PropertyIdHandler/StructPropertyIdHandler.cpp
@@ -0,0 +1,57 @@
+// Copyright Epic Games, Inc. All Rights Reserved.
+
+#include "StructPropertyIdHandler.h"
+
+bool FStructPropertyIdHandler::IsPropertySupported(const FProperty* InProperty) const
+{
+	if (InProperty)
+	{
+		const FProperty* Property = TryGetPropertyInContainerOrSelf(InProperty);
+		return Property->GetClass()->IsChildOf(FStructProperty::StaticClass());
+	}
+	return false;
+}
+
+EPropertyBagPropertyType FStructPropertyIdHandler::GetPropertyType(const FProperty* InProperty) const
+{
+	return EPropertyBagPropertyType::Struct;
+}
+
+FName FStructPropertyIdHandler::GetPropertyTypeName(const FProperty* InProperty) const
+{
+	if (InProperty)
+	{
+		const FProperty* Property = TryGetPropertyInContainerOrSelf(InProperty);
+		if (const FStructProperty* StructProperty = CastField<FStructProperty>(Property))
+		{
+			if (StructProperty->Struct->GetFName() == NAME_LinearColor)
+			{
+				return NAME_Color;
+			}
+			return StructProperty->Struct->GetFName();
+		}
+	}
+	return FName(TEXT(""));
+}
+
+UObject* FStructPropertyIdHandler::GetPropertyTypeObject(const FProperty* InProperty) const
+{
+	if (InProperty)
+	{
+		const FProperty* Property = TryGetPropertyInContainerOrSelf(InProperty);
+		if (const FStructProperty* StructProperty = CastField<FStructProperty>(Property))
+		{
+			if (StructProperty->Struct->GetFName() == NAME_LinearColor)
+			{
+				return TBaseStructure<FColor>::Get();
+			}
+			return StructProperty->Struct;
+		}
+	}
+	return nullptr;
+}
+
+TObjectPtr<UObject> FStructPropertyIdHandler::GetObjectPropertyDefaultValue(const FProperty* InProperty, const UClass* InClassToCreate) const
+{
+	return nullptr;
+}
diff --git a/Engine/Plugins/VirtualProduction/RemoteControl/Source/RemoteControl/Private/PropertyIdHandler/StructPropertyIdHandler.h b/Engine/Plugins/VirtualProduction/RemoteControl/Source/RemoteControl/Private/PropertyIdHandler/StructPropertyIdHandler.h
new file mode 100644
--- /dev/null
+++ b/Engine/Plugins/VirtualProduction/RemoteControl/Source/RemoteControl/Private/PropertyIdHandler/StructPropertyIdHandler.h
@@ -0,0 +1,17 @@
+// Copyright Epic Games, Inc. All Rights Reserved.
+
+#pragma once
+
+#include "IPropertyIdHandler.h"
+
+class FStructPropertyIdHandler : public IPropertyIdHandler
+{
+public:
+	//~ Begin IPropertyIdHandler
+	virtual bool IsPropertySupported(const FProperty* InProperty) const override;
+	virtual EPropertyBagPropertyType GetPropertyType(const FProperty* InProperty) const override;
+	virtual FName GetPropertyTypeName(const FProperty* InProperty) const override;
+	virtual UObject* GetPropertyTypeObject(const FProperty* InProperty) const override;
+	virtual TObjectPtr<UObject> GetObjectPropertyDefaultValue(const FProperty* InProperty, const UClass* InClassToCreate) const override;
+	//~ End IPropertyIdHandler
+};
diff --git a/Engine/Plugins/VirtualProduction/RemoteControl/Source/RemoteControl/Private/RCVirtualProperty.cpp b/Engine/Plugins/VirtualProduction/RemoteControl/Source/RemoteControl/Private/RCVirtualProperty.cpp
--- a/Engine/Plugins/VirtualProduction/RemoteControl/Source/RemoteControl/Private/RCVirtualProperty.cpp
+++ b/Engine/Plugins/VirtualProduction/RemoteControl/Source/RemoteControl/Private/RCVirtualProperty.cpp
@@ -1,4 +1,4 @@
-﻿// Copyright Epic Games, Inc. All Rights Reserved.
+// Copyright Epic Games, Inc. All Rights Reserved.
 
 #include "RCVirtualProperty.h"
 
@@ -260,6 +260,51 @@ bool URCVirtualPropertyBase::CopyCompleteValue(const FProperty* InTargetProperty
 
 }
 
+bool URCVirtualPropertyBase::CopyCompleteValue(const FProperty* InTargetProperty, uint8* InTargetValuePtr, bool bPassByteEnumPropertyComparison)
+{
+	const FProperty* SourceProperty = GetProperty();
+
+	if (SourceProperty == nullptr || InTargetProperty == nullptr || InTargetValuePtr == nullptr)
+	{
+		ensureMsgf(false, TEXT("Invalid input passed to CopyCompleteValue"));
+
+		return false;
+	}
+
+	bool bAreClassDifferent = SourceProperty->GetClass() != InTargetProperty->GetClass();
+
+	//FByteProperties are saved inside RC as FEnumProperty this will cause the original property and the RC property to differ.
+	//If it is the case than check if both has the same enum, and if they have ignore that they are different types.
+	if (bAreClassDifferent)
+	{
+		bool bShouldIgnore = false;
+		if (bPassByteEnumPropertyComparison)
+		{
+			if (const FEnumProperty* SourceEnumProperty = CastField<FEnumProperty>(SourceProperty))
+			{
+				if (const FByteProperty* ByteTargetProperty = CastField<FByteProperty>(InTargetProperty))
+				{
+					if (SourceEnumProperty->GetEnum() && ByteTargetProperty->Enum)
+					{
+						bShouldIgnore = SourceEnumProperty->GetEnum()->GetFName() == ByteTargetProperty->Enum->GetFName();
+					}
+				}
+			}
+		}
+		if (bShouldIgnore == false)
+		{
+			ensureMsgf(false, TEXT("Invalid property type passed to CopyCompleteValue.\nExpected: %s, Found: %s"), *SourceProperty->GetClass()->GetName(), *InTargetProperty->GetClass()->GetName());
+
+			return false;
+		}
+	}
+
+	SourceProperty->CopyCompleteValue(InTargetValuePtr /*Dest*/, GetValuePtr() /*Source*/);
+
+	return true;
+
+}
+
 bool URCVirtualPropertyBase::CopyCompleteValue(URCVirtualPropertyBase* InVirtualProperty)
 {
 	return CopyCompleteValue(InVirtualProperty->GetProperty(), InVirtualProperty->GetValuePtr());
diff --git a/Engine/Plugins/VirtualProduction/RemoteControl/Source/RemoteControl/Private/RemoteControlModule.cpp b/Engine/Plugins/VirtualProduction/RemoteControl/Source/RemoteControl/Private/RemoteControlModule.cpp
--- a/Engine/Plugins/VirtualProduction/RemoteControl/Source/RemoteControl/Private/RemoteControlModule.cpp
+++ b/Engine/Plugins/VirtualProduction/RemoteControl/Source/RemoteControl/Private/RemoteControlModule.cpp
@@ -11,6 +11,10 @@
 #include "Factories/RCDefaultValueFactories.h"
 #include "Factories/RemoteControlMaskingFactories.h"
 #include "Features/IModularFeatures.h"
+#include "PropertyIdHandler/BasePropertyIdHandler.h"
+#include "PropertyIdHandler/EnumPropertyIdHandler.h"
+#include "PropertyIdHandler/ObjectPropertyIdHandler.h"
+#include "PropertyIdHandler/StructPropertyIdHandler.h"
 #include "IRemoteControlInterceptionFeature.h"
 #include "IRemoteControlModule.h"
 #include "IStructDeserializerBackend.h"
@@ -640,6 +644,9 @@ void FRemoteControlModule::StartupModule()
 
 	// Register Masking Factories
 	RegisterMaskingFactories();
+
+	// Register PropertyIdHandler
+	RegisterPropertyIdHandler();
 }
 
 void FRemoteControlModule::ShutdownModule()
@@ -2164,6 +2171,23 @@ bool FRemoteControlModule::CanBeAccessedRemotely(UObject* Object) const
 	return true;
 }
 
+TSharedPtr<IPropertyIdHandler> FRemoteControlModule::GetPropertyIdHandlerFor(FProperty* InProperty)
+{
+	for (TSharedPtr<IPropertyIdHandler> Handler : PropertyIdPropertyHandlers)
+	{
+		if (Handler->IsPropertySupported(InProperty))
+		{
+			return Handler;
+		}
+	}
+	return nullptr;
+}
+
+void FRemoteControlModule::RegisterPropertyIdPropertyHandlerImpl(const TSharedRef<IPropertyIdHandler>& InPropertyIdPropertyHandler)
+{
+	PropertyIdPropertyHandlers.Add(InPropertyIdPropertyHandler);
+}
+
 void FRemoteControlModule::CachePresets() const
 {
 	TArray<FAssetData> Assets;
@@ -2477,6 +2501,14 @@ void FRemoteControlModule::RegisterMaskingFactories()
 	RegisterMaskingFactoryForType(TBaseStructure<FLinearColor>::Get(), FLinearColorMaskingFactory::MakeInstance());
 }
 
+void FRemoteControlModule::RegisterPropertyIdHandler()
+{
+	RegisterPropertyIdPropertyHandler<FBasePropertyIdHandler>();
+	RegisterPropertyIdPropertyHandler<FStructPropertyIdHandler>();
+	RegisterPropertyIdPropertyHandler<FEnumPropertyIdHandler>();
+	RegisterPropertyIdPropertyHandler<FObjectPropertyIdHandler>();
+}
+
 bool FRemoteControlModule::CanInterceptFunction(const FRCCall& RCCall) const
 {
 	if (RCCall.IsValid())
diff --git a/Engine/Plugins/VirtualProduction/RemoteControl/Source/RemoteControl/Private/RemoteControlModule.h b/Engine/Plugins/VirtualProduction/RemoteControl/Source/RemoteControl/Private/RemoteControlModule.h
--- a/Engine/Plugins/VirtualProduction/RemoteControl/Source/RemoteControl/Private/RemoteControlModule.h
+++ b/Engine/Plugins/VirtualProduction/RemoteControl/Source/RemoteControl/Private/RemoteControlModule.h
@@ -73,8 +73,12 @@ public:
 	virtual int32 EndManualEditorTransaction(const FGuid& TransactionId) override;
 	virtual const TMap<FName, TSharedPtr<IRemoteControlPropertyFactory>>& GetEntityFactories() const override { return EntityFactories; };
 	virtual bool CanBeAccessedRemotely(UObject* Object) const override;
+	virtual TSharedPtr<IPropertyIdHandler> GetPropertyIdHandlerFor(FProperty* InProperty) override;
 	//~ End IRemoteControlModule
 
+protected:
+	virtual void RegisterPropertyIdPropertyHandlerImpl(const TSharedRef<IPropertyIdHandler>& InPropertyIdPropertyHandler) override;
+
 private:
 	/** Refreshes Editor related visuals like location Gizmo for relevant properties (like Location of a SceneComponent) */
 	void RefreshEditorPostSetObjectProperties(const FRCObjectReference& ObjectAccess);
@@ -125,6 +129,11 @@ private:
 	 */
 	void RegisterMaskingFactories();
 
+	/**
+ 	 * Register(s) masking factories of supported types.
+ 	 */
+	void RegisterPropertyIdHandler();
+
 	/**
 	 * Whether function can be intercepted by a remote control interceptor.
 	 */
@@ -277,6 +286,8 @@ private:
 
 	/** Map of the factories which is responsible for resetting the Remote Control property to its default value. */
 	TMap<FName, TSharedPtr<IRCDefaultValueFactory>> DefaultValueFactories;
+
+	TSet<TSharedPtr<IPropertyIdHandler>> PropertyIdPropertyHandlers; 
 };
 
 PRAGMA_ENABLE_DEPRECATION_WARNINGS
diff --git a/Engine/Plugins/VirtualProduction/RemoteControl/Source/RemoteControl/Private/RemoteControlPreset.cpp b/Engine/Plugins/VirtualProduction/RemoteControl/Source/RemoteControl/Private/RemoteControlPreset.cpp
--- a/Engine/Plugins/VirtualProduction/RemoteControl/Source/RemoteControl/Private/RemoteControlPreset.cpp
+++ b/Engine/Plugins/VirtualProduction/RemoteControl/Source/RemoteControl/Private/RemoteControlPreset.cpp
@@ -19,12 +19,13 @@
 #include "Misc/CoreDelegates.h"
 #include "Misc/TransactionObjectEvent.h"
 #include "Misc/Optional.h"
-#include "RemoteControlExposeRegistry.h"
-#include "RemoteControlFieldPath.h"
 #include "RemoteControlActor.h"
 #include "RemoteControlBinding.h"
-#include "RemoteControlLogger.h"
 #include "RemoteControlEntityFactory.h"
+#include "RemoteControlExposeRegistry.h"
+#include "RemoteControlFieldPath.h"
+#include "RemoteControlPropertyIdRegistry.h"
+#include "RemoteControlLogger.h"
 #include "RemoteControlObjectVersion.h"
 #include "RemoteControlPresetRebindingManager.h"
 #include "RemoteControlTransactionListenerHelper.h"
@@ -517,6 +518,9 @@ URemoteControlPreset::URemoteControlPreset()
 	, RebindingManager(MakePimpl<FRemoteControlPresetRebindingManager>())
 {
 	Registry = CreateDefaultSubobject<URemoteControlExposeRegistry>(FName("ExposeRegistry"));
+
+	PropertyIdRegistry = CreateDefaultSubobject<URemoteControlPropertyIdRegistry>(FName("PropertyIdRegistry"));
+	PropertyIdRegistry->Initialize(this);
 }
 
 void URemoteControlPreset::PostInitProperties()
@@ -560,6 +564,8 @@ void URemoteControlPreset::PostLoad()
 	PostLoadProperties();
 
 	RemoveUnusedBindings();
+
+	PropertyIdRegistry->Initialize(this);
 }
 
 void URemoteControlPreset::PostDuplicate(bool bDuplicateForPIE)
@@ -957,6 +963,7 @@ TWeakPtr<FRemoteControlFunction> URemoteControlPreset::ExposeFunction(UObject* O
 TSharedPtr<FRemoteControlEntity> URemoteControlPreset::Expose(FRemoteControlEntity&& Entity, UScriptStruct* EntityType, const FGuid& GroupId)
 {
 	Registry->Modify();
+	PropertyIdRegistry->Modify();
 
 #if WITH_EDITOR
 	if (FEngineAnalytics::IsAvailable())
@@ -970,7 +977,12 @@ TSharedPtr<FRemoteControlEntity> URemoteControlPreset::Expose(FRemoteControlEnti
 	
 	TSharedPtr<FRemoteControlEntity> RCEntity = Registry->AddExposedEntity(MoveTemp(Entity), EntityType);
 	InitializeEntityMetadata(RCEntity);
-	
+
+	if (const TSharedPtr<FRemoteControlField> RCField = StaticCastSharedPtr<FRemoteControlField>(RCEntity))
+	{
+		PropertyIdRegistry->AddIdentifiedField(RCField.ToSharedRef());
+	}
+
 	RCEntity->OnEntityModifiedDelegate.BindUObject(this, &URemoteControlPreset::OnEntityModified);
 	FRemoteControlPresetGroup* Group = Layout.GetGroup(GroupId);
 	if (!Group)
@@ -991,6 +1003,24 @@ TSharedPtr<FRemoteControlEntity> URemoteControlPreset::Expose(FRemoteControlEnti
 	return RCEntity;
 }
 
+void URemoteControlPreset::PerformChainReaction(const FRemoteControlPropertyIdArgs& InArgs) const
+{
+	if (PropertyIdRegistry->IsEmpty())
+	{
+		return;
+	}
+	PropertyIdRegistry->PerformChainReaction(InArgs);
+}
+
+void URemoteControlPreset::UpdateIdentifiedField(const TSharedRef<FRemoteControlField>& InFieldToIdentify) const
+{
+	if (PropertyIdRegistry->IsEmpty())
+	{
+		return;
+	}
+	PropertyIdRegistry->UpdateIdentifiedField(InFieldToIdentify);
+}
+
 URemoteControlBinding* URemoteControlPreset::FindOrAddBinding(const TSoftObjectPtr<UObject>& Object)
 {
 	if (!ensureAlways(Object.ToString().Len()))
@@ -1544,6 +1574,10 @@ void URemoteControlPreset::Unexpose(const FGuid& EntityId)
 
 		Registry->Modify();
 		Registry->RemoveExposedEntity(EntityId);
+
+		PropertyIdRegistry->Modify();
+		PropertyIdRegistry->RemoveIdentifiedField(EntityId);
+
 		if (FRCCachedFieldData* CachedData = FieldCache.Find(EntityId))
 		{
 			Layout.RemoveField(CachedData->LayoutGroupId, EntityId);
@@ -1575,6 +1609,7 @@ void URemoteControlPreset::RebindUnboundEntities()
 	Modify();
 	RebindingManager->Rebind(this);
 	Algo::Transform(Registry->GetExposedEntities(), PerFrameUpdatedEntities, [](const TSharedPtr<FRemoteControlEntity>& Entity) { return Entity->GetId(); });
+	PropertyIdRegistry->Initialize(this);
 }
 
 void URemoteControlPreset::RebindAllEntitiesUnderSameActor(const FGuid& EntityId, AActor* NewActor, bool bUseRebindingContext)
diff --git a/Engine/Plugins/VirtualProduction/RemoteControl/Source/RemoteControl/Private/RemoteControlPropertyIdRegistry.cpp b/Engine/Plugins/VirtualProduction/RemoteControl/Source/RemoteControl/Private/RemoteControlPropertyIdRegistry.cpp
new file mode 100644
--- /dev/null
+++ b/Engine/Plugins/VirtualProduction/RemoteControl/Source/RemoteControl/Private/RemoteControlPropertyIdRegistry.cpp
@@ -0,0 +1,603 @@
+// Copyright Epic Games, Inc. All Rights Reserved.
+
+#include "RemoteControlPropertyIdRegistry.h"
+
+#include "IRemoteControlModule.h"
+#include "Materials/MaterialInterface.h"
+#include "RCVirtualProperty.h"
+#include "RemoteControlPreset.h"
+#include "UObject/Field.h"
+
+#if !WITH_EDITOR
+#include "Backends/CborStructDeserializerBackend.h"
+#include "Backends/CborStructSerializerBackend.h"
+#include "IStructDeserializerBackend.h"
+#include "StructSerializer.h"
+#endif
+
+namespace
+{
+	bool GetObjectRef(const TSharedPtr<FRemoteControlProperty>& Field, const ERCAccess Access, FRCObjectReference& OutObjectRef)
+	{
+		if (!Field.IsValid())
+		{
+			return false;
+		}
+		if (UObject* FieldBoundObject = Field->GetBoundObject())
+		{
+			FRCObjectReference ObjectRef;
+			FString ErrorText;
+			if (IRemoteControlModule::Get().ResolveObjectProperty(Access, FieldBoundObject, Field->FieldPathInfo, ObjectRef, &ErrorText))
+			{
+				OutObjectRef = ObjectRef;
+				return true;
+			}
+			UE_LOG(LogRemoteControl, Warning, TEXT("Could not resolve object property \"%s\" in object \"%s\": %s"), *Field->FieldName.ToString(), *FieldBoundObject->GetPathName(), *ErrorText);
+		}
+		return false;
+	}
+}
+
+FRCPropertyIdWrapper::FRCPropertyIdWrapper(const TSharedRef<FRemoteControlProperty>& InRCProperty)
+{
+	EntityId = InRCProperty->GetId();
+	
+	PropertyId = InRCProperty->PropertyId;
+	if (FProperty* Property = InRCProperty->GetProperty())
+	{
+		if (UObject* BoundObject = InRCProperty->GetBoundObject())
+		{
+			UpdateTypes(InRCProperty, BoundObject, Property);
+		}
+	}
+}
+
+const FGuid& FRCPropertyIdWrapper::GetEntityId() const
+{
+	return EntityId;
+}
+
+FName FRCPropertyIdWrapper::GetPropertyId() const
+{
+	return PropertyId;
+}
+
+void FRCPropertyIdWrapper::SetPropertyId(FName InNewPropertyId)
+{
+	PropertyId = InNewPropertyId;
+}
+
+const FName& FRCPropertyIdWrapper::GetSuperType() const
+{
+	return SuperType;
+}
+
+const FName& FRCPropertyIdWrapper::GetSubType() const
+{
+	return SubType;
+}
+
+const UClass* FRCPropertyIdWrapper::GetClassToCreate() const
+{
+	return ClassToCreate;
+}
+
+bool FRCPropertyIdWrapper::IsValid() const
+{
+	return EntityId.IsValid() && (!SuperType.IsNone() || !SubType.IsNone());
+}
+
+bool FRCPropertyIdWrapper::IsValidPropertyId() const
+{
+	return PropertyId != NAME_None;
+}
+
+uint32 GetTypeHash(const FRCPropertyIdWrapper& Wrapper)
+{
+	return GetTypeHash(Wrapper.EntityId);
+}
+
+bool FRCPropertyIdWrapper::operator==(const FGuid& WrappedId) const
+{
+	if (!IsValid())
+	{
+		return false;
+	}
+	return EntityId == WrappedId;
+}
+
+bool FRCPropertyIdWrapper::operator==(const FRCPropertyIdWrapper& Other) const
+{
+	if (IsValid() && Other.IsValid())
+	{
+		return EntityId == Other.EntityId;
+	}
+	return false;
+}
+
+void FRCPropertyIdWrapper::UpdateTypes(const TSharedRef<FRemoteControlProperty>& InRCProperty, UObject* InOwner, FProperty* InProperty)
+{
+	SuperType = InProperty->GetClass()->GetFName();
+	SubType = NAME_None;
+	ClassToCreate = nullptr;
+	if (const FEnumProperty* EnumProperty = CastField<FEnumProperty>(InProperty))
+	{
+		if (EnumProperty->GetEnum())
+		{
+			SubType = EnumProperty->GetEnum()->GetFName();
+		}
+	}
+	else if (const FByteProperty* ByteProperty = CastField<FByteProperty>(InProperty))
+	{
+		if (ByteProperty->Enum)
+		{
+			SubType = ByteProperty->Enum->GetFName();
+		}
+	}
+	else if (const FStructProperty* StructProperty = CastField<FStructProperty>(InProperty))
+	{
+		if (StructProperty->Struct->GetFName() == NAME_LinearColor)
+		{
+			SubType = NAME_Color;
+		}
+		else
+		{
+			SubType = StructProperty->Struct->GetFName();
+		}
+	}
+	else if (const FObjectProperty* ObjectProperty = CastField<FObjectProperty>(InProperty))
+	{
+		// ContainerPtrToValuePtr will crash if the GetOwner<UClass> is nullptr so we check before.
+		if (ObjectProperty->GetOwner<UClass>())
+		{
+			if (const uint8* PropertyValuePtr = ObjectProperty->ContainerPtrToValuePtr<uint8>(InOwner))
+			{
+				if (const TObjectPtr<UObject> Object = ObjectProperty->GetPropertyValue(PropertyValuePtr))
+				{
+					if (Object && Object->IsValidLowLevel())
+					{
+						if (Object->IsA(UMaterialInterface::StaticClass()))
+						{
+							SubType = UMaterialInterface::StaticClass()->GetFName();
+							ClassToCreate = UMaterialInterface::StaticClass();
+						}
+						else
+						{
+							SubType = Object->GetClass()->GetFName();
+							ClassToCreate = Object->GetClass();
+						}
+					}
+				}
+			}
+		}
+	}
+	else if (const FArrayProperty* ArrayProperty = CastField<FArrayProperty>(InProperty))
+	{
+		SuperType = ArrayProperty->Inner->GetClass()->GetFName();
+		if (const uint8* PropertyValuePtr = ArrayProperty->ContainerPtrToValuePtr<uint8>(InOwner))
+		{
+			if (const FObjectProperty* InnerObjectProperty = CastField<FObjectProperty>(ArrayProperty->Inner))
+			{
+				FScriptArrayHelper ArrayHelper(ArrayProperty, PropertyValuePtr);
+				
+				if (const TSharedRef<FRemoteControlField> TargetRCField = StaticCastSharedRef<FRemoteControlField>(InRCProperty); !TargetRCField->FieldPathInfo.Segments.IsEmpty())
+				{
+					const int32 MaterialIndex = TargetRCField->FieldPathInfo.Segments[0].ArrayIndex;
+					const uint8* ObjPtrContainer = ArrayHelper.GetRawPtr(MaterialIndex);
+					if (const UObject* CurrentObject = InnerObjectProperty->GetObjectPropertyValue(ObjPtrContainer))
+					{
+						if (CurrentObject->IsA(UMaterialInterface::StaticClass()))
+						{
+							SubType = UMaterialInterface::StaticClass()->GetFName();
+							ClassToCreate = UMaterialInterface::StaticClass();
+						}
+						else
+						{
+							SubType = CurrentObject->GetClass()->GetFName();
+							ClassToCreate = CurrentObject->GetClass();
+						}
+					}
+				}
+			}
+		}
+	}
+}
+
+void URemoteControlPropertyIdRegistry::Initialize(URemoteControlPreset* InSourcePreset)
+{
+	SourcePreset = InSourcePreset;
+	IdentifiedFields.Reset();
+	for (TWeakPtr<FRemoteControlEntity> RCEntity : SourcePreset->GetExposedEntities())
+	{
+		if (const TSharedPtr<FRemoteControlField> RCField = StaticCastSharedPtr<FRemoteControlField>(RCEntity.Pin()))
+		{
+			AddIdentifiedField(RCField.ToSharedRef());
+		}
+	}
+}
+
+void URemoteControlPropertyIdRegistry::PerformChainReaction(const FRemoteControlPropertyIdArgs& InArgs)
+{
+	TSet<FGuid> TargetProperties;
+	Algo::TransformIf(IdentifiedFields, TargetProperties,
+		[&](const FRCPropertyIdWrapper& Wrapper)
+		{
+			return Wrapper.IsValid()
+			&& Wrapper.GetPropertyId() == InArgs.PropertyId
+			&& Wrapper.GetSuperType() == InArgs.SuperType
+			&& Wrapper.GetSubType() == InArgs.SubType;
+		},
+		[&](const FRCPropertyIdWrapper& Wrapper)
+		{
+			return Wrapper.GetEntityId();
+		}
+		);
+	if (!TargetProperties.IsEmpty() && IsValid(SourcePreset))
+	{
+		for (const FGuid& TargetProperty : TargetProperties)
+		{
+			if (TSharedPtr<FRemoteControlProperty> TargetRCProperty = SourcePreset->GetExposedEntity<FRemoteControlProperty>(TargetProperty).Pin())
+			{
+				if (UObject* BoundObject = TargetRCProperty->GetBoundObject())
+				{
+					bool bCopyComplete = false;
+
+					if (FProperty* Property = TargetRCProperty->GetProperty())
+					{
+#if WITH_EDITOR
+						BoundObject->PreEditChange(Property);
+						
+						BoundObject->Modify();
+#endif // WITH_EDITOR
+						// Note : For all other object types except materials.
+						if (FObjectProperty* ObjectProperty = CastField<FObjectProperty>(Property))
+						{
+							if (uint8* PropertyValuePtr = ObjectProperty->ContainerPtrToValuePtr<uint8>(BoundObject))
+							{
+								UObject* CurrentObject = ObjectProperty->GetObjectPropertyValue(PropertyValuePtr);
+								if (CurrentObject && CurrentObject->IsA(InArgs.SourceClass))
+								{
+									ObjectProperty->SetObjectPropertyValue(PropertyValuePtr, InArgs.SourceObject);
+									bCopyComplete = true;
+								}
+							}
+						}
+						// Note : Specialization for Materials.
+						else if (FArrayProperty* ArrayProperty = CastField<FArrayProperty>(Property))
+						{
+							if (uint8* PropertyValuePtr = ArrayProperty->ContainerPtrToValuePtr<uint8>(BoundObject))
+							{
+								if (FObjectProperty* InnerObjectProperty = CastField<FObjectProperty>(ArrayProperty->Inner))
+								{
+									FScriptArrayHelper ArrayHelper(ArrayProperty, PropertyValuePtr);
+									if (TSharedPtr<FRemoteControlField> TargetRCField = StaticCastSharedPtr<FRemoteControlField>(TargetRCProperty); !TargetRCField->FieldPathInfo.Segments.IsEmpty())
+									{
+										const int32 MaterialIndex = TargetRCField->FieldPathInfo.Segments[0].ArrayIndex;
+										uint8* ObjPtrContainer = ArrayHelper.GetRawPtr(MaterialIndex);
+										UObject* CurrentObject = InnerObjectProperty->GetObjectPropertyValue(ObjPtrContainer);
+										if (CurrentObject && CurrentObject->IsA(InArgs.SourceClass))
+										{
+											InnerObjectProperty->SetObjectPropertyValue(ObjPtrContainer, InArgs.SourceObject);
+											bCopyComplete = true;
+										}
+									}
+								}
+							}
+						}
+						// Note : For primitive types not in container and that has a valid UClass Owner.
+						// ContainerPtrToValuePtr will crash if the GetOwner<UClass> is nullptr so we check before.
+						if (!bCopyComplete && Property->GetOwner<UClass>() != nullptr &&
+							!Property->IsA<FArrayProperty>() && !Property->IsA<FSetProperty>() && !Property->IsA<FMapProperty>())
+						{
+							if (uint8* PropertyValuePtr = Property->ContainerPtrToValuePtr<uint8>(BoundObject))
+							{
+								if (const FStructProperty* StructProperty = CastField<FStructProperty>(Property))
+								{
+									if (StructProperty->Struct->GetFName() == NAME_LinearColor)
+									{
+										FColor ColorValue;
+										InArgs.VirtualProperty->GetValueColor(ColorValue);
+										const FLinearColor RealValue(ColorValue);
+										Property->CopyCompleteValue(PropertyValuePtr, &RealValue);
+										return;
+									}
+								}
+								bCopyComplete = InArgs.VirtualProperty->CopyCompleteValue(Property, PropertyValuePtr, true);
+							}
+						}
+						//Note : For all the other cases, Container and UStruct Owner.
+						else if (!bCopyComplete)
+						{
+#if !WITH_EDITOR
+							StructMemoryContainer = nullptr;
+#endif // !WITH_EDITOR
+							bCopyComplete = CopyNonUClassOwnerProperty(InArgs, TargetRCProperty, Property, BoundObject);
+						}
+						if (bCopyComplete)
+						{
+#if WITH_EDITOR
+							TargetRCProperty->FieldPathInfo.Resolve(BoundObject);
+							FPropertyChangedEvent PropertyEvent = TargetRCProperty->FieldPathInfo.ToPropertyChangedEvent(EPropertyChangeType::ValueSet);
+							FEditPropertyChain EditPropertyChain;
+							TargetRCProperty->FieldPathInfo.ToEditPropertyChain(EditPropertyChain);
+							if (EditPropertyChain.IsEmpty())
+							{
+								BoundObject->PostEditChangeProperty(PropertyEvent);
+							}
+							else
+							{
+								FPropertyChangedChainEvent ChainEvent(EditPropertyChain, PropertyEvent);
+								TArray<TMap<FString, int32>> ArrayIndicesPerObject;
+								{
+									TMap<FString, int32> ArrayIndices;
+									ArrayIndices.Reserve(TargetRCProperty->FieldPathInfo.Segments.Num());
+									for (const FRCFieldPathSegment& Segment : TargetRCProperty->FieldPathInfo.Segments)
+									{
+										ArrayIndices.Add(Segment.Name.ToString(), Segment.ArrayIndex);
+									}
+									ArrayIndicesPerObject.Add(MoveTemp(ArrayIndices));
+								}
+								ChainEvent.ObjectIteratorIndex = 0;
+								ChainEvent.SetArrayIndexPerObject(ArrayIndicesPerObject);
+								BoundObject->PostEditChangeChainProperty(ChainEvent);
+							}
+							
+#else // NOTE: During runtime use Serialization API to update the property properly.
+							TArray<uint8> Buffer;
+							FMemoryWriter Writer(Buffer);
+							FCborStructSerializerBackend WriterBackend(Writer, EStructSerializerBackendFlags::Default);
+							FStructSerializerPolicies Policies;
+							Policies.MapSerialization = EStructSerializerMapPolicies::Array;
+
+							// We do this at runtime because otherwise Property Owned by UStruct won't update.
+							// So during the copy we take the StructContainer and use that to serialize the property.
+							if (StructMemoryContainer != nullptr)
+							{
+								FStructSerializer::SerializeElement(StructMemoryContainer, Property, INDEX_NONE, WriterBackend, Policies);
+							}
+							else
+							{
+								FStructSerializer::SerializeElement(BoundObject, Property, INDEX_NONE, WriterBackend, Policies);
+							}
+							// Deserialization
+							FMemoryReader Reader(Buffer);
+							FCborStructDeserializerBackend ReaderBackend(Reader);
+							FRCObjectReference TargetObjectRef;
+							if (!GetObjectRef(TargetRCProperty, ERCAccess::WRITE_ACCESS, TargetObjectRef))
+							{
+								continue;
+							}
+							if (!IRemoteControlModule::Get().SetObjectProperties(TargetObjectRef, ReaderBackend, ERCPayloadType::Cbor, Buffer))
+							{
+								continue;
+							}
+#endif // WITH_EDITOR
+						}
+					}
+				}
+			}
+		}
+	}
+}
+
+void URemoteControlPropertyIdRegistry::AddIdentifiedField(const TSharedRef<FRemoteControlField>& InFieldToIdentify)
+{
+	if (InFieldToIdentify->FieldType == EExposedFieldType::Property)
+	{
+		if (const TSharedPtr<FRemoteControlProperty> RCProperty = StaticCastSharedRef<FRemoteControlProperty>(InFieldToIdentify))
+		{
+			FRCPropertyIdWrapper Wrapper{ RCProperty.ToSharedRef()};
+			if (Wrapper.IsValid())
+			{
+				IdentifiedFields.Add(MoveTemp(Wrapper));
+			}
+		}
+	}
+}
+
+bool URemoteControlPropertyIdRegistry::IsEmpty() const
+{
+	return IdentifiedFields.IsEmpty();
+}
+
+void URemoteControlPropertyIdRegistry::UpdateIdentifiedField(const TSharedRef<FRemoteControlField>& InFieldToIdentify)
+{
+	if (InFieldToIdentify->FieldType == EExposedFieldType::Property)
+	{
+		const uint32 Hash = GetTypeHash(InFieldToIdentify->GetId());
+		if (FRCPropertyIdWrapper* Wrapper = IdentifiedFields.FindByHash(Hash, InFieldToIdentify->GetId()))
+		{
+			Wrapper->SetPropertyId(InFieldToIdentify->PropertyId);
+			if (IsValid(SourcePreset))
+			{
+				SourcePreset.Get()->GetPropertyIdRegistry()->OnPropertyIdUpdated().Broadcast();
+			}
+		}
+	}
+}
+
+void URemoteControlPropertyIdRegistry::RemoveIdentifiedField(const FGuid& InEntityId)
+{
+	const uint32 Hash = GetTypeHash(InEntityId);
+
+	if (IdentifiedFields.ContainsByHash(Hash, InEntityId))
+	{
+		IdentifiedFields.RemoveByHash(Hash, InEntityId);
+	}
+}
+
+const UClass* URemoteControlPropertyIdRegistry::GetClassByEntityId(const FGuid& InEntityId) const
+{
+	const uint32 Hash = GetTypeHash(InEntityId);
+
+	if (const FRCPropertyIdWrapper* Wrapper = IdentifiedFields.FindByHash(Hash, InEntityId))
+	{
+		return Wrapper->GetClassToCreate();
+	}
+	return nullptr;
+}
+
+TSet<FName> URemoteControlPropertyIdRegistry::GetFieldIdsNameList()
+{
+	TSet<FName> OutIds;
+	for (FRCPropertyIdWrapper Field : IdentifiedFields)
+	{
+		if (Field.IsValidPropertyId())
+		{
+			OutIds.Add(Field.GetPropertyId());
+		}
+	}
+	return OutIds;
+}
+
+TSet<FGuid> URemoteControlPropertyIdRegistry::GetEntityIdsForPropertyId(const FName& InPropertyId) const
+{
+	TSet<FGuid> OutEntityIds;
+	if (InPropertyId == NAME_None)
+	{
+		return OutEntityIds;
+	}
+	for (FRCPropertyIdWrapper Wrappers : IdentifiedFields)
+	{
+		if (Wrappers.GetPropertyId() == InPropertyId)
+		{
+			OutEntityIds.Add(Wrappers.GetEntityId());
+		}
+	}
+	return OutEntityIds;
+}
+
+TSet<FGuid> URemoteControlPropertyIdRegistry::GetEntityIdsList()
+{
+	TSet<FGuid> OutIds;
+	for (FRCPropertyIdWrapper Field : IdentifiedFields)
+	{
+		if (Field.IsValidPropertyId())
+		{
+			OutIds.Add(Field.GetEntityId());
+		}
+	}
+	return OutIds;
+}
+
+bool URemoteControlPropertyIdRegistry::CopyNonUClassOwnerProperty(const FRemoteControlPropertyIdArgs& InArgs, TSharedPtr<FRemoteControlProperty> TargetRCProperty, FProperty* InProperty, UObject* BoundObject)
+{
+	bool bCheckByteEnumComparison = false;
+	if (InArgs.VirtualProperty->GetProperty())
+	{
+		bCheckByteEnumComparison = InArgs.VirtualProperty->GetProperty()->IsA<FEnumProperty>();
+	}
+	TargetRCProperty->FieldPathInfo.Resolve(BoundObject);
+	return TryCopyNonUClassOwnerProperty(InArgs.VirtualProperty, TargetRCProperty, InProperty, BoundObject, TargetRCProperty->FieldPathInfo, 0, nullptr, bCheckByteEnumComparison);
+}
+
+bool URemoteControlPropertyIdRegistry::TryCopyNonUClassOwnerProperty(const TObjectPtr<URCVirtualPropertySelfContainer>& InVirtualPropertySelfContainer, const TSharedPtr<FRemoteControlProperty>& TargetRCProperty, FProperty* InProperty, UObject* BoundObject, FRCFieldPathInfo Segment, const int32 IndexSegment, uint8* PropertyContainer, const bool bCheckByteEnumComparison)
+{
+	if (IndexSegment >= Segment.GetSegmentCount())
+	{
+		return false;
+	}
+
+	if (const FArrayProperty* ArrayProperty = CastField<FArrayProperty>(Segment.Segments[IndexSegment].ResolvedData.Field))
+	{
+		if (const uint8* PropertyValuePtr = ArrayProperty->ContainerPtrToValuePtr<uint8>(BoundObject))
+		{
+			FScriptArrayHelper ArrayHelper(ArrayProperty, PropertyValuePtr);
+			if (const TSharedPtr<FRemoteControlField> TargetRCField = StaticCastSharedPtr<FRemoteControlField>(TargetRCProperty);
+				!TargetRCField->FieldPathInfo.Segments.IsEmpty())
+			{
+				const int32 Index = Segment.Segments[IndexSegment].ArrayIndex;
+				uint8* PtrContainer = ArrayHelper.GetElementPtr(Index);
+				if (IndexSegment < Segment.GetSegmentCount() - 2)
+				{
+					return TryCopyNonUClassOwnerProperty(InVirtualPropertySelfContainer, TargetRCProperty, InProperty, BoundObject, Segment, IndexSegment + 1, PtrContainer, bCheckByteEnumComparison);
+				}
+				if (const FArrayProperty* ArrayRealProperty = CastField<FArrayProperty>(InProperty))
+				{
+					return InVirtualPropertySelfContainer->CopyCompleteValue(ArrayRealProperty->Inner, PtrContainer, bCheckByteEnumComparison);
+				}
+				return false;
+			}
+		}
+	}
+	else if (const FSetProperty* SetProperty = CastField<FSetProperty>(Segment.Segments[IndexSegment].ResolvedData.Field))
+	{
+		if (const uint8* PropertyValuePtr = SetProperty->ContainerPtrToValuePtr<uint8>(BoundObject))
+		{
+			FScriptSetHelper SetHelper(SetProperty, PropertyValuePtr);
+			if (const TSharedPtr<FRemoteControlField> TargetRCField = StaticCastSharedPtr<FRemoteControlField>(TargetRCProperty);
+				!TargetRCField->FieldPathInfo.Segments.IsEmpty())
+			{
+				const int32 Index = Segment.Segments[IndexSegment].ArrayIndex;
+				uint8* PtrContainer = SetHelper.GetElementPtr(Index);
+				if (IndexSegment < Segment.GetSegmentCount() - 2)
+				{
+					return TryCopyNonUClassOwnerProperty(InVirtualPropertySelfContainer, TargetRCProperty, InProperty, BoundObject, Segment, IndexSegment + 1, PtrContainer, bCheckByteEnumComparison);
+				}
+				if (const FSetProperty* SetRealProperty = CastField<FSetProperty>(InProperty))
+				{
+					return InVirtualPropertySelfContainer->CopyCompleteValue(SetRealProperty->ElementProp, PtrContainer, bCheckByteEnumComparison);
+				}
+				return false;
+			}
+		}
+	}
+	else if (const FMapProperty* MapProperty = CastField<FMapProperty>(Segment.Segments[IndexSegment].ResolvedData.Field))
+	{
+		if (const uint8* PropertyValuePtr = MapProperty->ContainerPtrToValuePtr<uint8>(BoundObject))
+		{
+			FScriptMapHelper MapHelper(MapProperty, PropertyValuePtr);
+			if (const TSharedPtr<FRemoteControlField> TargetRCField = StaticCastSharedPtr<FRemoteControlField>(TargetRCProperty);
+				!TargetRCField->FieldPathInfo.Segments.IsEmpty())
+			{
+				const int32 Index = Segment.Segments[IndexSegment].ArrayIndex;
+				uint8* PtrContainer = MapHelper.GetValuePtr(Index);
+				if (IndexSegment < Segment.GetSegmentCount() - 2)
+				{
+					return TryCopyNonUClassOwnerProperty(InVirtualPropertySelfContainer, TargetRCProperty, InProperty, BoundObject, Segment, IndexSegment + 1, PtrContainer, bCheckByteEnumComparison);
+				}
+				if (const FMapProperty* MapRealProperty = CastField<FMapProperty>(InProperty))
+				{
+					return InVirtualPropertySelfContainer->CopyCompleteValue(MapRealProperty->ValueProp, PtrContainer, bCheckByteEnumComparison);
+				}
+				return false;
+			}
+		}
+	}
+	else if (const FStructProperty* StructProperty = CastField<FStructProperty>(Segment.Segments[IndexSegment].ResolvedData.Field))
+	{
+		uint8* InnerStructPtrContainer;
+		if (PropertyContainer != nullptr)
+		{
+			InnerStructPtrContainer = StructProperty->ContainerPtrToValuePtr<uint8>(PropertyContainer);
+		}
+		else
+		{
+			InnerStructPtrContainer = StructProperty->ContainerPtrToValuePtr<uint8>(BoundObject);
+		}
+		if (uint8* PtrContainer = InProperty->ContainerPtrToValuePtr<uint8>(InnerStructPtrContainer))
+		{
+			if (IndexSegment < Segment.GetSegmentCount() - 2)
+			{
+				return TryCopyNonUClassOwnerProperty(InVirtualPropertySelfContainer, TargetRCProperty, InProperty, BoundObject, Segment, IndexSegment + 1, InnerStructPtrContainer, bCheckByteEnumComparison);
+			}
+			if (const FStructProperty* StructRealProperty = CastField<FStructProperty>(InProperty))
+			{
+#if !WITH_EDITOR
+				StructMemoryContainer = InnerStructPtrContainer;
+#endif // !WITH_EDITOR
+				if (StructRealProperty->Struct->GetFName() == NAME_LinearColor)
+				{
+					FColor ColorValue;
+					InVirtualPropertySelfContainer->GetValueColor(ColorValue);
+					const FLinearColor RealValue(ColorValue);
+					InProperty->CopyCompleteValue(PtrContainer, &RealValue);
+					return true;
+				}
+			}
+			return InVirtualPropertySelfContainer->CopyCompleteValue(InProperty, PtrContainer, bCheckByteEnumComparison);
+		}
+	}
+
+	return false;
+}
diff --git a/Engine/Plugins/VirtualProduction/RemoteControl/Source/RemoteControl/Public/IPropertyIdHandler.h b/Engine/Plugins/VirtualProduction/RemoteControl/Source/RemoteControl/Public/IPropertyIdHandler.h
new file mode 100644
--- /dev/null
+++ b/Engine/Plugins/VirtualProduction/RemoteControl/Source/RemoteControl/Public/IPropertyIdHandler.h
@@ -0,0 +1,52 @@
+// Copyright Epic Games, Inc. All Rights Reserved.
+
+#pragma once
+
+#include "UObject/Field.h"
+#include "UObject/NameTypes.h"
+#include "UObject/ObjectPtr.h"
+#include "UObject/UnrealType.h"
+#include "PropertyBag.h"
+
+class FProperty;
+class UObject;
+
+class IPropertyIdHandler
+{
+public:
+	virtual ~IPropertyIdHandler() = default;
+
+	/** Whether this PropertyId Handler handles the given Property */
+	virtual bool IsPropertySupported(const FProperty* InProperty) const = 0;
+
+	/** Get the PropertyBag type of the given property */
+	virtual EPropertyBagPropertyType GetPropertyType(const FProperty* InProperty) const = 0;
+
+	/** Get the Property type FName */
+	virtual FName GetPropertyTypeName(const FProperty* InProperty) const = 0;
+
+	/** Get the Property type Object */
+	virtual UObject* GetPropertyTypeObject(const FProperty* InProperty) const = 0;
+
+	/** Get the Property default value (Used for UObject to create a new one for the ValueWidget) */
+	virtual TObjectPtr<UObject> GetObjectPropertyDefaultValue(const FProperty* InProperty, const UClass* InClassToCreate) const = 0;
+
+protected:
+	/** Get the property inside the container if not inside any container return the property passed */
+	const FProperty* TryGetPropertyInContainerOrSelf(const FProperty* InProperty) const
+	{
+		if (const FArrayProperty* ArrayProperty = CastField<FArrayProperty>(InProperty))
+		{
+			return ArrayProperty->Inner;
+		}
+		if (const FMapProperty* MapProperty = CastField<FMapProperty>(InProperty))
+		{
+			return MapProperty->ValueProp;
+		}
+		if (const FSetProperty* SetProperty = CastField<FSetProperty>(InProperty))
+        {
+        	return SetProperty->ElementProp;
+        }
+		return InProperty;
+	}
+};
diff --git a/Engine/Plugins/VirtualProduction/RemoteControl/Source/RemoteControl/Public/IRemoteControlModule.h b/Engine/Plugins/VirtualProduction/RemoteControl/Source/RemoteControl/Public/IRemoteControlModule.h
--- a/Engine/Plugins/VirtualProduction/RemoteControl/Source/RemoteControl/Public/IRemoteControlModule.h
+++ b/Engine/Plugins/VirtualProduction/RemoteControl/Source/RemoteControl/Public/IRemoteControlModule.h
@@ -14,6 +14,7 @@
 
 REMOTECONTROL_API DECLARE_LOG_CATEGORY_EXTERN(LogRemoteControl, Log, All);
 
+class IPropertyIdHandler;
 class IRemoteControlMaskingFactory;
 class IStructDeserializerBackend;
 class IStructSerializerBackend;
@@ -600,5 +601,18 @@ public:
 	 *  @note Check Remote Control project settings to configure.
 	 */
 	virtual bool CanBeAccessedRemotely(UObject* InObject) const = 0;
-};
 
+	template<class InPropertyIdPropertyHandlerType, typename... InArgTypes
+		, TEMPLATE_REQUIRES(TIsDerivedFrom<InPropertyIdPropertyHandlerType, IPropertyIdHandler>::Value)>
+	TSharedRef<InPropertyIdPropertyHandlerType> RegisterPropertyIdPropertyHandler(InArgTypes&&... InArgs)
+	{
+		TSharedRef<InPropertyIdPropertyHandlerType> KeyPropertyHandler = MakeShared<InPropertyIdPropertyHandlerType>(Forward<InArgTypes>(InArgs)...);
+		this->RegisterPropertyIdPropertyHandlerImpl(KeyPropertyHandler);
+		return KeyPropertyHandler;
+	}
+
+	virtual TSharedPtr<IPropertyIdHandler> GetPropertyIdHandlerFor(FProperty* InProperty) = 0;
+
+protected:
+	virtual void RegisterPropertyIdPropertyHandlerImpl(const TSharedRef<IPropertyIdHandler>& InKeyPropertyHandler) = 0;
+};
diff --git a/Engine/Plugins/VirtualProduction/RemoteControl/Source/RemoteControl/Public/RCVirtualProperty.h b/Engine/Plugins/VirtualProduction/RemoteControl/Source/RemoteControl/Public/RCVirtualProperty.h
--- a/Engine/Plugins/VirtualProduction/RemoteControl/Source/RemoteControl/Public/RCVirtualProperty.h
+++ b/Engine/Plugins/VirtualProduction/RemoteControl/Source/RemoteControl/Public/RCVirtualProperty.h
@@ -1,4 +1,4 @@
-﻿// Copyright Epic Games, Inc. All Rights Reserved.
+// Copyright Epic Games, Inc. All Rights Reserved.
  
 #pragma once
 
@@ -125,6 +125,16 @@ public:
 	*/
 	bool CopyCompleteValue(const FProperty* InTargetProperty, uint8* InTargetValuePtr);
 
+	/** Copy this virtual property's data onto a given FProperty
+	* 
+	* @param InTargetProperty - The property onto which our value is to be copied
+	* @param InTargetValuePtr - The memory location for the target property
+	* @param bPassByteEnumPropertyComparison - Go to the copy if the Property and the Target property are Enum and Byte
+	* 
+	* @return true if copied successfully
+	*/
+	bool CopyCompleteValue(const FProperty* InTargetProperty, uint8* InTargetValuePtr, bool bPassByteEnumPropertyComparison);
+
 	/** Copy this virtual property's data onto a given Virtual Property
 	*
 	* @param InVirtualProperty - The virtual property onto which our value is to be copied
diff --git a/Engine/Plugins/VirtualProduction/RemoteControl/Source/RemoteControl/Public/RemoteControlField.h b/Engine/Plugins/VirtualProduction/RemoteControl/Source/RemoteControl/Public/RemoteControlField.h
--- a/Engine/Plugins/VirtualProduction/RemoteControl/Source/RemoteControl/Public/RemoteControlField.h
+++ b/Engine/Plugins/VirtualProduction/RemoteControl/Source/RemoteControl/Public/RemoteControlField.h
@@ -73,6 +73,12 @@ public:
 	UPROPERTY(VisibleAnywhere, BlueprintReadOnly, Category = "RemoteControlEntity")
 	FName FieldName;
 
+	/**
+	 * The exposed field's identifier.
+	 */
+	UPROPERTY(VisibleAnywhere, BlueprintReadOnly, Category = "RemoteControlEntity")
+	FName PropertyId;
+
 	/**
 	 * Path information pointing to this field
 	 */
diff --git a/Engine/Plugins/VirtualProduction/RemoteControl/Source/RemoteControl/Public/RemoteControlPreset.h b/Engine/Plugins/VirtualProduction/RemoteControl/Source/RemoteControl/Public/RemoteControlPreset.h
--- a/Engine/Plugins/VirtualProduction/RemoteControl/Source/RemoteControl/Public/RemoteControlPreset.h
+++ b/Engine/Plugins/VirtualProduction/RemoteControl/Source/RemoteControl/Public/RemoteControlPreset.h
@@ -1,14 +1,14 @@
 // Copyright Epic Games, Inc. All Rights Reserved.
 #pragma once
 
-#include "CoreTypes.h"
-
 #include "Algo/Transform.h"
+#include "CoreTypes.h"
 #include "RemoteControlField.h"
 #include "RemoteControlEntity.h"
-#include "UObject/SoftObjectPtr.h"
+#include "RemoteControlPropertyIdRegistry.h"
 #include "Templates/PimplPtr.h"
 #include "Templates/UnrealTypeTraits.h"
+#include "UObject/SoftObjectPtr.h"
 
 #include "RemoteControlPreset.generated.h"
 
@@ -27,6 +27,7 @@ class UBlueprint;
 class URCVirtualPropertyBase;
 class URCVirtualPropertyContainerBase;
 class URCVirtualPropertyInContainer;
+class URCVirtualPropertySelfContainer;
 class URemoteControlExposeRegistry;
 class URemoteControlBinding;
 class URemoteControlPreset;
@@ -47,6 +48,37 @@ struct REMOTECONTROL_API FRemoteControlPresetExposeArgs
 	bool bEnableEditCondition;
 };
 
+/** Arguments used to expose an entity (Actor, property, function, etc.) */
+struct REMOTECONTROL_API FRemoteControlPropertyIdArgs
+{
+	FRemoteControlPropertyIdArgs() = default;
+
+	bool IsValid() const
+	{
+		return VirtualProperty && (SourceObject || SourceClass)
+			&& (!SuperType.IsNone() || !SubType.IsNone());
+	}
+
+public:
+	/** PropertyId */
+	FName PropertyId;
+
+	/** SuperType of the Property */
+	FName SuperType;
+
+	/** SubType of the Property */
+	FName SubType;
+
+	/** (Optional) The source object to use for special cases. */
+	TObjectPtr<UObject> SourceObject;
+
+	/** (Optional) The class of the source object to use for special cases. */
+	UClass* SourceClass;
+
+	/** (Optional) The virtual property to use for all the cases. */
+	TObjectPtr<URCVirtualPropertySelfContainer> VirtualProperty;
+};
+
 /**
  * Data cached for every exposed field.
  */
@@ -761,9 +793,27 @@ public:
 	/** Expose an entity in the registry. */
 	TSharedPtr<FRemoteControlEntity> Expose(FRemoteControlEntity&& Entity, UScriptStruct* EntityType, const FGuid& GroupId);
 
+	/**
+	 * Calls to the PropertyIdRegistry PerformChainReaction to update the value(s) of the property(ies) that are bound to a PropertyIdAction.
+	 * @param InArgs Argument used to update the property(ies) value.
+	 */
+	void PerformChainReaction(const FRemoteControlPropertyIdArgs& InArgs) const;
+
+	/**
+	 * Calls to the PropertyIdRegistry UpdateIdentifiedField to updates a field from the set of identified fields.
+	 * @param InFieldToIdentify the entity to identify.
+	 */
+	void UpdateIdentifiedField(const TSharedRef<FRemoteControlField>& InFieldToIdentify) const;
 
 	/** Try to get a binding and creates a new one if it doesn't exist. */
 	URemoteControlBinding* FindOrAddBinding(const TSoftObjectPtr<UObject>& Object);
+
+	/**
+	 * Get the PropertyIdRegistry.
+	 * @return The PropertyIdRegistry of this Preset.
+	 */
+	TObjectPtr<URemoteControlPropertyIdRegistry> GetPropertyIdRegistry() const { return PropertyIdRegistry; }
+
 private:
 
 	/** Find a binding that has the same boundobjectmap but that currently points to the object passed as argument. */
@@ -822,6 +872,10 @@ private:
 	/** Holds exposed entities on the preset. */
 	TObjectPtr<URemoteControlExposeRegistry> Registry = nullptr;
 
+	/** Holds identities of exposed entities on the preset. */
+	UPROPERTY(Transient)
+	TObjectPtr<URemoteControlPropertyIdRegistry> PropertyIdRegistry = nullptr;
+
 	/** Delegate triggered when an entity is exposed. */
 	FOnPresetEntityEvent OnEntityExposedDelegate;
 	/** Delegate triggered when an entity is unexposed from the preset. */
diff --git a/Engine/Plugins/VirtualProduction/RemoteControl/Source/RemoteControl/Public/RemoteControlPropertyIdRegistry.h b/Engine/Plugins/VirtualProduction/RemoteControl/Source/RemoteControl/Public/RemoteControlPropertyIdRegistry.h
new file mode 100644
--- /dev/null
+++ b/Engine/Plugins/VirtualProduction/RemoteControl/Source/RemoteControl/Public/RemoteControlPropertyIdRegistry.h
@@ -0,0 +1,223 @@
+// Copyright Epic Games, Inc. All Rights Reserved.
+
+#pragma once
+
+#include "RemoteControlFieldPath.h"
+#include "Templates/SharedPointer.h"
+#include "UObject/NameTypes.h"
+#include "UObject/Object.h"
+#include "UObject/ObjectPtr.h"
+
+#include "RemoteControlPropertyIdRegistry.generated.h"
+
+class FProperty;
+class UClass;
+class URCVirtualPropertySelfContainer;
+class URemoteControlPreset;
+struct FGuid;
+struct FRemoteControlField;
+struct FRemoteControlPropertyIdArgs;
+struct FRemoteControlProperty;
+template <class CPPSTRUCT>
+struct TStructOpsTypeTraitsBase2;
+
+/** Wrapper class used to serialize identities of exposable entities in a generic way. */
+USTRUCT()
+struct FRCPropertyIdWrapper
+{
+	GENERATED_BODY()
+
+	FRCPropertyIdWrapper() = default;
+
+	FRCPropertyIdWrapper(const TSharedRef<FRemoteControlProperty>& InRCProperty);
+
+	/** Get the type of the wrapped entity.
+	 * @return Entity Id of this PropertyIdWrapper.
+	 */
+	const FGuid& GetEntityId() const;
+
+	/** Get the type of the wrapped entity.
+	 * @return PropertyId of this PropertyIdWrapper.
+	 */
+	FName GetPropertyId() const;
+
+	/** Get the type of the wrapped entity.
+	 * @param InNewPropertyId The new PropertyId to assign to this PropertyIdWrapper.
+	 */
+	void SetPropertyId(FName InNewPropertyId);
+
+	/** Get the super type of the wrapped property.
+	 * @return The super type of this PropertyIdWrapper.
+	 */
+	const FName& GetSuperType() const;
+
+	/** Get the sub type of the wrapped entity.
+	 * @return The sub type of this PropertyIdWrapper.
+	 */
+	const FName& GetSubType() const;
+
+	/** Get the object class to create of this wrapped entity.
+	 * @return The Object UClass to create for this property (used for ObjectProperty).
+	 */
+	const UClass* GetClassToCreate() const;
+
+	/** Returns whether the type and the underlying data is valid.
+	 * @return True if EntityId and either SuperType or SubType is valid, false otherwise.
+	 */
+	bool IsValid() const;
+
+	/** Returns whether the type and the underlying data is valid.
+	 * @return True if the PropertyId is different from the default value.
+	 */
+	bool IsValidPropertyId() const;
+
+public:
+	friend uint32 GetTypeHash(const FRCPropertyIdWrapper& Wrapper);
+
+	bool operator==(const FGuid& WrappedId) const;
+
+	bool operator==(const FRCPropertyIdWrapper& Other) const;
+
+private:
+	void UpdateTypes(const TSharedRef<FRemoteControlProperty>& InRCProperty, UObject* InOwner, FProperty* InProperty);
+
+private:
+	/** Entity identifier of the remote control property. */
+	FGuid EntityId;
+
+	/** Holds the field identifier of the remote control property. */
+	FName PropertyId;
+
+	/** Holds the type of the wrapped property. */
+	FName SuperType;
+
+	/** Holds the type of the wrapped object. */
+	FName SubType;
+
+	/** Holds the type to be created. */
+	UClass* ClassToCreate;
+};
+
+template<> struct TStructOpsTypeTraits<FRCPropertyIdWrapper> : public TStructOpsTypeTraitsBase2<FRCPropertyIdWrapper>
+{
+	enum
+	{
+		WithSerializer = false,
+		WithIdenticalViaEquality = true
+	};
+};
+
+/**
+ *	Holds a Registry of the PropertyId PropertyId
+ */
+UCLASS()
+class REMOTECONTROL_API URemoteControlPropertyIdRegistry : public UObject
+{
+	GENERATED_BODY()
+
+public:
+	/**
+	 * Initialize the Registry with the given Preset.
+	 * @param InSourcePreset Preset to get the data from.
+	 */
+	void Initialize(URemoteControlPreset* InSourcePreset);
+
+	/**
+	 * Update the value(s) of the property(ies) that are bound to a PropertyIdAction.
+	 * @param InArgs Argument used to update the property(ies) value.
+	 */
+	void PerformChainReaction(const FRemoteControlPropertyIdArgs& InArgs);
+
+	/**
+	 * Add a field to the set of identified fields.
+	 * @param InFieldToIdentify the entity to identify.
+	 */
+	void AddIdentifiedField(const TSharedRef<FRemoteControlField>& InFieldToIdentify);
+
+	/**
+	 * Check if there is any PropertyId Field.
+	 * @return True if the Registry is Empty, false otherwise.
+	 */
+	bool IsEmpty() const;
+
+	/**
+	 * Updates a field from the set of identified fields.
+	 * @param InFieldToIdentify the entity to identify.
+	 */
+	void UpdateIdentifiedField(const TSharedRef<FRemoteControlField>& InFieldToIdentify);
+
+	/**
+	 * Remove an identified field from the registry using its id.
+	 * @param InEntityId Id of the Property to remove.
+	 */
+	void RemoveIdentifiedField(const FGuid& InEntityId);
+
+	/**
+	 * Get the Object UClass to create for the property with the given Id.
+	 * @param InEntityId Id of the Property to get the class from.
+	 * @return The UClass to create for this property.
+	 */
+	const UClass* GetClassByEntityId(const FGuid& InEntityId) const;
+	
+	/**
+	 * Get the list of existing Property Ids.
+	 * @return The list of existing Property Ids that has a valid FieldId.
+	 */
+	TSet<FName> GetFieldIdsNameList();
+
+	/**
+	 * Return all EntityIds that has the given PropertyId.
+	 * @param InPropertyId Property Id to match.
+	 * @return The list of existing EntityIds that has the given PropertyId.
+	 */
+	TSet<FGuid> GetEntityIdsForPropertyId(const FName& InPropertyId) const;
+
+	/**
+	 * Get the list of existing Property Ids.
+	 * @return The list of existing Property Ids that has a valid FieldId.
+	 */
+	TSet<FGuid> GetEntityIdsList();
+
+	DECLARE_MULTICAST_DELEGATE(FOnPropertyIdUpdated)
+	FOnPropertyIdUpdated& OnPropertyIdUpdated() { return OnPropertyIdUpdatedDelegate; }
+
+#if WITH_EDITOR
+	/** Callback for refreshing the UI when an identity is updated. */
+	DECLARE_MULTICAST_DELEGATE(FOnPropertyIdActionNeedsRefresh)
+	FOnPropertyIdActionNeedsRefresh& OnPropertyIdActionNeedsRefresh() { return OnPropertyIdActionNeedsRefreshDelegate; };
+#endif
+
+private:
+	/** Call the recursive function TryCopyNonUClassOwnerProperty for the last try on copying the value to the property
+	 *	@param InArgs Argument used to expose a property.
+	 *	@param TargetRCProperty The RC property that should change.
+	 *	@param InProperty The real property to assign the new value.
+	 *	@param BoundObject The Object that has the property that need to change.
+	 *	@return True if it was able to copy the value into the property, false otherwise.
+	 */
+	bool CopyNonUClassOwnerProperty(const FRemoteControlPropertyIdArgs& InArgs, TSharedPtr<FRemoteControlProperty> TargetRCProperty, FProperty* InProperty, UObject* BoundObject);
+
+	/** Try to copy the Virtual Property value into the real property, last try for property inside container and non UClass Owner */
+	bool TryCopyNonUClassOwnerProperty(const TObjectPtr<URCVirtualPropertySelfContainer>& InVirtualPropertySelfContainer, const TSharedPtr<FRemoteControlProperty>& TargetRCProperty, FProperty* InProperty, UObject* BoundObject, FRCFieldPathInfo Segment, int32 IndexSegment, uint8* PropertyContainer,bool bCheckByteEnumComparison);
+
+private:
+	/** Holds the identified fields. */
+	UPROPERTY()
+	TSet<FRCPropertyIdWrapper> IdentifiedFields;
+
+	/** Holds the source remote control preset asset. */
+	UPROPERTY()
+	TObjectPtr<URemoteControlPreset> SourcePreset;
+
+	/** Delegate triggered when a property has its field Id changed. */
+	FOnPropertyIdUpdated OnPropertyIdUpdatedDelegate;
+
+#if WITH_EDITOR
+	/** Delegate triggered when PropertyID UI needs to refresh. */
+	FOnPropertyIdActionNeedsRefresh OnPropertyIdActionNeedsRefreshDelegate;
+#endif // WITH_EDITOR
+
+#if !WITH_EDITOR
+	uint8* StructMemoryContainer;
+#endif // !WITH_EDITOR
+};
diff --git a/Engine/Plugins/VirtualProduction/RemoteControl/Source/RemoteControlLogic/Private/Action/RCActionContainer.cpp b/Engine/Plugins/VirtualProduction/RemoteControl/Source/RemoteControlLogic/Private/Action/RCActionContainer.cpp
--- a/Engine/Plugins/VirtualProduction/RemoteControl/Source/RemoteControlLogic/Private/Action/RCActionContainer.cpp
+++ b/Engine/Plugins/VirtualProduction/RemoteControl/Source/RemoteControlLogic/Private/Action/RCActionContainer.cpp
@@ -1,9 +1,10 @@
-﻿// Copyright Epic Games, Inc. All Rights Reserved.
+// Copyright Epic Games, Inc. All Rights Reserved.
 
 #include "Action/RCActionContainer.h"
 
 #include "Action/RCAction.h"
 #include "Action/RCFunctionAction.h"
+#include "Action/RCPropertyIdAction.h"
 #include "Action/RCPropertyAction.h"
 #include "Behaviour/Builtin/Path/RCSetAssetByPathBehaviour.h"
 #include "Behaviour/RCBehaviour.h"
@@ -41,6 +42,31 @@ TRCActionUniquenessTest URCActionContainer::GetDefaultActionUniquenessTest(const
 	};
 }
 
+URCAction* URCActionContainer::AddAction()
+{
+	// Create new PropertyIdAction
+	URCPropertyIdAction* NewPropertyIdAction = NewObject<URCPropertyIdAction>(this);
+	NewPropertyIdAction->PresetWeakPtr = PresetWeakPtr;
+	NewPropertyIdAction->Id = FGuid::NewGuid();
+	NewPropertyIdAction->Initialize();
+	NewPropertyIdAction->UpdatePropertyId();
+	AddAction(NewPropertyIdAction);
+	return NewPropertyIdAction;
+}
+
+URCAction* URCActionContainer::AddAction(FName InFieldId)
+{
+	// Create new PropertyIdAction
+	URCPropertyIdAction* NewPropertyIdAction = NewObject<URCPropertyIdAction>(this);
+	NewPropertyIdAction->PresetWeakPtr = PresetWeakPtr;
+	NewPropertyIdAction->Id = FGuid::NewGuid();
+	NewPropertyIdAction->PropertyId = InFieldId;
+	NewPropertyIdAction->Initialize();
+	NewPropertyIdAction->UpdatePropertyId();
+	AddAction(NewPropertyIdAction);
+	return NewPropertyIdAction;
+}
+
 URCAction* URCActionContainer::AddAction(const TSharedRef<const FRemoteControlField> InRemoteControlField)
 {
 	return AddAction(GetDefaultActionUniquenessTest(InRemoteControlField), InRemoteControlField);
diff --git a/Engine/Plugins/VirtualProduction/RemoteControl/Source/RemoteControlLogic/Private/Action/RCPropertyIdAction.cpp b/Engine/Plugins/VirtualProduction/RemoteControl/Source/RemoteControlLogic/Private/Action/RCPropertyIdAction.cpp
new file mode 100644
--- /dev/null
+++ b/Engine/Plugins/VirtualProduction/RemoteControl/Source/RemoteControlLogic/Private/Action/RCPropertyIdAction.cpp
@@ -0,0 +1,216 @@
+﻿// Copyright Epic Games, Inc. All Rights Reserved.
+
+#include "Action/RCPropertyIdAction.h"
+
+#include "Engine/StaticMesh.h"
+#include "IPropertyIdHandler.h"
+#include "IRemoteControlModule.h"
+#include "Materials/MaterialInterface.h"
+#include "RCVirtualProperty.h"
+#include "RemoteControlPreset.h"
+#include "RemoteControlPropertyIdRegistry.h"
+
+URCPropertyIdAction::~URCPropertyIdAction()
+{
+	if (PresetWeakPtr.IsValid())
+	{
+		PresetWeakPtr->OnEntityUnexposed().RemoveAll(this);
+		PresetWeakPtr->GetPropertyIdRegistry()->OnPropertyIdUpdated().RemoveAll(this);
+	}
+}
+
+void URCPropertyIdAction::Execute() const
+{
+	if (!PresetWeakPtr.IsValid())
+	{
+		return;
+	}
+	
+	if (PropertySelfContainer.IsEmpty())
+	{
+		return;
+	}
+
+	for (const TPair<FName, TObjectPtr<URCVirtualPropertySelfContainer>>& PropertyContainer : PropertySelfContainer)
+	{
+		if (FProperty* Property = PropertyContainer.Value->GetProperty())
+		{
+			FRemoteControlPropertyIdArgs PropertyIdArgs;
+			PropertyIdArgs.VirtualProperty = PropertyContainer.Value;
+			PropertyIdArgs.PropertyId = PropertyId;
+			PropertyIdArgs.SuperType = Property->GetClass()->GetFName();
+
+			TSharedPtr<IPropertyIdHandler> PropertyIdHandler = IRemoteControlModule::Get().GetPropertyIdHandlerFor(Property);
+			if (!PropertyIdHandler.IsValid())
+			{
+				continue;
+			}
+			const EPropertyBagPropertyType PropertyBagType = PropertyIdHandler->GetPropertyType(Property);
+
+			if (PropertyBagType == EPropertyBagPropertyType::Enum)
+			{
+				PropertyIdArgs.SuperType = NAME_EnumProperty;
+				PropertyIdArgs.SubType = PropertyIdHandler->GetPropertyTypeName(Property);
+			}
+			else if (PropertyBagType == EPropertyBagPropertyType::Object)
+			{
+				if (const TSharedPtr<FStructOnScope> StructOnScope = PropertyContainer.Value->CreateStructOnScope())
+				{
+					if (const FObjectProperty* ObjectProperty = CastField<FObjectProperty>(PropertyContainer.Value->GetProperty()))
+					{
+						if (const uint8* PropertyValuePtr = ObjectProperty->ContainerPtrToValuePtr<uint8>(StructOnScope->GetStructMemory()))
+						{
+							PropertyIdArgs.SourceObject = ObjectProperty->GetObjectPropertyValue(PropertyValuePtr);
+							if (PropertyIdArgs.SourceObject)
+							{
+								if (PropertyIdArgs.SourceObject->IsA(UMaterialInterface::StaticClass()))
+								{
+									PropertyIdArgs.SourceClass = UMaterialInterface::StaticClass();
+								}
+								else
+								{
+									PropertyIdArgs.SourceClass = PropertyIdArgs.SourceObject->GetClass();
+								}
+								PropertyIdArgs.SubType = PropertyIdArgs.SourceClass->GetFName();
+							}
+						}
+					}
+				}
+			}
+			else if (PropertyBagType == EPropertyBagPropertyType::Struct)
+			{
+				if (const FStructProperty* StructProperty = CastField<FStructProperty>(Property))
+				{
+					PropertyIdArgs.SubType = StructProperty->Struct->GetFName();
+				}
+			}
+			PresetWeakPtr->PerformChainReaction(PropertyIdArgs);
+		}
+	}
+	Super::Execute();
+}
+
+void URCPropertyIdAction::PostLoad()
+{
+	UObject::PostLoad();
+	Initialize();
+}
+
+#if WITH_EDITOR
+void URCPropertyIdAction::PostEditChangeChainProperty(FPropertyChangedChainEvent& PropertyChangedEvent)
+{
+	Super::PostEditChangeChainProperty(PropertyChangedEvent);
+	if (const FProperty* Property = PropertyChangedEvent.Property)
+	{
+		if (Property->GetFName() == GET_MEMBER_NAME_CHECKED(URCPropertyIdAction, PropertyId))
+		{
+			UpdatePropertyId();
+		}
+	}
+}
+#endif // WITH_EDITOR
+
+void URCPropertyIdAction::UpdatePropertyId()
+{
+	DefaultObject = nullptr;
+	if (URemoteControlPreset* Preset = PresetWeakPtr.Get())
+	{
+		PropertySelfContainer.Empty();
+		for (const FGuid& TargetProperty : Preset->GetPropertyIdRegistry().Get()->GetEntityIdsList())
+		{
+			if (const TSharedPtr<FRemoteControlProperty> TargetRCProperty = Preset->GetExposedEntity<FRemoteControlProperty>(TargetProperty).Pin())
+			{
+				if (TargetRCProperty->PropertyId == PropertyId)
+				{
+					if (FProperty* Property = TargetRCProperty->GetProperty())
+					{
+						TSharedPtr<IPropertyIdHandler> PropertyIdHandler = IRemoteControlModule::Get().GetPropertyIdHandlerFor(Property);
+						if (!PropertyIdHandler.IsValid())
+						{
+							continue;
+						}
+						FName PropertyClassName = PropertyIdHandler->GetPropertyTypeName(Property);
+						const FName& NewPropertyIdName = *(PropertyId.ToString() + TEXT(".") + PropertyClassName.ToString());
+						if (PropertyIdHandler->GetPropertyType(Property) == EPropertyBagPropertyType::Object)
+						{
+							const FName PropertyNameToSearchFor = GET_MEMBER_NAME_CHECKED(URCPropertyIdAction, DefaultObject);
+							DefaultObject = PropertyIdHandler->GetObjectPropertyDefaultValue(Property, Preset->GetPropertyIdRegistry().Get()->GetClassByEntityId(TargetProperty));
+							if (DefaultObject && !PropertyNameToSearchFor.IsNone())
+							{
+								for (TFieldIterator<FObjectProperty> FieldIt(GetClass()); FieldIt; ++FieldIt)
+								{
+									if (FieldIt->GetFName() == PropertyNameToSearchFor)
+									{
+										FieldIt->PropertyClass = DefaultObject->GetClass();
+										if (CachedPropertySelfContainer.Contains(PropertyClassName))
+										{
+											PropertySelfContainer.Add(PropertyClassName, CachedPropertySelfContainer[PropertyClassName]);
+										}
+										else
+										{
+											CachedPropertySelfContainer.Add(PropertyClassName, NewObject<URCVirtualPropertySelfContainer>(this));
+											CachedPropertySelfContainer[PropertyClassName]->DuplicateProperty(NewPropertyIdName, *FieldIt);
+											PropertySelfContainer.Add(PropertyClassName, CachedPropertySelfContainer[PropertyClassName]);
+										}
+										break;
+									}
+								}
+							}
+						}
+						else
+						{
+							if (CachedPropertySelfContainer.Contains(PropertyClassName))
+							{
+								PropertySelfContainer.Add(PropertyClassName, CachedPropertySelfContainer[PropertyClassName]);
+							}
+							else
+							{
+								CachedPropertySelfContainer.Add(PropertyClassName, NewObject<URCVirtualPropertySelfContainer>(this));
+								CachedPropertySelfContainer[PropertyClassName]->AddProperty(NewPropertyIdName,
+									PropertyIdHandler->GetPropertyType(Property),
+									PropertyIdHandler->GetPropertyTypeObject(Property));
+								
+								PropertySelfContainer.Add(PropertyClassName, CachedPropertySelfContainer[PropertyClassName]);
+							}
+						}
+					}
+				}
+			}
+		}
+#if WITH_EDITOR
+		if (PresetWeakPtr.IsValid())
+		{
+			PresetWeakPtr->GetPropertyIdRegistry()->OnPropertyIdActionNeedsRefresh().Broadcast();
+		}
+#endif
+	}
+}
+
+void URCPropertyIdAction::Initialize()
+{
+	if (PresetWeakPtr.IsValid())
+	{
+		PresetWeakPtr->OnEntityUnexposed().AddUObject(this, &URCPropertyIdAction::OnEntityUnexposed);
+		PresetWeakPtr->GetPropertyIdRegistry()->OnPropertyIdUpdated().AddUObject(this, &URCPropertyIdAction::UpdatePropertyId);
+	}
+}
+
+void URCPropertyIdAction::OnEntityUnexposed(URemoteControlPreset* InPreset, const FGuid& InGuid)
+{
+	if (InPreset)
+	{
+		const TWeakPtr<FRemoteControlField> UnexposedEntity = StaticCastWeakPtr<FRemoteControlField>(InPreset->GetExposedEntity(InGuid));
+		if (UnexposedEntity.IsValid())
+		{
+			//If the propertyId of the property unexposed is the same of this PropertyIdAction refresh, otherwise don't update it
+			if (UnexposedEntity.Pin()->PropertyId == PropertyId)
+			{
+				UpdatePropertyId();
+			}
+			return;
+		}
+	}
+	//If there is anything wrong with the InPreset or InGuid it won't be able to check the PropertyId
+	//We refresh to be sure to not miss anything
+	UpdatePropertyId();
+}
diff --git a/Engine/Plugins/VirtualProduction/RemoteControl/Source/RemoteControlLogic/Private/Behaviour/RCBehaviour.cpp b/Engine/Plugins/VirtualProduction/RemoteControl/Source/RemoteControlLogic/Private/Behaviour/RCBehaviour.cpp
--- a/Engine/Plugins/VirtualProduction/RemoteControl/Source/RemoteControlLogic/Private/Behaviour/RCBehaviour.cpp
+++ b/Engine/Plugins/VirtualProduction/RemoteControl/Source/RemoteControlLogic/Private/Behaviour/RCBehaviour.cpp
@@ -1,4 +1,4 @@
-﻿// Copyright Epic Games, Inc. All Rights Reserved.
+// Copyright Epic Games, Inc. All Rights Reserved.
 
 #include "Behaviour/RCBehaviour.h"
 
@@ -30,6 +30,24 @@ void URCBehaviour::Execute()
 	}
 }
 
+URCAction* URCBehaviour::AddAction()
+{
+#if WITH_EDITOR
+	ActionContainer->Modify();
+#endif // WITH_EDITOR
+
+	return ActionContainer->AddAction();
+}
+
+URCAction* URCBehaviour::AddAction(FName InFieldId)
+{
+#if WITH_EDITOR
+	ActionContainer->Modify();
+#endif // WITH_EDITOR
+
+	return ActionContainer->AddAction(InFieldId);
+}
+
 URCAction* URCBehaviour::AddAction(const TSharedRef<const FRemoteControlField> InRemoteControlField)
 {
 #if WITH_EDITOR
diff --git a/Engine/Plugins/VirtualProduction/RemoteControl/Source/RemoteControlLogic/Public/Action/RCActionContainer.h b/Engine/Plugins/VirtualProduction/RemoteControl/Source/RemoteControlLogic/Public/Action/RCActionContainer.h
--- a/Engine/Plugins/VirtualProduction/RemoteControl/Source/RemoteControlLogic/Public/Action/RCActionContainer.h
+++ b/Engine/Plugins/VirtualProduction/RemoteControl/Source/RemoteControlLogic/Public/Action/RCActionContainer.h
@@ -1,4 +1,4 @@
-﻿// Copyright Epic Games, Inc. All Rights Reserved.
+// Copyright Epic Games, Inc. All Rights Reserved.
 
 #pragma once
 
@@ -12,6 +12,7 @@ class URCAction;
 class URCBehaviour;
 class URCController;
 class URCFunctionAction;
+class URCPropertyIdAction;
 class URCPropertyAction;
 class URCPropertyBindAction;
 
@@ -37,6 +38,12 @@ public:
 
 	TRCActionUniquenessTest GetDefaultActionUniquenessTest(const TSharedRef<const FRemoteControlField> InRemoteControlField);
 
+	/** Add remote control identity action. */
+	URCAction* AddAction();
+
+	/** Add remote control identity action. */
+	URCAction* AddAction(FName InFieldId);
+
 	/** Add remote control property action  */
 	URCAction* AddAction(const TSharedRef<const FRemoteControlField> InRemoteControlField);
 
diff --git a/Engine/Plugins/VirtualProduction/RemoteControl/Source/RemoteControlLogic/Public/Action/RCPropertyIdAction.h b/Engine/Plugins/VirtualProduction/RemoteControl/Source/RemoteControlLogic/Public/Action/RCPropertyIdAction.h
new file mode 100644
--- /dev/null
+++ b/Engine/Plugins/VirtualProduction/RemoteControl/Source/RemoteControlLogic/Public/Action/RCPropertyIdAction.h
@@ -0,0 +1,55 @@
+﻿// Copyright Epic Games, Inc. All Rights Reserved.
+
+#pragma once
+
+#include "CoreMinimal.h"
+#include "Action/RCAction.h"
+#include "RCPropertyIdAction.generated.h"
+
+class URCVirtualPropertySelfContainer;
+
+/**
+ * Action for PropertyId
+ */
+UCLASS()
+class REMOTECONTROLLOGIC_API URCPropertyIdAction : public URCAction
+{
+	GENERATED_BODY()
+	virtual ~URCPropertyIdAction() override;
+
+public:
+	//~ BEGIN : URCAction Interface
+	virtual void Execute() const override;
+	virtual void PostLoad() override;
+	//~ END : URCAction Interface
+	
+	//~ BEGIN : UObject Interface
+#if WITH_EDITOR
+	virtual void PostEditChangeChainProperty(struct FPropertyChangedChainEvent& PropertyChangedEvent) override;
+#endif // WITH_EDITOR
+	//~ END : URCAction Interface
+
+	void UpdatePropertyId();
+
+	void Initialize();
+
+	void OnEntityUnexposed(URemoteControlPreset* InPreset, const FGuid& InGuid);
+
+public:
+	/** Holds the field identifier associated with this. */
+	UPROPERTY()
+	FName PropertyId = NAME_None;
+
+	/** Virtual Property Container */
+	UPROPERTY()
+	TMap<FName, TObjectPtr<URCVirtualPropertySelfContainer>> PropertySelfContainer;
+
+	/** Cached Virtual Property Container */
+	UPROPERTY()
+	TMap<FName, TObjectPtr<URCVirtualPropertySelfContainer>> CachedPropertySelfContainer;
+
+private:
+	/** Holds the default Object. */
+	UPROPERTY(Transient)
+	TObjectPtr<UObject> DefaultObject = nullptr;
+};
diff --git a/Engine/Plugins/VirtualProduction/RemoteControl/Source/RemoteControlLogic/Public/Behaviour/RCBehaviour.h b/Engine/Plugins/VirtualProduction/RemoteControl/Source/RemoteControlLogic/Public/Behaviour/RCBehaviour.h
--- a/Engine/Plugins/VirtualProduction/RemoteControl/Source/RemoteControlLogic/Public/Behaviour/RCBehaviour.h
+++ b/Engine/Plugins/VirtualProduction/RemoteControl/Source/RemoteControlLogic/Public/Behaviour/RCBehaviour.h
@@ -1,4 +1,4 @@
-﻿// Copyright Epic Games, Inc. All Rights Reserved.
+// Copyright Epic Games, Inc. All Rights Reserved.
 
 #pragma once
 
@@ -12,6 +12,7 @@ struct FRemoteControlField;
 class URCBehaviourNode;
 class URCController;
 class URCFunctionAction;
+class URCPropertyIdAction;
 class URCPropertyAction;
 class URemoteControlPreset;
 class URCAction;
@@ -41,6 +42,12 @@ public:
 	/** Execute the behaviour */
 	virtual void Execute();
 
+	/** Add a Logic action as an identity action. */
+	virtual URCAction* AddAction();
+
+	/** Add a Logic action as an identity action. */
+	virtual URCAction* AddAction(FName InFieldId);
+
 	/** Add a Logic action using a remote control field as input */
 	virtual URCAction* AddAction(const TSharedRef<const FRemoteControlField> InRemoteControlField);
 
diff --git a/Engine/Plugins/VirtualProduction/RemoteControl/Source/RemoteControlUI/Private/UI/Action/Bind/RCActionBindModel.h b/Engine/Plugins/VirtualProduction/RemoteControl/Source/RemoteControlUI/Private/UI/Action/Bind/RCActionBindModel.h
--- a/Engine/Plugins/VirtualProduction/RemoteControl/Source/RemoteControlUI/Private/UI/Action/Bind/RCActionBindModel.h
+++ b/Engine/Plugins/VirtualProduction/RemoteControl/Source/RemoteControlUI/Private/UI/Action/Bind/RCActionBindModel.h
@@ -3,6 +3,7 @@
 #pragma once
 
 #include "Action/Bind/RCPropertyBindAction.h"
+#include "Action/RCPropertyIdAction.h"
 #include "UI/Action/RCActionModel.h"
 
 class SActionItemListRow;
@@ -48,4 +49,34 @@ public:
 	virtual FLinearColor GetActionTypeColor() const override;
 };
 
+/**
+ * FRCPropertyIdActionBindModel
+ *
+ * UI model for PropertyId based Bind Actions
+ */
+class FRCPropertyIdActionBindModel : public FRCActionBindModel, public FRCPropertyIdActionType
+{
+public:
+	FRCPropertyIdActionBindModel(URCPropertyIdAction* InPropertyIdAction, const TSharedPtr<class FRCBehaviourModel> InBehaviourItem, const TSharedPtr<SRemoteControlPanel> InRemoteControlPanel)
+		: FRCActionBindModel(InPropertyIdAction, InBehaviourItem, InRemoteControlPanel)
+		, FRCPropertyIdActionType(InPropertyIdAction)
+	{
+	}
+	/** Color code for this Action*/
+	virtual FLinearColor GetActionTypeColor() const override
+	{
+		return GetPropertyIdTypeColor();
+	}
+	/** Widget representing Action Name field */
+	virtual TSharedRef<SWidget> GetNameWidget() const override
+	{
+		return GetPropertyIdNameWidget();
+	}
+	/** Widget representing the Value field */
+	virtual TSharedRef<SWidget> GetWidget() const override
+	{
+		return GetPropertyIdValueWidget();
+	}
+};
+
 // Note: Bind only supports Property based Actions, so Function action model is not necessary here
\ No newline at end of file
diff --git a/Engine/Plugins/VirtualProduction/RemoteControl/Source/RemoteControlUI/Private/UI/Action/Conditional/RCActionConditionalModel.cpp b/Engine/Plugins/VirtualProduction/RemoteControl/Source/RemoteControlUI/Private/UI/Action/Conditional/RCActionConditionalModel.cpp
--- a/Engine/Plugins/VirtualProduction/RemoteControl/Source/RemoteControlUI/Private/UI/Action/Conditional/RCActionConditionalModel.cpp
+++ b/Engine/Plugins/VirtualProduction/RemoteControl/Source/RemoteControlUI/Private/UI/Action/Conditional/RCActionConditionalModel.cpp
@@ -31,7 +31,7 @@ namespace UE::RCActionPanelConditionalList
 		const FName TypeColorTag = TEXT("TypeColorTag");
 		const FName DragDropHandle = TEXT("DragDropHandle");
 		const FName Condition = TEXT("Condition");
-		const FName Description = TEXT("Description");
+		const FName Description = TEXT("PropertyID");
 		const FName Value = TEXT("Value");
 	}
 
@@ -99,7 +99,7 @@ TSharedPtr<SHeaderRow> FRCActionConditionalModel::GetHeaderRow()
 		.HeaderContentPadding(RCPanelStyle->HeaderRowPadding)
 
 		+ SHeaderRow::Column(UE::RCActionPanelConditionalList::Columns::Description)
-		.DefaultLabel(LOCTEXT("RCActionDescColumnHeader", "Description"))
+		.DefaultLabel(LOCTEXT("RCActionDescColumnHeader", "PropertyID"))
 		.FillWidth(0.4f)
 		.HeaderContentPadding(RCPanelStyle->HeaderRowPadding)
 
@@ -124,6 +124,10 @@ TSharedPtr<FRCActionConditionalModel> FRCActionConditionalModel::GetModelByActio
 	{
 		return MakeShared<FRCFunctionActionConditionalModel>(FunctionAction, InBehaviourItem, InRemoteControlPanel);
 	}
+	else if (URCPropertyIdAction* PropertyIdAction = Cast<URCPropertyIdAction>(InAction))
+	{
+		return MakeShared<FRCPropertyIdActionConditionalModel>(PropertyIdAction, InBehaviourItem, InRemoteControlPanel);
+	}
 
 	return nullptr;
 }
diff --git a/Engine/Plugins/VirtualProduction/RemoteControl/Source/RemoteControlUI/Private/UI/Action/Conditional/RCActionConditionalModel.h b/Engine/Plugins/VirtualProduction/RemoteControl/Source/RemoteControlUI/Private/UI/Action/Conditional/RCActionConditionalModel.h
--- a/Engine/Plugins/VirtualProduction/RemoteControl/Source/RemoteControlUI/Private/UI/Action/Conditional/RCActionConditionalModel.h
+++ b/Engine/Plugins/VirtualProduction/RemoteControl/Source/RemoteControlUI/Private/UI/Action/Conditional/RCActionConditionalModel.h
@@ -3,6 +3,7 @@
 #pragma once
 
 #include "Action/RCFunctionAction.h"
+#include "Action/RCPropertyIdAction.h"
 #include "Action/RCPropertyAction.h"
 #include "UI/Action/RCActionModel.h"
 
@@ -84,4 +85,36 @@ public:
 	{
 		return GetFunctionTypeColor();
 	}
-};
\ No newline at end of file
+};
+
+/**
+ * FRCPropertyIdActionConditionalModel
+ *
+ * UI model for PropertyId based Conditional Actions
+ */
+class FRCPropertyIdActionConditionalModel : public FRCActionConditionalModel, public FRCPropertyIdActionType
+{
+public:
+	FRCPropertyIdActionConditionalModel(URCPropertyIdAction* InPropertyIdAction, const TSharedPtr<class FRCBehaviourModel> InBehaviourItem, const TSharedPtr<SRemoteControlPanel> InRemoteControlPanel)
+		: FRCActionConditionalModel(InPropertyIdAction, InBehaviourItem, InRemoteControlPanel)
+		, FRCPropertyIdActionType(InPropertyIdAction)
+	{}
+
+	/** Color code for this Action*/
+	virtual FLinearColor GetActionTypeColor() const override
+	{
+		return GetPropertyIdTypeColor();
+	}
+
+	/** Widget representing Action Name field */
+	virtual TSharedRef<SWidget> GetNameWidget() const override
+	{
+		return GetPropertyIdNameWidget();
+	}
+
+	/** Widget representing the Value field */
+	virtual TSharedRef<SWidget> GetWidget() const override
+	{
+		return GetPropertyIdValueWidget();
+	}
+};
diff --git a/Engine/Plugins/VirtualProduction/RemoteControl/Source/RemoteControlUI/Private/UI/Action/RCActionModel.cpp b/Engine/Plugins/VirtualProduction/RemoteControl/Source/RemoteControlUI/Private/UI/Action/RCActionModel.cpp
--- a/Engine/Plugins/VirtualProduction/RemoteControl/Source/RemoteControlUI/Private/UI/Action/RCActionModel.cpp
+++ b/Engine/Plugins/VirtualProduction/RemoteControl/Source/RemoteControlUI/Private/UI/Action/RCActionModel.cpp
@@ -1,4 +1,4 @@
-﻿// Copyright Epic Games, Inc. All Rights Reserved.
+// Copyright Epic Games, Inc. All Rights Reserved.
 
 #include "RCActionModel.h"
 
@@ -177,6 +177,10 @@ TSharedPtr<FRCActionModel> FRCActionModel::GetModelByActionType(URCAction* InAct
 	{
 		return MakeShared<FRCFunctionActionModel>(FunctionAction, InBehaviourItem, InRemoteControlPanel);
 	}
+	else if (URCPropertyIdAction* PropertyIdAction = Cast<URCPropertyIdAction>(InAction))
+	{
+		return MakeShared<FRCPropertyIdActionModel>(PropertyIdAction, InBehaviourItem, InRemoteControlPanel);
+	}
 	else
 		return nullptr;
 }
@@ -290,4 +294,169 @@ void FRCActionModel::OnSelectionExit()
 	}
 }
 
+FRCPropertyIdActionType::FRCPropertyIdActionType(URCPropertyIdAction* InPropertyIdAction)
+	: PropertyIdActionWeakPtr(InPropertyIdAction)
+{
+	FPropertyRowGeneratorArgs Args;
+	Args.bShouldShowHiddenProperties = true;
+	
+	PropertyIdNameRowGenerator = FModuleManager::GetModuleChecked<FPropertyEditorModule>("PropertyEditor").CreatePropertyRowGenerator(Args);
+	if (InPropertyIdAction)
+	{
+		RefreshNameWidget();
+		RefreshValueWidget();
+	}
+#if WITH_EDITOR
+	if (URemoteControlPreset* Preset = InPropertyIdAction->PresetWeakPtr.Get())
+	{
+		Preset->GetPropertyIdRegistry()->OnPropertyIdActionNeedsRefresh().AddRaw(this, &FRCPropertyIdActionType::RefreshValueWidget);
+	}
+#endif
+}
+
+FRCPropertyIdActionType::~FRCPropertyIdActionType()
+{
+#if WITH_EDITOR
+	if (const URCPropertyIdAction* PropertyIdAction = PropertyIdActionWeakPtr.Get())
+	{
+		if (URemoteControlPreset* Preset = PropertyIdAction->PresetWeakPtr.Get())
+		{
+			Preset->GetPropertyIdRegistry()->OnPropertyIdActionNeedsRefresh().RemoveAll(this);
+		}
+	}
+#endif
+}
+
+FLinearColor FRCPropertyIdActionType::GetPropertyIdTypeColor() const
+{
+	// @todo: Confirm color to be used for this with VP team.
+	return FLinearColor(FColor(32, 191, 107));
+}
+
+TSharedRef<SWidget> FRCPropertyIdActionType::GetPropertyIdNameWidget() const
+{
+	if (!FieldIdTreeNodeWeakPtr.IsValid())
+	{
+		return SNullWidget::NullWidget;
+	}
+	const FNodeWidgets FieldIdNodeWidgets = FieldIdTreeNodeWeakPtr.Pin()->CreateNodeWidgets();
+	const TSharedRef<SHorizontalBox> NameWidget = SNew(SHorizontalBox);
+	if (FieldIdNodeWidgets.ValueWidget)
+	{
+		NameWidget->AddSlot()
+			.Padding(3.f, 2.f)
+			.VAlign(VAlign_Center)
+			[
+				FieldIdNodeWidgets.ValueWidget.ToSharedRef()
+			];
+	}
+	return NameWidget;
+}
+
+TSharedRef<SWidget> FRCPropertyIdActionType::GetPropertyIdValueWidget() const
+{
+	TSharedRef<SVerticalBox> VerticalBox = SNew(SVerticalBox);
+
+	for (TPair<FName, TWeakPtr<IDetailTreeNode>> ValueTreeNode : ValueTreeNodeWeakPtr)
+	{
+		// store it and just remove the unusued to not reset the value widget
+		VerticalBox->AddSlot()
+			.AutoHeight()
+			[
+				UE::RCUIHelpers::GetGenericFieldWidget(ValueTreeNode.Value.Pin())
+			];
+	}
+	return VerticalBox;
+}
+
+void FRCPropertyIdActionType::RefreshNameWidget()
+{
+	if (URCPropertyIdAction* PropertyIdAction = PropertyIdActionWeakPtr.Get())
+	{
+		// Since we must keep many PRG objects alive in order to access the handle data, validating the nodes each tick is very taxing.
+		// We can override the validation with a lambda since the validation function in PRG is not necessary for our implementation
+		auto ValidationLambda = ([](const FRootPropertyNodeList& PropertyNodeList) { return true; });
+		PropertyIdNameRowGenerator->SetObjects({ PropertyIdAction });
+		PropertyIdNameRowGenerator->SetCustomValidatePropertyNodesFunction(FOnValidatePropertyRowGeneratorNodes::CreateLambda(MoveTemp(ValidationLambda)));
+		for (const TSharedRef<IDetailTreeNode>& CategoryNode : PropertyIdNameRowGenerator->GetRootTreeNodes())
+		{
+			TArray<TSharedRef<IDetailTreeNode>> Children;
+			CategoryNode->GetChildren(Children);
+			bool bFoundFieldId = false;
+			for (const TSharedRef<IDetailTreeNode>& Child : Children)
+			{
+				const TSharedPtr<IPropertyHandle> PropertyHandle = Child->CreatePropertyHandle();
+				if (PropertyHandle && PropertyHandle->IsValidHandle())
+				{
+					if (const FProperty* Property = PropertyHandle->GetProperty())
+					{
+						if (Property->GetFName() == GET_MEMBER_NAME_CHECKED(URCPropertyIdAction, PropertyId))
+						{
+							FieldIdTreeNodeWeakPtr = Child;
+							bFoundFieldId = true;
+						}
+						if (bFoundFieldId)
+						{
+							break;
+						}
+					}
+				}
+			}
+		}
+	}
+}
+
+void FRCPropertyIdActionType::RefreshValueWidget()
+{
+	if (URCPropertyIdAction* PropertyIdAction = PropertyIdActionWeakPtr.Get())
+	{
+		// Generate UI widget for Action input
+		PropertyIdValueRowGenerator.Empty();
+		ValueTreeNodeWeakPtr.Empty();
+		for (const TPair<FName, TObjectPtr<URCVirtualPropertySelfContainer>>& PropertyContainer : PropertyIdAction->PropertySelfContainer)
+		{
+			if (IsValid(PropertyContainer.Value))
+			{
+				if (const TSharedPtr<FStructOnScope> StructOnScope = PropertyContainer.Value->CreateStructOnScope())
+				{
+					if (const TSharedPtr<IPropertyRowGenerator>* Generator = CachedPropertyIdValueRowGenerator.Find(PropertyContainer.Key))
+					{
+						PropertyIdValueRowGenerator.Add(PropertyContainer.Key, (*Generator));
+					}
+					else
+					{
+						FPropertyRowGeneratorArgs Args;
+						Args.bShouldShowHiddenProperties = true;
+
+						CachedPropertyIdValueRowGenerator.Add(PropertyContainer.Key, FModuleManager::GetModuleChecked<FPropertyEditorModule>("PropertyEditor").CreatePropertyRowGenerator(Args));
+						CachedPropertyIdValueRowGenerator[PropertyContainer.Key]->SetStructure(StructOnScope);
+
+						PropertyIdValueRowGenerator.Add(PropertyContainer.Key, CachedPropertyIdValueRowGenerator[PropertyContainer.Key]);
+					}
+					for (const TSharedRef<IDetailTreeNode>& CategoryNode : PropertyIdValueRowGenerator[PropertyContainer.Key]->GetRootTreeNodes())
+					{
+						TArray<TSharedRef<IDetailTreeNode>> Children;
+						CategoryNode->GetChildren(Children);
+						for (const TSharedRef<IDetailTreeNode>& Child : Children)
+						{
+							// For regular properties (non-container)
+							if (const TWeakPtr<IDetailTreeNode>* ValueTreeNode = CachedValueTreeNodeWeakPtr.Find(PropertyContainer.Key))
+							{
+								ValueTreeNodeWeakPtr.Add(PropertyContainer.Key, (*ValueTreeNode));
+								break;
+							}
+							else
+							{
+								CachedValueTreeNodeWeakPtr.Add(PropertyContainer.Key, Child);
+								ValueTreeNodeWeakPtr.Add(PropertyContainer.Key, CachedValueTreeNodeWeakPtr[PropertyContainer.Key]);
+								break;
+							}
+						}
+					}
+				}
+			}
+		}
+	}
+}
+
 #undef LOCTEXT_NAMESPACE
\ No newline at end of file
diff --git a/Engine/Plugins/VirtualProduction/RemoteControl/Source/RemoteControlUI/Private/UI/Action/RCActionModel.h b/Engine/Plugins/VirtualProduction/RemoteControl/Source/RemoteControlUI/Private/UI/Action/RCActionModel.h
--- a/Engine/Plugins/VirtualProduction/RemoteControl/Source/RemoteControlUI/Private/UI/Action/RCActionModel.h
+++ b/Engine/Plugins/VirtualProduction/RemoteControl/Source/RemoteControlUI/Private/UI/Action/RCActionModel.h
@@ -1,8 +1,9 @@
-﻿// Copyright Epic Games, Inc. All Rights Reserved.
+// Copyright Epic Games, Inc. All Rights Reserved.
 
 #pragma once
 
 #include "Action/RCFunctionAction.h"
+#include "Action/RCPropertyIdAction.h"
 #include "Action/RCPropertyAction.h"
 #include "UI/BaseLogicUI/RCLogicModeBase.h"
 
@@ -131,7 +132,7 @@ public:
 	FLinearColor GetFunctionTypeColor() const
 	{
 		// @todo: Confirm color to be used for this with VP team.
-		return FLinearColor::Blue; 
+		return FLinearColor::Blue;
 	}
 
 protected:
@@ -139,6 +140,59 @@ protected:
 	TWeakObjectPtr<URCFunctionAction> FunctionActionWeakPtr;
 };
 
+/**
+ * FRCPropertyIdActionType
+ *
+ * Reusable class containing PropertyID action related data and functions
+ */
+class FRCPropertyIdActionType
+{
+public:
+	FRCPropertyIdActionType(URCPropertyIdAction* InPropertyIdAction);
+
+	~FRCPropertyIdActionType();
+
+	/** Color code for function actions in the Actions table */
+	FLinearColor GetPropertyIdTypeColor() const;
+
+	/** Widget representing Action Name field */
+	TSharedRef<SWidget> GetPropertyIdNameWidget() const;
+
+	/** The widget to be rendered for this PropertyId
+	 * Represents the input field which the user will use to set a value for this Action
+	 */
+	TSharedRef<SWidget> GetPropertyIdValueWidget() const;
+
+private:
+	/** Recreates the name section of this row widget. */
+	void RefreshNameWidget();
+
+	/** Recreates the value section of this row widget. */
+	void RefreshValueWidget();
+
+protected:
+	/** The PropertyId Action (data model) associated with this. */
+	TWeakObjectPtr<URCPropertyIdAction> PropertyIdActionWeakPtr;
+
+	/** The row generator used to represent name of this widget as a row, when used with SListView */
+	TSharedPtr<IPropertyRowGenerator> PropertyIdNameRowGenerator;
+	
+	/** The cache of the row generator used to not recreate it */
+	TMap<FName, TSharedPtr<IPropertyRowGenerator>> PropertyIdValueRowGenerator;
+
+	/** The row generator used to represent value of this widget as a row, when used with SListView */
+	TMap<FName, TSharedPtr<IPropertyRowGenerator>> CachedPropertyIdValueRowGenerator;
+
+	/** Used to create a generic Field Id Widget for the property row widget */
+	TWeakPtr<IDetailTreeNode> FieldIdTreeNodeWeakPtr;
+	
+	/** Used to create a generic Value Widget for the property row widget */
+	TMap<FName, TWeakPtr<IDetailTreeNode>> ValueTreeNodeWeakPtr;
+
+	/** Used to cache a generic Value Widget for the property row widget to not recreate it */
+	TMap<FName, TWeakPtr<IDetailTreeNode>> CachedValueTreeNodeWeakPtr;
+};
+
 /*
 * ~ FRCPropertyActionModel ~
 *
@@ -181,3 +235,35 @@ public:
 		return GetFunctionTypeColor();
 	}
 };
+
+/**
+ * FRCPropertyIdActionModel
+ *
+ * UI model for PropertyId based Actions
+ */
+class FRCPropertyIdActionModel : public FRCActionModel, public FRCPropertyIdActionType
+{
+public:
+	FRCPropertyIdActionModel(URCPropertyIdAction* InPropertyIdAction, const TSharedPtr<class FRCBehaviourModel> InBehaviourItem, const TSharedPtr<SRemoteControlPanel> InRemoteControlPanel)
+		: FRCActionModel(InPropertyIdAction, InBehaviourItem, InRemoteControlPanel)
+		, FRCPropertyIdActionType(InPropertyIdAction)
+	{}
+
+	/** Color code for this Action*/
+	virtual FLinearColor GetActionTypeColor() const override
+	{
+		return GetPropertyIdTypeColor();
+	}
+
+	/** Widget representing Action Name field */
+	virtual TSharedRef<SWidget> GetNameWidget() const override
+	{
+		return GetPropertyIdNameWidget();
+	}
+
+	/** Widget representing the Value field */
+	virtual TSharedRef<SWidget> GetWidget() const override
+	{
+		return GetPropertyIdValueWidget();
+	}
+};
diff --git a/Engine/Plugins/VirtualProduction/RemoteControl/Source/RemoteControlUI/Private/UI/Action/RangeMap/RCActionRangeMapModel.h b/Engine/Plugins/VirtualProduction/RemoteControl/Source/RemoteControlUI/Private/UI/Action/RangeMap/RCActionRangeMapModel.h
--- a/Engine/Plugins/VirtualProduction/RemoteControl/Source/RemoteControlUI/Private/UI/Action/RangeMap/RCActionRangeMapModel.h
+++ b/Engine/Plugins/VirtualProduction/RemoteControl/Source/RemoteControlUI/Private/UI/Action/RangeMap/RCActionRangeMapModel.h
@@ -3,6 +3,7 @@
 #pragma once
 
 #include "Action/RCFunctionAction.h"
+#include "Action/RCPropertyIdAction.h"
 #include "Action/RCPropertyAction.h"
 #include "UI/Action/RCActionModel.h"
 
@@ -78,4 +79,33 @@ public:
 	{
 		return GetFunctionTypeColor();
 	}
-};
\ No newline at end of file
+};
+
+/**
+ * FRCPropertyIdActionRangeMapModel
+ *
+ * UI model for PropertyId based RangeMap Actions
+ */
+class FRCPropertyIdActionRangeMapModel : public FRCActionRangeMapModel, public FRCPropertyIdActionType
+{
+public:
+	FRCPropertyIdActionRangeMapModel(URCPropertyIdAction* InPropertyIdAction, const TSharedPtr<class FRCBehaviourModel> InBehaviourItem, const TSharedPtr<SRemoteControlPanel> InRemoteControlPanel)
+		: FRCActionRangeMapModel(InPropertyIdAction, InBehaviourItem, InRemoteControlPanel)
+		, FRCPropertyIdActionType(InPropertyIdAction)
+	{}
+	/** Color code for this Action*/
+	virtual FLinearColor GetActionTypeColor() const override
+	{
+		return GetPropertyIdTypeColor();
+	}
+	/** Widget representing Action Name field */
+	virtual TSharedRef<SWidget> GetNameWidget() const override
+	{
+		return GetPropertyIdNameWidget();
+	}
+	/** Widget representing the Value field */
+	virtual TSharedRef<SWidget> GetWidget() const override
+	{
+		return GetPropertyIdValueWidget();
+	}
+};
diff --git a/Engine/Plugins/VirtualProduction/RemoteControl/Source/RemoteControlUI/Private/UI/Action/SRCActionPanel.cpp b/Engine/Plugins/VirtualProduction/RemoteControl/Source/RemoteControlUI/Private/UI/Action/SRCActionPanel.cpp
--- a/Engine/Plugins/VirtualProduction/RemoteControl/Source/RemoteControlUI/Private/UI/Action/SRCActionPanel.cpp
+++ b/Engine/Plugins/VirtualProduction/RemoteControl/Source/RemoteControlUI/Private/UI/Action/SRCActionPanel.cpp
@@ -1,16 +1,20 @@
-﻿// Copyright Epic Games, Inc. All Rights Reserved.
+// Copyright Epic Games, Inc. All Rights Reserved.
 
 #include "SRCActionPanel.h"
 
 #include "RemoteControlField.h"
+#include "RemoteControlPropertyIdRegistry.h"
 #include "RemoteControlPreset.h"
 
 #include "Action/RCActionContainer.h"
 #include "Action/RCFunctionAction.h"
+#include "Action/RCPropertyIdAction.h"
 #include "Action/RCPropertyAction.h"
 
-#include "Behaviour/RCBehaviour.h"
 #include "Controller/RCController.h"
+#include "Behaviour/Builtin/Conditional/RCBehaviourConditional.h"
+#include "Behaviour/Builtin/RCBehaviourOnValueChangedNode.h"
+#include "Behaviour/RCBehaviour.h"
 
 #include "EdGraphSchema_K2.h"
 #include "EdGraph/EdGraphPin.h"
@@ -70,6 +74,8 @@ void SRCActionPanel::Construct(const FArguments& InArgs, const TSharedRef<SRemot
 	{
 		Preset->Layout.OnFieldAdded().AddSP(this, &SRCActionPanel::OnRemoteControlFieldAdded);
 		Preset->Layout.OnFieldDeleted().AddSP(this, &SRCActionPanel::OnRemoteControlFieldDeleted);
+		Preset->GetPropertyIdRegistry()->OnPropertyIdUpdated().AddLambda([this](){ bAddActionMenuNeedsRefresh = true; });
+		Preset->GetPropertyIdRegistry()->OnPropertyIdActionNeedsRefresh().AddLambda([this](){ bAddActionMenuNeedsRefresh = true; });
 	}
 }
 
@@ -224,10 +230,56 @@ TSharedRef<SWidget> SRCActionPanel::GetActionMenuContentWidget()
 		return MenuBuilder.MakeWidget();
 	}
 
+	if (const URCBehaviour* Behaviour = SelectedBehaviourItemWeakPtr.Pin()->GetBehaviour())
+	{
+		if (Behaviour->IsA<URCBehaviourConditional>() || Behaviour->IsA<URCBehaviourOnValueChangedNode>())
+		{
+			MenuBuilder.BeginSection(NAME_None, LOCTEXT("PropertyIdTitle", "PropertyId"));
+			// Create property identity entry
+			FUIAction PropertyIdAction(FExecuteAction::CreateSP(this, &SRCActionPanel::OnAddActionClicked));
+			MenuBuilder.AddMenuEntry(
+				LOCTEXT("AddAction", "Add PropertyId (Property)"),
+				LOCTEXT("AddActionTooltip", "Adds an PropertyId action of type property."),
+				FSlateIcon(),
+				MoveTemp(PropertyIdAction));
+
+			if (URemoteControlPreset* Preset = GetPreset())
+			{
+				if (Preset->GetPropertyIdRegistry())
+				{
+					TSet<FName> IdList = Preset->GetPropertyIdRegistry().Get()->GetFieldIdsNameList();
+					if (IdList.Num())
+					{
+						MenuBuilder.AddSubMenu(
+						LOCTEXT("AddActionSubMenu", "Add specific ID action"),
+						LOCTEXT("AddActionSubMenu_ToolTip", "Choose the ID based on the current list of different you have"),
+						FNewMenuDelegate::CreateLambda([this, IdList](FMenuBuilder& InMenuBuilder)
+						{
+							for (FName Id : IdList)
+							{
+								FText LabelName = FText::FromString(TEXT("ID: ") + Id.ToString());
+								const FText ToolTipName = LOCTEXT("AddAction_SpecificToolTip", "Create an action widget with this Id");
+								FUIAction PropertyIdAction(FExecuteAction::CreateSP(this, &SRCActionPanel::OnAddActionClicked, Id));
+								InMenuBuilder.AddMenuEntry(LabelName, ToolTipName, FSlateIcon(), PropertyIdAction);
+							}
+						}));
+					}
+				}
+			}
+			MenuBuilder.EndSection();
+		}
+	}
+
 	// List of exposed entities
 	if (URemoteControlPreset* Preset = GetPreset())
 	{
 		const TArray<TWeakPtr<FRemoteControlField>>& RemoteControlFields = Preset->GetExposedEntities<FRemoteControlField>();
+
+		if (!RemoteControlFields.IsEmpty())
+		{
+			MenuBuilder.BeginSection(NAME_None, LOCTEXT("FieldsTitle", "Fields"));
+		}
+
 		for (const TWeakPtr<FRemoteControlField>& RemoteControlFieldWeakPtr : RemoteControlFields)
 		{
 			if (const TSharedPtr<FRemoteControlField> RemoteControlField = RemoteControlFieldWeakPtr.Pin())
@@ -315,6 +367,26 @@ void SRCActionPanel::OnAddActionClicked(TSharedPtr<FRemoteControlField> InRemote
 	AddAction(InRemoteControlField.ToSharedRef());
 }
 
+void SRCActionPanel::OnAddActionClicked()
+{
+	if (!SelectedBehaviourItemWeakPtr.IsValid())
+	{
+		return;
+	}
+	FScopedTransaction Transaction(LOCTEXT("AddActionTransaction", "Add Action"));
+	AddAction();
+}
+
+void SRCActionPanel::OnAddActionClicked(FName InFieldId)
+{
+	if (!SelectedBehaviourItemWeakPtr.IsValid())
+	{
+		return;
+	}
+	FScopedTransaction Transaction(LOCTEXT("AddActionTransaction", "Add Action"));
+	AddAction(InFieldId);
+}
+
 FReply SRCActionPanel::OnClickEmptyButton()
 {
 	if (const TSharedPtr<FRCBehaviourModel> BehaviourItem = SelectedBehaviourItemWeakPtr.Pin())
@@ -527,6 +599,46 @@ FText SRCActionPanel::GetPasteItemMenuEntrySuffix()
 	return FText::GetEmpty();
 }
 
+URCAction* SRCActionPanel::AddAction()
+{
+	if (const TSharedPtr<FRCBehaviourModel> BehaviourItem = SelectedBehaviourItemWeakPtr.Pin())
+	{
+		if (const URCBehaviour* Behaviour = BehaviourItem->GetBehaviour())
+		{
+			Behaviour->ActionContainer->Modify();
+			URCAction* NewAction = BehaviourItem->AddAction();
+			AddNewActionToList(NewAction);
+			// Broadcast new Action to other panels
+			if (const TSharedPtr<SRemoteControlPanel> RemoteControlPanel = GetRemoteControlPanel())
+			{
+				RemoteControlPanel->OnActionAdded.Broadcast(NewAction);
+			}
+			return NewAction;
+		}
+	}
+	return nullptr;
+}
+
+URCAction* SRCActionPanel::AddAction(FName InFieldId)
+{
+	if (const TSharedPtr<FRCBehaviourModel> BehaviourItem = SelectedBehaviourItemWeakPtr.Pin())
+	{
+		if (const URCBehaviour* Behaviour = BehaviourItem->GetBehaviour())
+		{
+			Behaviour->ActionContainer->Modify();
+			URCAction* NewAction = BehaviourItem->AddAction(InFieldId);
+			AddNewActionToList(NewAction);
+			// Broadcast new Action to other panels
+			if (const TSharedPtr<SRemoteControlPanel> RemoteControlPanel = GetRemoteControlPanel())
+			{
+				RemoteControlPanel->OnActionAdded.Broadcast(NewAction);
+			}
+			return NewAction;
+		}
+	}
+	return nullptr;
+}
+
 FReply SRCActionPanel::RequestDeleteSelectedItem()
 {
 	if (!ActionPanelList.IsValid())
diff --git a/Engine/Plugins/VirtualProduction/RemoteControl/Source/RemoteControlUI/Private/UI/Action/SRCActionPanel.h b/Engine/Plugins/VirtualProduction/RemoteControl/Source/RemoteControlUI/Private/UI/Action/SRCActionPanel.h
--- a/Engine/Plugins/VirtualProduction/RemoteControl/Source/RemoteControlUI/Private/UI/Action/SRCActionPanel.h
+++ b/Engine/Plugins/VirtualProduction/RemoteControl/Source/RemoteControlUI/Private/UI/Action/SRCActionPanel.h
@@ -1,14 +1,10 @@
-﻿// Copyright Epic Games, Inc. All Rights Reserved.
+// Copyright Epic Games, Inc. All Rights Reserved.
 
 #pragma once
 
 #include "UI/BaseLogicUI/SRCLogicPanelBase.h"
 
-enum class ECheckBoxState : uint8;
-
-struct FRCPanelStyle;
 class FRCActionModel;
-struct FRemoteControlField;
 class FRCBehaviourModel;
 class SBox;
 class SCheckBox;
@@ -18,6 +14,11 @@ class SRemoteControlPanel;
 class URCAction;
 class URCBehaviour;
 
+struct FRCPanelStyle;
+struct FRemoteControlField;
+
+enum class ECheckBoxState : uint8;
+
 /*
 * ~ SRCActionPanel ~
 *
@@ -63,6 +64,12 @@ public:
 	/** Provides an item suffix for the Paste context menu to provide users with useful context on the nature of the item being pasted */
 	virtual FText GetPasteItemMenuEntrySuffix() override;
 
+	/** Adds an PropertyId Action for the currently active Behaviour and broadcasts to parent panels */
+	URCAction* AddAction();
+
+	/** Adds an PropertyId Action for the currently active Behaviour and broadcasts to parent panels */
+	URCAction* AddAction(FName InFieldId);
+
 	/** Adds an Action for the currently active Behaviour and broadcasts to parent panels */
 	URCAction* AddAction(const TSharedRef<const FRemoteControlField> InRemoteControlField);
 
@@ -118,6 +125,12 @@ private:
 	/** Handles click event for Add Action button*/
 	void OnAddActionClicked(TSharedPtr<FRemoteControlField> InRemoteControlField);
 
+	/** Handles click event for Add Action (PropertyId) button */
+	void OnAddActionClicked();
+
+	/** Handles click event for Add specific PropertyId Action buttons */
+	void OnAddActionClicked(FName InFieldId);
+
 	/** Handles click event for Empty button; clear all Actions from the panel*/
 	FReply OnClickEmptyButton();
 
diff --git a/Engine/Plugins/VirtualProduction/RemoteControl/Source/RemoteControlUI/Private/UI/Action/SRCActionPanelList.h b/Engine/Plugins/VirtualProduction/RemoteControl/Source/RemoteControlUI/Private/UI/Action/SRCActionPanelList.h
--- a/Engine/Plugins/VirtualProduction/RemoteControl/Source/RemoteControlUI/Private/UI/Action/SRCActionPanelList.h
+++ b/Engine/Plugins/VirtualProduction/RemoteControl/Source/RemoteControlUI/Private/UI/Action/SRCActionPanelList.h
@@ -1,4 +1,4 @@
-﻿// Copyright Epic Games, Inc. All Rights Reserved.
+// Copyright Epic Games, Inc. All Rights Reserved.
 
 #pragma once
 
@@ -95,6 +95,22 @@ public:
 		}
 
 		Reset();
+
+		if (URemoteControlPreset* Preset = RemoteControlPanel->GetPreset())
+		{
+			Preset->GetPropertyIdRegistry()->OnPropertyIdActionNeedsRefresh().AddSP(this, &SRCActionPanelList::Refresh);
+		}
+	}
+
+	virtual ~SRCActionPanelList()
+	{
+		if (ActionPanelWeakPtr.IsValid())
+		{
+			if (URemoteControlPreset* Preset = ActionPanelWeakPtr.Pin()->GetPreset())
+			{
+				Preset->GetPropertyIdRegistry()->OnPropertyIdActionNeedsRefresh().RemoveAll(this);
+			}
+		}
 	}
 
 	void OnActionsListModified()
@@ -266,6 +282,12 @@ private:
 
 		ListView->RebuildList();
 	}
+	
+	/** Refreshes the list */
+	void Refresh()
+	{
+		ListView->RebuildList();
+	}
 
 	/** Handles broadcasting of a successful remove item operation.*/
 	virtual void BroadcastOnItemRemoved() override {}
diff --git a/Engine/Plugins/VirtualProduction/RemoteControl/Source/RemoteControlUI/Private/UI/Behaviour/RCBehaviourModel.cpp b/Engine/Plugins/VirtualProduction/RemoteControl/Source/RemoteControlUI/Private/UI/Behaviour/RCBehaviourModel.cpp
--- a/Engine/Plugins/VirtualProduction/RemoteControl/Source/RemoteControlUI/Private/UI/Behaviour/RCBehaviourModel.cpp
+++ b/Engine/Plugins/VirtualProduction/RemoteControl/Source/RemoteControlUI/Private/UI/Behaviour/RCBehaviourModel.cpp
@@ -1,4 +1,4 @@
-﻿// Copyright Epic Games, Inc. All Rights Reserved.
+// Copyright Epic Games, Inc. All Rights Reserved.
 
 #include "RCBehaviourModel.h"
 
@@ -36,6 +36,28 @@ FRCBehaviourModel::FRCBehaviourModel(URCBehaviour* InBehaviour
 	}
 }
 
+URCAction* FRCBehaviourModel::AddAction()
+{
+	URCAction* NewAction = nullptr;
+	if (URCBehaviour* Behaviour = BehaviourWeakPtr.Get())
+	{
+		NewAction = Behaviour->AddAction();
+		OnActionAdded(NewAction);
+	}
+	return NewAction;
+}
+
+URCAction* FRCBehaviourModel::AddAction(FName InFieldId)
+{
+	URCAction* NewAction = nullptr;
+	if (URCBehaviour* Behaviour = BehaviourWeakPtr.Get())
+	{
+		NewAction = Behaviour->AddAction(InFieldId);
+		OnActionAdded(NewAction);
+	}
+	return NewAction;
+}
+
 URCAction* FRCBehaviourModel::AddAction(const TSharedRef<const FRemoteControlField> InRemoteControlField)
 {
 	URCAction* NewAction = nullptr;
diff --git a/Engine/Plugins/VirtualProduction/RemoteControl/Source/RemoteControlUI/Private/UI/Behaviour/RCBehaviourModel.h b/Engine/Plugins/VirtualProduction/RemoteControl/Source/RemoteControlUI/Private/UI/Behaviour/RCBehaviourModel.h
--- a/Engine/Plugins/VirtualProduction/RemoteControl/Source/RemoteControlUI/Private/UI/Behaviour/RCBehaviourModel.h
+++ b/Engine/Plugins/VirtualProduction/RemoteControl/Source/RemoteControlUI/Private/UI/Behaviour/RCBehaviourModel.h
@@ -1,16 +1,17 @@
-﻿// Copyright Epic Games, Inc. All Rights Reserved.
+// Copyright Epic Games, Inc. All Rights Reserved.
 
 #pragma once
+
 #include "Behaviour/RCBehaviour.h"
 #include "UI/BaseLogicUI/RCLogicModeBase.h"
 
-struct FRCPanelStyle;
 class IPropertyRowGenerator;
 class STextBlock;
 class SRCActionPanel;
 class SRCLogicPanelListBase;
 class SWidget;
 class URCBehaviour;
+struct FRCPanelStyle;
 
 /*
 * ~ FRCBehaviourModel ~
@@ -23,6 +24,12 @@ class FRCBehaviourModel : public FRCLogicModeBase
 public:
 	FRCBehaviourModel(URCBehaviour* InBehaviour, const TSharedPtr<SRemoteControlPanel> InRemoteControlPanel = nullptr);
 
+	/** Add a Logic Action using as an identity action. */
+	virtual URCAction* AddAction();
+
+	/** Add a Logic Action using as an identity action. */
+	virtual URCAction* AddAction(FName InFieldId);
+
 	/** Add a Logic Action using a remote control field as input */
 	virtual URCAction* AddAction(const TSharedRef<const FRemoteControlField> InRemoteControlField);
 
diff --git a/Engine/Plugins/VirtualProduction/RemoteControl/Source/RemoteControlUI/Private/UI/SRCPanelExposedEntitiesList.cpp b/Engine/Plugins/VirtualProduction/RemoteControl/Source/RemoteControlUI/Private/UI/SRCPanelExposedEntitiesList.cpp
--- a/Engine/Plugins/VirtualProduction/RemoteControl/Source/RemoteControlUI/Private/UI/SRCPanelExposedEntitiesList.cpp
+++ b/Engine/Plugins/VirtualProduction/RemoteControl/Source/RemoteControlUI/Private/UI/SRCPanelExposedEntitiesList.cpp
@@ -18,6 +18,7 @@
 #include "Misc/MessageDialog.h"
 #include "PropertyPath.h"
 #include "RCPanelWidgetRegistry.h"
+#include "RemoteControlPropertyIdRegistry.h"
 #include "RemoteControlPanelStyle.h"
 #include "RemoteControlPreset.h"
 #include "RemoteControlUIModule.h"
@@ -269,10 +270,16 @@ void SRCPanelExposedEntitiesList::Construct(const FArguments& InArgs, URemoteCon
 			.ShouldGenerateWidget(true)
 			.ShouldGenerateSubMenuEntry(false)
 
+			+ SRCHeaderRow::Column(RemoteControlPresetColumns::PropertyIdentifier)
+			.DefaultLabel(LOCTEXT("RCPresetLinkIdColumnHeader", "Property ID"))
+			.HAlignHeader(HAlign_Center)
+			.FillWidth(0.1f)
+			.HeaderContentPadding(RCPanelStyle->HeaderRowPadding)
+
 			+ SRCHeaderRow::Column(RemoteControlPresetColumns::OwnerName)
 			.DefaultLabel(LOCTEXT("RCPresetOwnerNameColumnHeader", "Owner Name"))
 			.HAlignHeader(HAlign_Center)
-			.FillWidth(0.15f)
+			.FillWidth(0.1f)
 			.HeaderContentPadding(RCPanelStyle->HeaderRowPadding)
 
 			+ SRCHeaderRow::Column(RemoteControlPresetColumns::SubobjectPath)
@@ -1452,7 +1459,7 @@ SHeaderRow::FColumn::FArguments SRCPanelExposedEntitiesList::CreateColumn(const
 
 int32 SRCPanelExposedEntitiesList::GetColumnIndex(const FName& ForColumn) const
 {
-	if (ForColumn == RemoteControlPresetColumns::LinkIdentifier)
+	if (ForColumn == RemoteControlPresetColumns::PropertyIdentifier)
 	{
 		return GetColumnIndex_Internal(ForColumn, RemoteControlPresetColumns::DragDropHandle, ERCColumn::ERC_After);
 	}
@@ -1530,9 +1537,9 @@ FText SRCPanelExposedEntitiesList::GetColumnLabel(const FName& ForColumn) const
 	{
 		return LOCTEXT("RCPresetBindingStatusColumnHeader", "REC");
 	}
-	else if (ForColumn == RemoteControlPresetColumns::LinkIdentifier)
+	else if (ForColumn == RemoteControlPresetColumns::PropertyIdentifier)
 	{
-		return LOCTEXT("RCPresetLinkIDColumnHeader", "Link ID");
+		return LOCTEXT("RCPresetLinkIDColumnHeader", "Property ID");
 	}
 	else if (ForColumn == RemoteControlPresetColumns::Mask)
 	{
@@ -1562,7 +1569,7 @@ float SRCPanelExposedEntitiesList::GetColumnSize(const FName ForColumn) const
 {
 	float ColumnSize = ProtocolColumnConstants::ColumnSizeNormal;
 
-	if (ForColumn == RemoteControlPresetColumns::LinkIdentifier)
+	if (ForColumn == RemoteControlPresetColumns::PropertyIdentifier)
 	{
 		ColumnSize = ProtocolColumnConstants::ColumnSizeMini;
 	}
@@ -1630,6 +1637,10 @@ void SRCPanelExposedEntitiesList::ProcessRefresh()
 
 	if (Preset.IsValid())
 	{
+		//Refresh during PostUndo/Redo to keep the action updated.
+		Preset->GetPropertyIdRegistry()->Initialize(Preset.Get());
+		Preset->GetPropertyIdRegistry()->OnPropertyIdUpdated().Broadcast();
+
 		constexpr bool bForceMouseClick = true;
 
 		if (const FRemoteControlPresetGroup* SelectedGroup = Preset->Layout.GetGroup(CurrentlySelectedGroup))
diff --git a/Engine/Plugins/VirtualProduction/RemoteControl/Source/RemoteControlUI/Private/UI/SRCPanelExposedEntity.cpp b/Engine/Plugins/VirtualProduction/RemoteControl/Source/RemoteControlUI/Private/UI/SRCPanelExposedEntity.cpp
--- a/Engine/Plugins/VirtualProduction/RemoteControl/Source/RemoteControlUI/Private/UI/SRCPanelExposedEntity.cpp
+++ b/Engine/Plugins/VirtualProduction/RemoteControl/Source/RemoteControlUI/Private/UI/SRCPanelExposedEntity.cpp
@@ -33,6 +33,7 @@
 #include "Widgets/Input/SButton.h"
 #include "Widgets/Input/SCheckBox.h"
 #include "Widgets/Input/SComboButton.h"
+#include "Widgets/Input/SEditableTextBox.h"
 #include "Widgets/Layout/SBorder.h"
 #include "Widgets/Layout/SBox.h"
 #include "Widgets/Notifications/SPopUpErrorText.h"
@@ -418,6 +419,33 @@ void SRCPanelExposedEntity::OnLabelCommitted(const FText& InLabel, ETextCommit::
 	}
 }
 
+FText SRCPanelExposedEntity::OnPropertyIdText() const
+{
+	if (const TSharedPtr<FRemoteControlEntity> RCEntity = GetEntity())
+	{
+		if (const TSharedPtr<FRemoteControlField> RCField = StaticCastSharedPtr<FRemoteControlField>(RCEntity))
+		{
+			return FText::FromName(RCField->PropertyId);
+		}
+	}
+	return FText::FromName(NAME_None);
+}
+
+void SRCPanelExposedEntity::OnPropertyIdTextCommitted(const FText& InText, ETextCommit::Type InCommitInfo) const
+{
+	if (URemoteControlPreset* RCPreset = Preset.Get())
+	{
+		if (const TSharedPtr<FRemoteControlEntity> RCEntity = GetEntity())
+		{
+			if (const TSharedPtr<FRemoteControlField> RCField = StaticCastSharedPtr<FRemoteControlField>(RCEntity))
+			{
+				RCField->PropertyId = FName(InText.ToString());
+				RCPreset->UpdateIdentifiedField(RCField.ToSharedRef());
+			}
+		}
+	}
+}
+
 void SRCPanelExposedEntity::OnActorSelected(AActor* InActor) const
 {
 	if (TSharedPtr<FRemoteControlEntity> Entity = GetEntity())
@@ -480,6 +508,17 @@ TSharedRef<SWidget> SRCPanelExposedEntity::CreateEntityWidget(TSharedPtr<SWidget
 			.Widget(Widget)
 		];
 
+	Args.PropertyIdWidget = SNew(SBox)
+		[
+			SNew(SEditableTextBox)
+			.MinDesiredWidth(50.f)
+			.SelectAllTextWhenFocused(true)
+			.RevertTextOnEscape(true)
+			.ClearKeyboardFocusOnCommit(true)
+			.Text(this, &SRCPanelExposedEntity::OnPropertyIdText)
+			.OnTextCommitted(this, &SRCPanelExposedEntity::OnPropertyIdTextCommitted)
+		];
+
 	const FSlateBrush* TrashBrush = FAppStyle::Get().GetBrush("Icons.Delete");
 
 	Args.OwnerNameWidget = SNew(SBox)
diff --git a/Engine/Plugins/VirtualProduction/RemoteControl/Source/RemoteControlUI/Private/UI/SRCPanelExposedEntity.h b/Engine/Plugins/VirtualProduction/RemoteControl/Source/RemoteControlUI/Private/UI/SRCPanelExposedEntity.h
--- a/Engine/Plugins/VirtualProduction/RemoteControl/Source/RemoteControlUI/Private/UI/SRCPanelExposedEntity.h
+++ b/Engine/Plugins/VirtualProduction/RemoteControl/Source/RemoteControlUI/Private/UI/SRCPanelExposedEntity.h
@@ -61,6 +61,8 @@ protected:
 	bool bValidBinding = false;
 	/** Display name of the entity. */
 	FName CachedLabel;
+	/** Display PropertyId of the entity. */
+	FName PropertyIdLabel;
 	/** Cached entity field path. */
 	FString CachedFieldPath;
 	/** Text to be highlighted while searching. */
@@ -93,6 +95,10 @@ private:
 	bool OnVerifyItemLabelChanged(const FText& InLabel, FText& OutErrorMessage);
 	/** Handles committing a entity label. */
 	void OnLabelCommitted(const FText& InLabel, ETextCommit::Type InCommitInfo);
+	/** Handles text change of a property id label. */
+	FText OnPropertyIdText() const;
+	/** Handles committing a property id label. */
+	void OnPropertyIdTextCommitted(const FText& InText, ETextCommit::Type InCommitInfo) const;
 	/** Returns whether or not the actor is selectable for a binding replacement. */
 	bool IsActorSelectable(const AActor* Parent) const;
 	/** Handle clicking on the unexpose button. */
diff --git a/Engine/Plugins/VirtualProduction/RemoteControl/Source/RemoteControlUI/Private/UI/SRCPanelTreeNode.cpp b/Engine/Plugins/VirtualProduction/RemoteControl/Source/RemoteControlUI/Private/UI/SRCPanelTreeNode.cpp
--- a/Engine/Plugins/VirtualProduction/RemoteControl/Source/RemoteControlUI/Private/UI/SRCPanelTreeNode.cpp
+++ b/Engine/Plugins/VirtualProduction/RemoteControl/Source/RemoteControlUI/Private/UI/SRCPanelTreeNode.cpp
@@ -15,6 +15,7 @@
 
 TSet<FName> SRCPanelTreeNode::DefaultColumns = {
 	RemoteControlPresetColumns::DragDropHandle,
+	RemoteControlPresetColumns::PropertyIdentifier,
 	RemoteControlPresetColumns::OwnerName,
 	RemoteControlPresetColumns::SubobjectPath,
 	RemoteControlPresetColumns::Description,
@@ -54,6 +55,10 @@ TSharedRef<SWidget> SRCPanelTreeNode::GetWidget(const FName ForColumnName, const
 	{
 		return DragHandleWidget.ToSharedRef();
 	}
+	else if (ForColumnName == RemoteControlPresetColumns::PropertyIdentifier)
+	{
+		return PropertyIdWidget.ToSharedRef();
+	}
 	else if (ForColumnName == RemoteControlPresetColumns::OwnerName)
 	{
 		return NodeOwnerNameWidget.ToSharedRef();
@@ -117,6 +122,13 @@ TSharedRef<SWidget> SRCPanelTreeNode::MakeNodeWidget(const FMakeNodeWidgetArgs&
 		[
 			DragHandleWidget.ToSharedRef()
 		]
+		// Link identifier widget
+		+ SHorizontalBox::Slot()
+		.VAlign(VAlign_Center)
+		.AutoWidth()
+		[
+			PropertyIdWidget.ToSharedRef()
+		]
 		// Owner name
 		+ SHorizontalBox::Slot()
 		.VAlign(VAlign_Center)
@@ -169,6 +181,8 @@ void SRCPanelTreeNode::MakeNodeWidgets(const FMakeNodeWidgetArgs& Args)
 
 	DragHandleWidget = WidgetOrNull(Args.DragHandle);
 
+	PropertyIdWidget = WidgetOrNull(Args.PropertyIdWidget);
+
 	NodeOwnerNameWidget = WidgetOrNull(Args.OwnerNameWidget);
 	
 	SubobjectPathWidget = WidgetOrNull(Args.SubObjectPathWidget);
diff --git a/Engine/Plugins/VirtualProduction/RemoteControl/Source/RemoteControlUI/Private/UI/SRCPanelTreeNode.h b/Engine/Plugins/VirtualProduction/RemoteControl/Source/RemoteControlUI/Private/UI/SRCPanelTreeNode.h
--- a/Engine/Plugins/VirtualProduction/RemoteControl/Source/RemoteControlUI/Private/UI/SRCPanelTreeNode.h
+++ b/Engine/Plugins/VirtualProduction/RemoteControl/Source/RemoteControlUI/Private/UI/SRCPanelTreeNode.h
@@ -18,7 +18,7 @@ struct SRCPanelExposedMaterial;
 namespace RemoteControlPresetColumns
 {
 	static FName DragDropHandle = TEXT("DragDropHandle");
-	static FName LinkIdentifier = TEXT("LinkID");
+	static FName PropertyIdentifier = TEXT("PropertyID");
 	static FName OwnerName = TEXT("OwnerName");
 	static FName SubobjectPath = TEXT("Subobject Path");
 	static FName Description = TEXT("Description");
@@ -89,6 +89,7 @@ protected:
 	struct FMakeNodeWidgetArgs
 	{
 		TSharedPtr<SWidget> DragHandle;
+		TSharedPtr<SWidget> PropertyIdWidget;
 		TSharedPtr<SWidget> OwnerNameWidget;
 		TSharedPtr<SWidget> SubObjectPathWidget;
 		TSharedPtr<SWidget> NameWidget;
@@ -126,6 +127,9 @@ private:
 	/** Cached widget of drag handle. */
 	TSharedPtr<SWidget> DragHandleWidget;
 
+	/** Cached widget of PropertyId identifier. */
+	TSharedPtr<SWidget> PropertyIdWidget;
+
 	/** Cached widget of node owner name. */
 	TSharedPtr<SWidget> NodeOwnerNameWidget;
 	
