[RemoteControl]
Fix Issue 369116: Warning C6011 Dereferencing NULL pointer "InPropertyIdAction"

diff --git a/Engine/Plugins/VirtualProduction/RemoteControl/Source/RemoteControlUI/Private/UI/Action/RCActionModel.cpp b/Engine/Plugins/VirtualProduction/RemoteControl/Source/RemoteControlUI/Private/UI/Action/RCActionModel.cpp
--- a/Engine/Plugins/VirtualProduction/RemoteControl/Source/RemoteControlUI/Private/UI/Action/RCActionModel.cpp
+++ b/Engine/Plugins/VirtualProduction/RemoteControl/Source/RemoteControlUI/Private/UI/Action/RCActionModel.cpp
@@ -305,13 +305,13 @@ FRCPropertyIdActionType::FRCPropertyIdActionType(URCPropertyIdAction* InProperty
 	{
 		RefreshNameWidget();
 		RefreshValueWidget();
-	}
 #if WITH_EDITOR
-	if (URemoteControlPreset* Preset = InPropertyIdAction->PresetWeakPtr.Get())
-	{
-		Preset->GetPropertyIdRegistry()->OnPropertyIdActionNeedsRefresh().AddRaw(this, &FRCPropertyIdActionType::RefreshValueWidget);
-	}
+		if (const URemoteControlPreset* Preset = InPropertyIdAction->PresetWeakPtr.Get())
+		{
+			Preset->GetPropertyIdRegistry()->OnPropertyIdActionNeedsRefresh().AddRaw(this, &FRCPropertyIdActionType::RefreshValueWidget);
+		}
 #endif
+	}
 }
 
 FRCPropertyIdActionType::~FRCPropertyIdActionType()
