[RemoteControl]
Fix Horde Issue 369338 Localization warning

diff --git a/Engine/Plugins/VirtualProduction/RemoteControl/Source/RemoteControlUI/Private/UI/Action/Conditional/RCActionConditionalModel.cpp b/Engine/Plugins/VirtualProduction/RemoteControl/Source/RemoteControlUI/Private/UI/Action/Conditional/RCActionConditionalModel.cpp
--- a/Engine/Plugins/VirtualProduction/RemoteControl/Source/RemoteControlUI/Private/UI/Action/Conditional/RCActionConditionalModel.cpp
+++ b/Engine/Plugins/VirtualProduction/RemoteControl/Source/RemoteControlUI/Private/UI/Action/Conditional/RCActionConditionalModel.cpp
@@ -22,7 +22,7 @@
 #include "Widgets/Views/SListView.h"
 #include "Widgets/Views/STableRow.h"
 
-#define LOCTEXT_NAMESPACE "RCActionModel"
+#define LOCTEXT_NAMESPACE "RCActionConditionalModel"
 
 namespace UE::RCActionPanelConditionalList
 {
diff --git a/Engine/Plugins/VirtualProduction/RemoteControl/Source/RemoteControlUI/Private/UI/Action/RangeMap/RCActionRangeMapModel.cpp b/Engine/Plugins/VirtualProduction/RemoteControl/Source/RemoteControlUI/Private/UI/Action/RangeMap/RCActionRangeMapModel.cpp
--- a/Engine/Plugins/VirtualProduction/RemoteControl/Source/RemoteControlUI/Private/UI/Action/RangeMap/RCActionRangeMapModel.cpp
+++ b/Engine/Plugins/VirtualProduction/RemoteControl/Source/RemoteControlUI/Private/UI/Action/RangeMap/RCActionRangeMapModel.cpp
@@ -22,7 +22,7 @@
 #include "Widgets/Views/SListView.h"
 #include "Widgets/Views/STableRow.h"
 
-#define LOCTEXT_NAMESPACE "RCActionModel"
+#define LOCTEXT_NAMESPACE "RCActionRangeMapModel"
 
 namespace UE::RCActionPanelRangeMapList
 {
diff --git a/Engine/Plugins/VirtualProduction/RemoteControl/Source/RemoteControlUI/Private/UI/Action/SRCActionPanel.cpp b/Engine/Plugins/VirtualProduction/RemoteControl/Source/RemoteControlUI/Private/UI/Action/SRCActionPanel.cpp
--- a/Engine/Plugins/VirtualProduction/RemoteControl/Source/RemoteControlUI/Private/UI/Action/SRCActionPanel.cpp
+++ b/Engine/Plugins/VirtualProduction/RemoteControl/Source/RemoteControlUI/Private/UI/Action/SRCActionPanel.cpp
@@ -238,8 +238,8 @@ TSharedRef<SWidget> SRCActionPanel::GetActionMenuContentWidget()
 			// Create property identity entry
 			FUIAction PropertyIdAction(FExecuteAction::CreateSP(this, &SRCActionPanel::OnAddActionClicked));
 			MenuBuilder.AddMenuEntry(
-				LOCTEXT("AddAction", "Add PropertyId (Property)"),
-				LOCTEXT("AddActionTooltip", "Adds an PropertyId action of type property."),
+				LOCTEXT("AddPropertyIdAction", "Add PropertyId (Property)"),
+				LOCTEXT("AddPropertyIdAction_Tooltip", "Add a PropertyId action."),
 				FSlateIcon(),
 				MoveTemp(PropertyIdAction));
 
diff --git a/Engine/Plugins/VirtualProduction/RemoteControl/Source/RemoteControlUI/Private/UI/SRCPanelExposedEntitiesList.cpp b/Engine/Plugins/VirtualProduction/RemoteControl/Source/RemoteControlUI/Private/UI/SRCPanelExposedEntitiesList.cpp
--- a/Engine/Plugins/VirtualProduction/RemoteControl/Source/RemoteControlUI/Private/UI/SRCPanelExposedEntitiesList.cpp
+++ b/Engine/Plugins/VirtualProduction/RemoteControl/Source/RemoteControlUI/Private/UI/SRCPanelExposedEntitiesList.cpp
@@ -271,7 +271,7 @@ void SRCPanelExposedEntitiesList::Construct(const FArguments& InArgs, URemoteCon
 			.ShouldGenerateSubMenuEntry(false)
 
 			+ SRCHeaderRow::Column(RemoteControlPresetColumns::PropertyIdentifier)
-			.DefaultLabel(LOCTEXT("RCPresetLinkIdColumnHeader", "Property ID"))
+			.DefaultLabel(LOCTEXT("RCPresetPropertyIdColumnHeader", "Property ID"))
 			.HAlignHeader(HAlign_Center)
 			.FillWidth(0.1f)
 			.HeaderContentPadding(RCPanelStyle->HeaderRowPadding)
@@ -1539,7 +1539,7 @@ FText SRCPanelExposedEntitiesList::GetColumnLabel(const FName& ForColumn) const
 	}
 	else if (ForColumn == RemoteControlPresetColumns::PropertyIdentifier)
 	{
-		return LOCTEXT("RCPresetLinkIDColumnHeader", "Property ID");
+		return LOCTEXT("RCPresetPropertyIdColumnHeader_Label", "Property ID");
 	}
 	else if (ForColumn == RemoteControlPresetColumns::Mask)
 	{
