[RemoteControl]
Removed wrong return in the URemoteControlPropertyIdRegistry::PerformChainReaction, causing the Property to not update

diff --git a/Engine/Plugins/VirtualProduction/RemoteControl/Source/RemoteControl/Private/RemoteControlPropertyIdRegistry.cpp b/Engine/Plugins/VirtualProduction/RemoteControl/Source/RemoteControl/Private/RemoteControlPropertyIdRegistry.cpp
--- a/Engine/Plugins/VirtualProduction/RemoteControl/Source/RemoteControl/Private/RemoteControlPropertyIdRegistry.cpp
+++ b/Engine/Plugins/VirtualProduction/RemoteControl/Source/RemoteControl/Private/RemoteControlPropertyIdRegistry.cpp
@@ -293,16 +293,22 @@ void URemoteControlPropertyIdRegistry::PerformChainReaction(const FRemoteControl
 							{
 								if (const FStructProperty* StructProperty = CastField<FStructProperty>(Property))
 								{
+									// FLinearColor are treated as FColor to avoid creating 2 different widget for them and avoid confusion
+									// Instead of the normal CopyCompleteValue we use this for FLinearColor
 									if (StructProperty->Struct->GetFName() == NAME_LinearColor)
 									{
 										FColor ColorValue;
 										InArgs.VirtualProperty->GetValueColor(ColorValue);
 										const FLinearColor RealValue(ColorValue);
 										Property->CopyCompleteValue(PropertyValuePtr, &RealValue);
-										return;
+										bCopyComplete = true;
 									}
 								}
-								bCopyComplete = InArgs.VirtualProperty->CopyCompleteValue(Property, PropertyValuePtr, true);
+								// We do this to avoid copying it above and here
+								if (bCopyComplete == false)
+								{
+									bCopyComplete = InArgs.VirtualProperty->CopyCompleteValue(Property, PropertyValuePtr, true);
+								}
 							}
 						}
 						//Note : For all the other cases, Container and UStruct Owner.
@@ -581,6 +587,9 @@ bool URemoteControlPropertyIdRegistry::TryCopyNonUClassOwnerProperty(const TObje
 			{
 				return TryCopyNonUClassOwnerProperty(InVirtualPropertySelfContainer, TargetRCProperty, InProperty, BoundObject, Segment, IndexSegment + 1, InnerStructPtrContainer, bCheckByteEnumComparison);
 			}
+
+			// FLinearColor are treated as FColor to avoid creating 2 different widget for them and avoid confusion
+			// Instead of the normal CopyCompleteValue we use this for FLinearColor
 			if (const FStructProperty* StructRealProperty = CastField<FStructProperty>(InProperty))
 			{
 #if !WITH_EDITOR
