[Remote Control] Store as config the current settings of the RC Panel Tab

Settings added to config:

- Logic panel shown/hidden
- Exposed entities hidden columns
- Current mode (e.g. Setup/Operation)
- Action panel splitter position (split ratio)

diff --git a/Engine/Plugins/VirtualProduction/RemoteControl/Source/RemoteControlCommon/Public/RemoteControlSettings.h b/Engine/Plugins/VirtualProduction/RemoteControl/Source/RemoteControlCommon/Public/RemoteControlSettings.h
--- a/Engine/Plugins/VirtualProduction/RemoteControl/Source/RemoteControlCommon/Public/RemoteControlSettings.h
+++ b/Engine/Plugins/VirtualProduction/RemoteControl/Source/RemoteControlCommon/Public/RemoteControlSettings.h
@@ -255,6 +255,18 @@ public:
 		return false;
 	}
 
+	/** Returns the names of the columns visualized in the RC Panel Entities list. */
+	static const TSet<FName>& GetExposedEntitiesColumnNames()
+	{
+		static TSet<FName> ColumnNames =
+			{ TEXT("OwnerName")
+			, TEXT("Subobject Path")
+			, TEXT("Description")
+			, TEXT("Value")};
+
+		return ColumnNames;
+	}
+
 	/**
 	 * Should transactions be generated for events received through protocols (ie. MIDI, DMX etc.)
 	 * Disabling transactions improves performance but will prevent events from being transacted to Multi-User
@@ -307,6 +319,10 @@ public:
 	UPROPERTY(config)
 	float TreeBindingSplitRatio = 0.7;
 
+	/** The split widget control ratio between entity tree and details/protocol binding list. */
+	UPROPERTY(config)
+	float ActionPanelSplitRatio = 0.6;
+
 	UPROPERTY(config)
 	bool bUseRebindingContext = true;
 
@@ -353,6 +369,15 @@ public:
 	UPROPERTY(Config, EditAnywhere, Category = "Remote Control | Security", DisplayName = "Warn that Passphrase might be disabled ", meta=(EditCondition=bRestrictServerAccess))
 	bool bShowPassphraseDisabledWarning = true;
 
+	UPROPERTY(Config, EditAnywhere, Category = "Remote Control Preset", DisplayName = "Hide or Show the Logic panel by default")
+	bool bLogicPanelVisibility = false;
+
+	UPROPERTY(Config, EditAnywhere, Category = "Remote Control Preset", DisplayName = "Which columns to hide in the exposed entities list")
+	TSet<FName> EntitiesListHiddenColumns;
+	
+	UPROPERTY(Config, EditAnywhere, Category = "Remote Control Preset", DisplayName = "Default mode (e.g. Setup/Operation)")
+	FName DefaultPanelMode = TEXT("Setup");
+	
 private:
 	UPROPERTY(config)
     bool bSecuritySettingsReviewed = false;
diff --git a/Engine/Plugins/VirtualProduction/RemoteControl/Source/RemoteControlUI/Private/UI/Panels/SRCDockPanel.cpp b/Engine/Plugins/VirtualProduction/RemoteControl/Source/RemoteControlUI/Private/UI/Panels/SRCDockPanel.cpp
--- a/Engine/Plugins/VirtualProduction/RemoteControl/Source/RemoteControlUI/Private/UI/Panels/SRCDockPanel.cpp
+++ b/Engine/Plugins/VirtualProduction/RemoteControl/Source/RemoteControlUI/Private/UI/Panels/SRCDockPanel.cpp
@@ -10,7 +10,6 @@
 
 #include "Widgets/Layout/SBox.h"
 #include "Widgets/Layout/SSpacer.h"
-#include "Widgets/Layout/SSplitter.h"
 
 #include "Widgets/SBoxPanel.h"
 #include "Widgets/SWidget.h"
@@ -74,6 +73,10 @@ void SRCMajorPanel::Construct(const SRCMajorPanel::FArguments& InArgs)
 				.Orientation(InArgs._ChildOrientation)
 				.HitDetectionSplitterHandleSize(SplitterHandleSize)
 				.PhysicalSplitterHandleSize(SplitterHandleSize)
+				.OnSplitterFinishedResizing_Lambda([this]()
+				{
+					OnSplitterFinishedResizing().ExecuteIfBound();
+				})
 			]
 		]
 	);
@@ -142,9 +145,9 @@ void SRCMajorPanel::Construct(const SRCMajorPanel::FArguments& InArgs)
 
 END_SLATE_FUNCTION_BUILD_OPTIMIZATION
 
-void SRCMajorPanel::AddPanel(TSharedRef<SWidget> InContent, const float InDesiredSize, const bool bResizable/* = true */)
+int32 SRCMajorPanel::AddPanel(TSharedRef<SWidget> InContent, const TAttribute<float> InDesiredSize, const bool bResizable/* = true */)
 {
-	const SSplitter::ESizeRule SizeRule = InDesiredSize == 0.f ? SSplitter::ESizeRule::SizeToContent : SSplitter::ESizeRule::FractionOfParent;
+	const SSplitter::ESizeRule SizeRule = InDesiredSize.Get() == 0.f ? SSplitter::ESizeRule::SizeToContent : SSplitter::ESizeRule::FractionOfParent;
 
 	if (Children.IsValid())
 	{
@@ -155,7 +158,12 @@ void SRCMajorPanel::AddPanel(TSharedRef<SWidget> InContent, const float InDesire
 			[
 				InContent
 			];
+
+		const int32 SlotsCount = Children->GetChildren()->Num();
+		return SlotsCount - 1;
 	}
+
+	return -1;
 }
 
 const TSharedRef<SWidget>& SRCMajorPanel::GetContent() const
@@ -201,6 +209,11 @@ void SRCMajorPanel::AddFooterToolItem(EToolbar InToolbar, TSharedRef<SWidget> In
 	}
 }
 
+SSplitter::FSlot& SRCMajorPanel::GetSplitterSlotAt(const int32 InIndex) const
+{
+	return Children->SlotAt(InIndex);
+}
+
 SLATE_IMPLEMENT_WIDGET(SRCMinorPanel)
 
 void SRCMinorPanel::PrivateRegisterAttributes(FSlateAttributeInitializer& AttributeInitializer)
diff --git a/Engine/Plugins/VirtualProduction/RemoteControl/Source/RemoteControlUI/Private/UI/Panels/SRCDockPanel.h b/Engine/Plugins/VirtualProduction/RemoteControl/Source/RemoteControlUI/Private/UI/Panels/SRCDockPanel.h
--- a/Engine/Plugins/VirtualProduction/RemoteControl/Source/RemoteControlUI/Private/UI/Panels/SRCDockPanel.h
+++ b/Engine/Plugins/VirtualProduction/RemoteControl/Source/RemoteControlUI/Private/UI/Panels/SRCDockPanel.h
@@ -4,10 +4,10 @@
 
 #include "CoreMinimal.h"
 #include "Widgets/Layout/SBorder.h"
+#include "Widgets/Layout/SSplitter.h"
 
 struct FRCPanelStyle;
 class SHorizontalBox;
-class SSplitter;
 class SWidget;
 
 enum EToolbar
@@ -90,8 +90,11 @@ public:
 	 * Adds the given panel to the children.
 	 *
 	 * @param InPanel The panel widget.
+	 * @param InDesiredSize The desired size for the Panel widget
+	 * @param bResizable Can the slot be resize by the user
+	 * @return the index of the newly added slot. Slot can be accessed using this index by calling SRCMajorPanel::GetSplitterSlotAt(Index)
 	 */
-	virtual void AddPanel(TSharedRef<SWidget> InPanel, const float InDesiredSize, const bool bResizable = true);
+	virtual int32 AddPanel(TSharedRef<SWidget> InPanel, const TAttribute<float> InDesiredSize, const bool bResizable = true);
 
 	/**
 	 * Gets the content of this panel
@@ -105,6 +108,12 @@ public:
 
 	/** Adds widgets to the specified footer toolbar (Left or Right). */
 	void AddFooterToolItem(EToolbar InToolbar, TSharedRef<SWidget> InWidget);
+
+	/** Delegate called whenever the splitter is finished resizing */
+	FSimpleDelegate& OnSplitterFinishedResizing() { return OnSplitterFinishedResizingDelegate; }
+
+	/** Returns the slot at the specified index */
+	SSplitter::FSlot& GetSplitterSlotAt(const int32 InIndex) const;
 	
 private:
 
@@ -116,6 +125,9 @@ private:
 
 	/** Holds the entire child panels. */
 	TSharedPtr<SSplitter> Children;
+
+	/** Delegate called whenever the splitter is finished resizing */
+	FSimpleDelegate OnSplitterFinishedResizingDelegate;
 };
 
 /**
diff --git a/Engine/Plugins/VirtualProduction/RemoteControl/Source/RemoteControlUI/Private/UI/SRCHeaderRow.cpp b/Engine/Plugins/VirtualProduction/RemoteControl/Source/RemoteControlUI/Private/UI/SRCHeaderRow.cpp
--- a/Engine/Plugins/VirtualProduction/RemoteControl/Source/RemoteControlUI/Private/UI/SRCHeaderRow.cpp
+++ b/Engine/Plugins/VirtualProduction/RemoteControl/Source/RemoteControlUI/Private/UI/SRCHeaderRow.cpp
@@ -4,6 +4,7 @@
 #include "Framework/Application/SlateApplication.h"
 #include "Framework/MultiBox/MultiBoxBuilder.h"
 #include "Layout/WidgetPath.h"
+#include "RemoteControlSettings.h"
 
 void SRCHeaderRow::Construct(const FArguments& InArgs)
 {
@@ -12,6 +13,7 @@ void SRCHeaderRow::Construct(const FArguments& InArgs)
 	TArray<SHeaderRow::FColumn*> Slots;
 	for (FColumn* const Column : InArgs.Slots)
 	{
+		Column->bIsVisible = !InArgs._HiddenColumnsList.Contains(Column->ColumnId);
 		Columns.Add(Column);
 		SHeaderRow::FColumn* HeaderRowSlot = CreateHeaderRowColumn(*Column);
 		Slots.Add(HeaderRowSlot);
@@ -21,6 +23,7 @@ void SRCHeaderRow::Construct(const FArguments& InArgs)
 	HeaderRowArgs._Style = InArgs._Style;
 	HeaderRowArgs.Slots = Slots;
 	HeaderRowArgs._CanSelectGeneratedColumn = InArgs._CanSelectGeneratedColumn;
+	HeaderRowArgs._HiddenColumnsList = InArgs._HiddenColumnsList;
 	SHeaderRow::Construct(HeaderRowArgs);
 }
 
@@ -118,9 +121,39 @@ void SRCHeaderRow::ToggleOverriddenGeneratedColumn(FName ColumnId)
 			break;
 		}
 	}
+
+	StoreHiddenColumnsSettings();
 }
 
 ECheckBoxState SRCHeaderRow::GetOverriddenGeneratedColumnCheckedState(FName ColumnId) const
 {
 	return IsColumnGenerated(ColumnId) ? ECheckBoxState::Checked : ECheckBoxState::Unchecked;
 }
+
+void SRCHeaderRow::StoreHiddenColumnsSettings()
+{
+	URemoteControlSettings* Settings = GetMutableDefault<URemoteControlSettings>();
+	
+	if (!Settings)
+	{
+		return;
+	}
+
+	TSet<FName> HiddenColumns;
+	
+	for (const FColumn& SomeColumn : Columns)
+	{
+		if (!SomeColumn.bIsVisible)
+		{
+			const FName& ColumnName = SomeColumn.ColumnId;
+			if (URemoteControlSettings::GetExposedEntitiesColumnNames().Contains(ColumnName))
+			{
+				HiddenColumns.Add(ColumnName);
+			}
+		}
+	}
+	
+	Settings->EntitiesListHiddenColumns = HiddenColumns;
+	Settings->PostEditChange();
+	Settings->SaveConfig();
+}
diff --git a/Engine/Plugins/VirtualProduction/RemoteControl/Source/RemoteControlUI/Private/UI/SRCHeaderRow.h b/Engine/Plugins/VirtualProduction/RemoteControl/Source/RemoteControlUI/Private/UI/SRCHeaderRow.h
--- a/Engine/Plugins/VirtualProduction/RemoteControl/Source/RemoteControlUI/Private/UI/SRCHeaderRow.h
+++ b/Engine/Plugins/VirtualProduction/RemoteControl/Source/RemoteControlUI/Private/UI/SRCHeaderRow.h
@@ -102,11 +102,15 @@ public:
 
 	SLATE_BEGIN_ARGS(SRCHeaderRow)
 		: _CanSelectGeneratedColumn(false)
+		, _HiddenColumnsList({})
 	{}
 
 		SLATE_STYLE_ARGUMENT(FHeaderRowStyle, Style)
 		SLATE_SUPPORTS_SLOT_WITH_ARGS(FColumn)
 		SLATE_ARGUMENT(bool, CanSelectGeneratedColumn)
+	
+		/** Add here the columns to be hidden by default. They can still be re-enabled via right click context-menu in the header row. */
+		SLATE_ARGUMENT(TArray<FName>, HiddenColumnsList)
 
 	SLATE_END_ARGS()
 
@@ -146,6 +150,9 @@ private:
 	 */
 	ECheckBoxState GetOverriddenGeneratedColumnCheckedState(FName ColumnId) const;
 
+	/** Save the current settings for shown/hidden columns to Remote Control Config */
+	void StoreHiddenColumnsSettings();
+
 private:
 
 	/**
diff --git a/Engine/Plugins/VirtualProduction/RemoteControl/Source/RemoteControlUI/Private/UI/SRCModeSwitcher.cpp b/Engine/Plugins/VirtualProduction/RemoteControl/Source/RemoteControlUI/Private/UI/SRCModeSwitcher.cpp
--- a/Engine/Plugins/VirtualProduction/RemoteControl/Source/RemoteControlUI/Private/UI/SRCModeSwitcher.cpp
+++ b/Engine/Plugins/VirtualProduction/RemoteControl/Source/RemoteControlUI/Private/UI/SRCModeSwitcher.cpp
@@ -31,16 +31,30 @@ void SRCModeSwitcher::Construct(const FArguments& InArgs)
 		.BorderImage(&RCPanelStyle->ContentAreaBrushDark)
 	);
 
+	FRCMode* ActiveModeElem = nullptr;
+	
 	// Copy all the mode info from the declaration
 	for (FRCMode* const Mode : InArgs.Slots)
 	{
 		RCModes.Add(*Mode);
 
 		Mode->bIsVisible = true;
+
+		if (Mode->ModeId == ActiveMode)
+		{
+			ActiveModeElem = Mode;
+		}
 	}
 
 	// Generate widgets for all modes
 	RegenerateWidgets();
+
+	// Force-enable the default mode
+	if (ActiveModeElem)
+	{
+		constexpr bool bForceModeEnable = true;
+		SetModeEnabled_Internal(ECheckBoxState::Checked, ActiveModeElem, bForceModeEnable);
+	}
 }
 
 ECheckBoxState SRCModeSwitcher::IsModeEnabled(FRCMode* InMode) const
@@ -170,10 +184,15 @@ void SRCModeSwitcher::RegenerateWidgets()
 }
 
 void SRCModeSwitcher::SetModeEnabled(ECheckBoxState NewState, FRCMode* NewMode)
+{
+	SetModeEnabled_Internal(NewState, NewMode);
+}
+
+void SRCModeSwitcher::SetModeEnabled_Internal(ECheckBoxState NewState, FRCMode* NewMode, bool bInForceModeEnabled /* false */)
 {
 	check(NewMode);
 
-	if (ActiveMode != NewMode->ModeId)
+	if (ActiveMode != NewMode->ModeId || bInForceModeEnabled)
 	{
 		if (NewState == ECheckBoxState::Checked)
 		{
diff --git a/Engine/Plugins/VirtualProduction/RemoteControl/Source/RemoteControlUI/Private/UI/SRCModeSwitcher.h b/Engine/Plugins/VirtualProduction/RemoteControl/Source/RemoteControlUI/Private/UI/SRCModeSwitcher.h
--- a/Engine/Plugins/VirtualProduction/RemoteControl/Source/RemoteControlUI/Private/UI/SRCModeSwitcher.h
+++ b/Engine/Plugins/VirtualProduction/RemoteControl/Source/RemoteControlUI/Private/UI/SRCModeSwitcher.h
@@ -230,6 +230,12 @@ private:
 
 private:
 
+	/**
+	 * Sets the given mode as active. Also allows to force the operation.
+	 * Useful e.g. when Widget is being created and the current mode already matches the required one.
+	 */
+	void SetModeEnabled_Internal(ECheckBoxState NewState, FRCMode* NewMode, bool bInForceModeEnabled = false);
+
 	/** Holdes the active mode. */
 	FName ActiveMode;
 	
diff --git a/Engine/Plugins/VirtualProduction/RemoteControl/Source/RemoteControlUI/Private/UI/SRCPanelExposedEntitiesList.cpp b/Engine/Plugins/VirtualProduction/RemoteControl/Source/RemoteControlUI/Private/UI/SRCPanelExposedEntitiesList.cpp
--- a/Engine/Plugins/VirtualProduction/RemoteControl/Source/RemoteControlUI/Private/UI/SRCPanelExposedEntitiesList.cpp
+++ b/Engine/Plugins/VirtualProduction/RemoteControl/Source/RemoteControlUI/Private/UI/SRCPanelExposedEntitiesList.cpp
@@ -21,6 +21,7 @@
 #include "RemoteControlPropertyIdRegistry.h"
 #include "RemoteControlPanelStyle.h"
 #include "RemoteControlPreset.h"
+#include "RemoteControlSettings.h"
 #include "RemoteControlUIModule.h"
 #include "ScopedTransaction.h"
 #include "SRCHeaderRow.h"
@@ -191,6 +192,10 @@ void SRCPanelExposedEntitiesList::Construct(const FArguments& InArgs, URemoteCon
 	RCPanelStyle = &FRemoteControlPanelStyle::Get()->GetWidgetStyle<FRCPanelStyle>("RemoteControlPanel.MinorPanel");
 	ActiveListMode = EEntitiesListMode::Default;
 
+	// Retrieve from settings which columns should be hidden
+	const URemoteControlSettings* Settings = GetMutableDefault<URemoteControlSettings>();
+	const TArray<FName>& HiddenColumns = Settings->EntitiesListHiddenColumns.Array();
+
 	// The Default group will be selected by default.
 	if (Preset.IsValid())
 	{
@@ -262,6 +267,7 @@ void SRCPanelExposedEntitiesList::Construct(const FArguments& InArgs, URemoteCon
 			SAssignNew(FieldsHeaderRow, SRCHeaderRow)
 			.Style(&RCPanelStyle->HeaderRowStyle)
 			.CanSelectGeneratedColumn(true) //To show/hide columns
+			.HiddenColumnsList(HiddenColumns) // List of columns to hide by default. User can un-hide via context menu list
 
 			+ SRCHeaderRow::Column(RemoteControlPresetColumns::DragDropHandle)
 			.DefaultLabel(LOCTEXT("RCPresetDragDropHandleColumnHeader", ""))
diff --git a/Engine/Plugins/VirtualProduction/RemoteControl/Source/RemoteControlUI/Private/UI/SRemoteControlPanel.cpp b/Engine/Plugins/VirtualProduction/RemoteControl/Source/RemoteControlUI/Private/UI/SRemoteControlPanel.cpp
--- a/Engine/Plugins/VirtualProduction/RemoteControl/Source/RemoteControlUI/Private/UI/SRemoteControlPanel.cpp
+++ b/Engine/Plugins/VirtualProduction/RemoteControl/Source/RemoteControlUI/Private/UI/SRemoteControlPanel.cpp
@@ -342,6 +342,9 @@ void SRemoteControlPanel::Construct(const FArguments& InArgs, URemoteControlPres
 	ActivePanel = ERCPanels::RCP_Properties;
 
 	RCPanelStyle = &FRemoteControlPanelStyle::Get()->GetWidgetStyle<FRCPanelStyle>("RemoteControlPanel.MinorPanel");
+	
+	const URemoteControlSettings* RemoteControlSettings = GetMutableDefault<URemoteControlSettings>();
+	bIsLogicPanelEnabled = RemoteControlSettings->bLogicPanelVisibility;
 
 	TArray<TSharedRef<SWidget>> ExtensionWidgets;
 	FRemoteControlUIModule::Get().GetExtensionGenerators().Broadcast(ExtensionWidgets);
@@ -567,8 +570,39 @@ void SRemoteControlPanel::Construct(const FArguments& InArgs, URemoteControlPres
 	ActionPanel = SNew(SRCActionPanel, SharedThis(this))
 		.Visibility_Lambda([this] { return !bIsInLiveMode && bIsLogicPanelEnabled && (ActivePanel == ERCPanels::RCP_Properties || ActivePanel == ERCPanels::RCP_None) ? EVisibility::Visible : EVisibility::Collapsed; });
 
-	LogicPanel->AddPanel(ControllersAndBehavioursPanel, 0.6f);
-	LogicPanel->AddPanel(ActionPanel.ToSharedRef(), 0.4f);
+	// Retrieve Action Panel Split ratio from RC Settings.
+	// We can't just get the value, since it will update in real time as the user resizes the slot: this value needs to update in real time
+	const TAttribute<float> ActionPanelSplitRatioAttribute = TAttribute<float>::Create(
+		TAttribute<float>::FGetter::CreateLambda([]()
+		{
+			const URemoteControlSettings* Settings = GetMutableDefault<URemoteControlSettings>();
+			return Settings->ActionPanelSplitRatio;
+		}));
+	
+	int32 ActionPanelSlotIndex = LogicPanel->AddPanel(ControllersAndBehavioursPanel, ActionPanelSplitRatioAttribute, true);
+	LogicPanel->AddPanel(ActionPanel.ToSharedRef(), 1.0 - ActionPanelSplitRatioAttribute.Get(), true);
+
+	if (ActionPanelSlotIndex >= 0)
+	{
+		const TWeakPtr<SRCMajorPanel> LogicPanelWeak = LogicPanel.ToWeakPtr();
+
+		// We setup a lambda to fire whenever the user is done resizing the splitter.
+		// This lambda will record the current splitter position in Remote Control Settings
+		LogicPanel->OnSplitterFinishedResizing().BindLambda([ActionPanelSlotIndex, LogicPanelWeak]()
+		{
+			if (LogicPanelWeak.IsValid())
+			{
+				if (const TSharedPtr<SRCMajorPanel> LogicPanelShared = LogicPanelWeak.Pin())
+				{
+					const SSplitter::FSlot& Slot = LogicPanelShared->GetSplitterSlotAt(ActionPanelSlotIndex);		
+					URemoteControlSettings* Settings = GetMutableDefault<URemoteControlSettings>();
+					Settings->ActionPanelSplitRatio = Slot.GetSizeValue();
+					Settings->PostEditChange();
+					Settings->SaveConfig();
+				}
+			}
+		});
+	}
 
 	// Make 2 Columns with Panel Drawer + Main Panel
 	TSharedRef<SSplitter> ContentPanel = SNew(SSplitter)
@@ -2036,6 +2070,9 @@ void SRemoteControlPanel::GenerateToolbar()
 		MiscWidgets = MiscWidgetsHBox;
 	}
 
+	const URemoteControlSettings* Settings = GetMutableDefault<URemoteControlSettings>();
+	const FName& DefaultPanelMode = Settings->DefaultPanelMode;
+
 	Toolbar =
 		SNew(SHorizontalBox)
 		+ SHorizontalBox::Slot()
@@ -2067,7 +2104,7 @@ void SRemoteControlPanel::GenerateToolbar()
 		.AutoWidth()
 		[
 			SNew(SRCModeSwitcher)
-			.DefaultMode("Setup")
+			.DefaultMode(DefaultPanelMode)
 			.OnModeSwitched_Lambda([this](const SRCModeSwitcher::FRCMode& NewMode)
 				{
 					if (NewMode.ModeId == TEXT("Operation"))
@@ -2079,8 +2116,6 @@ void SRemoteControlPanel::GenerateToolbar()
 						bIsInLiveMode = false;
 					}
 
-					bIsLogicPanelEnabled = bIsInLiveMode;
-
 					if (PanelDrawer.IsValid())
 					{
 						TSharedRef<FRCPanelDrawerArgs> LivePanel = RegisteredDrawers.FindChecked(ERCPanels::RCP_Live);
@@ -2099,6 +2134,11 @@ void SRemoteControlPanel::GenerateToolbar()
 					}
 
 					OnLiveModeChange.ExecuteIfBound(SharedThis(this), bIsInLiveMode);
+
+					URemoteControlSettings* Settings = GetMutableDefault<URemoteControlSettings>();					
+					Settings->DefaultPanelMode = NewMode.ModeId;
+					Settings->PostEditChange();
+					Settings->SaveConfig();
 				}
 			)
 
@@ -2375,6 +2415,11 @@ void SRemoteControlPanel::ToggleLogicEditor_Execute()
 {
 	bIsLogicPanelEnabled = !bIsLogicPanelEnabled;
 
+	URemoteControlSettings* Settings = GetMutableDefault<URemoteControlSettings>();
+	Settings->bLogicPanelVisibility = bIsLogicPanelEnabled;
+	Settings->PostEditChange();
+	Settings->SaveConfig();	
+
 	if (PanelDrawer.IsValid() && (ActivePanel != ERCPanels::RCP_Properties || ActivePanel == ERCPanels::RCP_None))
 	{
 		TSharedRef<FRCPanelDrawerArgs> PropertiesPanel = RegisteredDrawers.FindChecked(ERCPanels::RCP_Properties);
