[Remote Control] Add Property Id to list of columns visibility state to be stored in RC config

diff --git a/Engine/Plugins/VirtualProduction/RemoteControl/Source/RemoteControlCommon/Public/RemoteControlSettings.h b/Engine/Plugins/VirtualProduction/RemoteControl/Source/RemoteControlCommon/Public/RemoteControlSettings.h
--- a/Engine/Plugins/VirtualProduction/RemoteControl/Source/RemoteControlCommon/Public/RemoteControlSettings.h
+++ b/Engine/Plugins/VirtualProduction/RemoteControl/Source/RemoteControlCommon/Public/RemoteControlSettings.h
@@ -259,7 +259,8 @@ public:
 	static const TSet<FName>& GetExposedEntitiesColumnNames()
 	{
 		static TSet<FName> ColumnNames =
-			{ TEXT("OwnerName")
+			{ TEXT("PropertyID")
+			, TEXT("OwnerName")
 			, TEXT("Subobject Path")
 			, TEXT("Description")
 			, TEXT("Value")};
