[RemoteControl] RemoteControlUI
Remote control removing Logic/Protocol Header bar, and moving SearchBox and filter in Property Header, Logic/Protocol and settings moved in the upper Header.

diff --git a/Engine/Plugins/VirtualProduction/RemoteControl/Source/RemoteControlUI/Private/UI/Panels/SRCDockPanel.cpp b/Engine/Plugins/VirtualProduction/RemoteControl/Source/RemoteControlUI/Private/UI/Panels/SRCDockPanel.cpp
--- a/Engine/Plugins/VirtualProduction/RemoteControl/Source/RemoteControlUI/Private/UI/Panels/SRCDockPanel.cpp
+++ b/Engine/Plugins/VirtualProduction/RemoteControl/Source/RemoteControlUI/Private/UI/Panels/SRCDockPanel.cpp
@@ -262,7 +262,7 @@ void SRCMinorPanel::Construct(const SRCMinorPanel::FArguments& InArgs)
 
 				// Header Label.
 				+SHorizontalBox::Slot()
-				.FillWidth(1.f)
+				.AutoWidth()
 				.VAlign(VAlign_Center)
 				.Padding(2.f, 4.f)
 				[
@@ -270,7 +270,16 @@ void SRCMinorPanel::Construct(const SRCMinorPanel::FArguments& InArgs)
 					.TextStyle(&RCPanelStyle->HeaderTextStyle)
 					.Text(InArgs._HeaderLabel)
 				]
-				
+
+				// Center Toolbar
+				+SHorizontalBox::Slot()
+				.FillWidth(1.0f)
+				.VAlign(VAlign_Center)
+				.Padding(2.f, 4.f, 5.f, 4.f)
+				[
+					SAssignNew(CenterHeaderToolbar, SHorizontalBox)
+				]
+
 				// Right Toolbar
 				+SHorizontalBox::Slot()
 				.AutoWidth()
@@ -402,6 +411,18 @@ void SRCMinorPanel::AddHeaderToolbarItem(EToolbar InToolbar, TSharedRef<SWidget>
 					];
 			}
 			break;
+		case Center:
+			if (CenterHeaderToolbar.IsValid() && bIsHeaderEnabled.Get())
+			{
+				CenterHeaderToolbar->AddSlot()
+					.FillWidth(1.0f)
+					.Padding(5.f, 0.f)
+					.VAlign(VAlign_Center)
+					[
+						InWidget
+					];
+			}
+			break;
 		default:
 			break;
 	}
diff --git a/Engine/Plugins/VirtualProduction/RemoteControl/Source/RemoteControlUI/Private/UI/Panels/SRCDockPanel.h b/Engine/Plugins/VirtualProduction/RemoteControl/Source/RemoteControlUI/Private/UI/Panels/SRCDockPanel.h
--- a/Engine/Plugins/VirtualProduction/RemoteControl/Source/RemoteControlUI/Private/UI/Panels/SRCDockPanel.h
+++ b/Engine/Plugins/VirtualProduction/RemoteControl/Source/RemoteControlUI/Private/UI/Panels/SRCDockPanel.h
@@ -16,7 +16,10 @@ enum EToolbar
 	Left,
 
 	/** Right toolbar. */
-	Right
+	Right,
+
+	/** Center toolbar. */
+	Center
 };
 
 /**
@@ -123,6 +126,9 @@ private:
 	/** Actual toolbar widget located right to the footer. */
 	TSharedPtr<SHorizontalBox> RightToolbar;
 
+	/** Actual toolbar widget located Center to the footer. */
+	TSharedPtr<SHorizontalBox> CenterToolbar;
+
 	/** Holds the entire child panels. */
 	TSharedPtr<SSplitter> Children;
 
@@ -197,4 +203,7 @@ private:
 
 	/** Actual toolbar widget located right to the header. */
 	TSharedPtr<SHorizontalBox> RightHeaderToolbar;
+
+	/** Actual toolbar widget located Center to the header. */
+    TSharedPtr<SHorizontalBox> CenterHeaderToolbar;
 };
diff --git a/Engine/Plugins/VirtualProduction/RemoteControl/Source/RemoteControlUI/Private/UI/SRCPanelExposedEntitiesList.cpp b/Engine/Plugins/VirtualProduction/RemoteControl/Source/RemoteControlUI/Private/UI/SRCPanelExposedEntitiesList.cpp
--- a/Engine/Plugins/VirtualProduction/RemoteControl/Source/RemoteControlUI/Private/UI/SRCPanelExposedEntitiesList.cpp
+++ b/Engine/Plugins/VirtualProduction/RemoteControl/Source/RemoteControlUI/Private/UI/SRCPanelExposedEntitiesList.cpp
@@ -8,6 +8,7 @@
 #include "Editor/EditorEngine.h"
 #include "EditorFontGlyphs.h"
 #include "Engine/Selection.h"
+#include "Filters/SRCPanelFilter.h"
 #include "GameFramework/Actor.h"
 #include "ISettingsModule.h"
 #include "Input/DragAndDrop.h"
@@ -28,10 +29,12 @@
 #include "SRCModeSwitcher.h"
 #include "SRCPanelFieldGroup.h"
 #include "SRCPanelExposedField.h"
+#include "SSearchToggleButton.h"
 #include "Styling/RemoteControlStyles.h"
 #include "UI/Panels/SRCDockPanel.h"
 #include "UObject/Object.h"
 #include "Widgets/Input/SButton.h"
+#include "Widgets/Input/SSearchBox.h"
 #include "Widgets/Layout/SSpacer.h"
 #include "Widgets/Layout/SSeparator.h"
 #include "Widgets/Layout/SWidgetSwitcher.h"
@@ -187,7 +190,10 @@ void SRCPanelExposedEntitiesList::Construct(const FArguments& InArgs, URemoteCon
 
 	bFilterApplicationRequested = false;
 	bSearchRequested = false;
-	SearchedText = MakeShared<FText>();
+	
+	// Setup search filter.
+	SearchTextFilter = MakeShared<TTextFilter<const SRCPanelTreeNode&>>(TTextFilter<const SRCPanelTreeNode&>::FItemToStringArray::CreateSP(this, &SRCPanelExposedEntitiesList::PopulateSearchStrings));
+	SearchedText = MakeShared<FText>(FText::GetEmpty());
 
 	RCPanelStyle = &FRemoteControlPanelStyle::Get()->GetWidgetStyle<FRCPanelStyle>("RemoteControlPanel.MinorPanel");
 	ActiveListMode = EEntitiesListMode::Default;
@@ -385,7 +391,28 @@ void SRCPanelExposedEntitiesList::Construct(const FArguments& InArgs, URemoteCon
 		}
 	}
 
+	// Create the Filter Widget
+	FilterPtr = SNew(SRCPanelFilter)
+		.OnFilterChanged(this, &SRCPanelExposedEntitiesList::OnFilterChanged)
+		.AddMetaData<FTagMetaData>(FTagMetaData(TEXT("RemoteControlFilters")));
+
+	// Create the Filter Combo Button
+	const TSharedPtr<SWidget> FilterComboButton = SRCPanelFilter::MakeAddFilterButton(FilterPtr.ToSharedRef());
+
+	const TSharedPtr<ISlateMetaData> FilterComboButtonMetaData = MakeShared<FTagMetaData>(TEXT("ContentBrowserFiltersCombo"));
+	FilterComboButton->AddMetadata(FilterComboButtonMetaData.ToSharedRef());
+
+	//Create the SearchBox for the EntitiesList
+	SAssignNew(SearchBoxPtr, SSearchBox)
+		.HintText(LOCTEXT("SearchHint", "Search"))
+		.OnTextChanged(this, &SRCPanelExposedEntitiesList::OnSearchTextChanged)
+		.OnTextCommitted(this, &SRCPanelExposedEntitiesList::OnSearchTextCommitted)
+		.DelayChangeNotificationsWhileTyping(true);
+
 	// Expose Button
+	ExposeDockPanel->AddHeaderToolbarItem(EToolbar::Center, SearchBoxPtr.ToSharedRef());
+	ExposeDockPanel->AddHeaderToolbarItem(EToolbar::Right, SNew(SSearchToggleButton, SearchBoxPtr.ToSharedRef()));
+	ExposeDockPanel->AddHeaderToolbarItem(EToolbar::Right, FilterComboButton.ToSharedRef());
 	ExposeDockPanel->AddHeaderToolbarItem(EToolbar::Right, ModeSwitcher.ToSharedRef());
 	ExposeDockPanel->AddHeaderToolbarItem(EToolbar::Left, InArgs._ExposeActorsComboButton.Get().ToSharedRef());
 	ExposeDockPanel->AddHeaderToolbarItem(EToolbar::Left, InArgs._ExposeFunctionsComboButton.Get().ToSharedRef());
@@ -572,6 +599,14 @@ void SRCPanelExposedEntitiesList::RebuildListWithColumns(EEntitiesListMode InLis
 	}
 }
 
+void SRCPanelExposedEntitiesList::UpdateSearch()
+{
+	if (SearchedText.IsValid() && !SearchedText->IsEmptyOrWhitespace() && SearchedText->ToString().Len() > 3)
+	{
+		OnSearchTextChanged(*SearchedText);
+	}
+}
+
 FText SRCPanelExposedEntitiesList::HandleEntityListHeaderLabel() const
 {
 	const FRemoteControlCommands& Commands = FRemoteControlCommands::Get();
@@ -1076,6 +1111,60 @@ void SRCPanelExposedEntitiesList::UnregisterEvents()
 	IRemoteControlProtocolWidgetsModule::Get().OnProtocolBindingAddedOrRemoved().Remove(OnProtocolBindingAddedOrRemovedHandle);
 }
 
+void SRCPanelExposedEntitiesList::OnFilterChanged()
+{
+	check(FilterPtr.IsValid());
+
+	const FRCFilter Filter = FilterPtr->GetCombinedBackendFilter();
+
+	SetBackendFilter(Filter);
+}
+
+void SRCPanelExposedEntitiesList::OnSearchTextChanged(const FText& InFilterText)
+{
+	SearchTextFilter->SetRawFilterText(InFilterText);
+	SearchBoxPtr->SetError(SearchTextFilter->GetFilterErrorText());
+	*SearchedText = InFilterText;
+
+	const int32 Length = InFilterText.ToString().Len();
+
+	if (Length > 3)
+	{
+		TryRefreshingSearch(InFilterText);
+	}
+	else if (Length == 3 || Length == 0) // Avoid unnecessary refresh if search text is below the threshold.
+		{
+			ResetSearch();
+
+			Refresh();
+		}
+}
+
+void SRCPanelExposedEntitiesList::OnSearchTextCommitted(const FText& InFilterText, ETextCommit::Type InCommitType)
+{
+	if (InCommitType == ETextCommit::OnCleared || InFilterText.IsEmpty())
+	{
+		ResetSearch();
+
+		Refresh();
+
+		return;
+	}
+
+	OnSearchTextChanged(InFilterText);
+}
+
+void SRCPanelExposedEntitiesList::PopulateSearchStrings(const SRCPanelTreeNode& Item, TArray<FString>& OutSearchStrings) const
+{
+	if (Preset.IsValid())
+	{
+		if (const TSharedPtr<FRemoteControlEntity> Entity = Preset->GetExposedEntity<FRemoteControlEntity>(Item.GetRCId()).Pin())
+		{
+			OutSearchStrings.Add(Entity->GetLabel().ToString());
+		}
+	}
+}
+
 TSharedPtr<SRCPanelGroup> SRCPanelExposedEntitiesList::FindGroupById(const FGuid& Id)
 {
 	TSharedPtr<SRCPanelGroup> TargetGroup;
diff --git a/Engine/Plugins/VirtualProduction/RemoteControl/Source/RemoteControlUI/Private/UI/SRCPanelExposedEntitiesList.h b/Engine/Plugins/VirtualProduction/RemoteControl/Source/RemoteControlUI/Private/UI/SRCPanelExposedEntitiesList.h
--- a/Engine/Plugins/VirtualProduction/RemoteControl/Source/RemoteControlUI/Private/UI/SRCPanelExposedEntitiesList.h
+++ b/Engine/Plugins/VirtualProduction/RemoteControl/Source/RemoteControlUI/Private/UI/SRCPanelExposedEntitiesList.h
@@ -3,6 +3,7 @@
 #pragma once
 
 #include "IRemoteControlUIModule.h"
+#include "Misc/TextFilter.h"
 #include "RemoteControlPreset.h"
 #include "SRCPanelTreeNode.h"
 #include "Widgets/DeclarativeSyntaxSupport.h"
@@ -21,10 +22,12 @@ struct FRemoteControlPresetGroup;
 struct FRemoteControlFunction;
 class ITableRow;
 class SRCHeaderRow;
+class SRCPanelFilter;
 class SRCPanelGroup;
 struct SRCPanelTreeNode;
 class SRemoteControlTarget;
 struct SRCPanelExposedField;
+class SSearchBox;
 class STableViewBase;
 class URemoteControlPreset;
 struct FRCPanelStyle;
@@ -104,7 +107,13 @@ public:
 
 	/** Returns delegate triggered upon a modification to an exposed entity. */
 	FSimpleDelegate OnEntityListUpdated() { return OnEntityListUpdatedDelegate; }
-	
+
+	/** Return the FilterPtr of this RCPanelEntitiesList*/
+	TSharedPtr<SRCPanelFilter> GetFilterPtr() const { return FilterPtr; }
+
+	/** Update the Search */
+	void UpdateSearch();
+
 private:
 	/** Handles label to be shown in the entity list header.  */
 	FText HandleEntityListHeaderLabel() const;
@@ -124,6 +133,8 @@ private:
 	void OnSelectionChanged(TSharedPtr<SRCPanelTreeNode> Node, ESelectInfo::Type SelectInfo);
 	/** Handlers for drag/drop events. */
 	FReply OnDropOnGroup(const TSharedPtr<FDragDropOperation>& DragDropOperation, const TSharedPtr<SRCPanelTreeNode>& TargetEntity, const TSharedPtr<SRCPanelTreeNode>& DragTargetGroup);
+	/** Handler for when a filter in the filter list has changed */
+	void OnFilterChanged();
 	/** Get the id of the group that holds a particular widget. */
 	FGuid GetGroupId(const FGuid& EntityId);
 	/** Handles creating a new group. */
@@ -141,6 +152,11 @@ private:
 	float OnGetRightColumnWidth() const { return ColumnWidth; }
 	void OnSetColumnWidth(float InWidth) { ColumnWidth = InWidth; }
 
+	// Exposed Entities filtering. (Filters the Exposed Entities view)
+	void OnSearchTextChanged(const FText& InFilterText);
+	void OnSearchTextCommitted(const FText& InFilterText, ETextCommit::Type InCommitType);
+	void PopulateSearchStrings(const SRCPanelTreeNode& Item, TArray<FString>& OutSearchStrings) const;
+
 	/** Find a group using its id. */
 	TSharedPtr<SRCPanelGroup> FindGroupById(const FGuid& Id);
 
@@ -224,12 +240,18 @@ private:
 	bool bFilterApplicationRequested;
 	/** If true, the entity items will be refreshed next frame. */
 	bool bSearchRequested;
-	/** Holds the text that is being searched actively. */
-	TSharedPtr<FText> SearchedText;
 	/** Panel Style reference. */
 	const FRCPanelStyle* RCPanelStyle;
 	/** Holds the header row of entities list. */
 	TSharedPtr<SRCHeaderRow> FieldsHeaderRow;
+	/** The filter list */
+	TSharedPtr<SRCPanelFilter> FilterPtr;
+	/** The text box used to search for tags. */
+	TSharedPtr<SSearchBox> SearchBoxPtr;
+	/** Text filter for the search text. */
+    TSharedPtr<TTextFilter<const SRCPanelTreeNode&>> SearchTextFilter;
+    /** Actively searched term. */
+    TSharedPtr<FText> SearchedText;
 	/** Holds the active list mode. */
 	EEntitiesListMode ActiveListMode;
 	/** Holds the active protocol enabled by the mode switcher. */
diff --git a/Engine/Plugins/VirtualProduction/RemoteControl/Source/RemoteControlUI/Private/UI/SRemoteControlPanel.cpp b/Engine/Plugins/VirtualProduction/RemoteControl/Source/RemoteControlUI/Private/UI/SRemoteControlPanel.cpp
--- a/Engine/Plugins/VirtualProduction/RemoteControl/Source/RemoteControlUI/Private/UI/SRemoteControlPanel.cpp
+++ b/Engine/Plugins/VirtualProduction/RemoteControl/Source/RemoteControlUI/Private/UI/SRemoteControlPanel.cpp
@@ -351,6 +351,23 @@ void SRemoteControlPanel::Construct(const FArguments& InArgs, URemoteControlPres
 
 	BindRemoteControlCommands();
 
+	GenerateAuxiliaryToolbar();
+
+	AddToolbarWidget(AuxiliaryToolbarWidgetContent.ToSharedRef());
+
+	// Settings
+	AddToolbarWidget(SNew(SButton)
+			.ButtonStyle(&RCPanelStyle->FlatButtonStyle)
+			.ContentPadding(2.0f)
+			.TextStyle(FRemoteControlPanelStyle::Get(), "RemoteControlPanel.Button.TextStyle")
+			.OnClicked(this, &SRemoteControlPanel::OnClickSettingsButton)
+			.ToolTipText(LOCTEXT("OpenRemoteControlSettings", "Open Remote Control settings."))
+			[
+				SNew(SImage)
+				.Image(FAppStyle::Get().GetBrush("Icons.Toolbar.Settings"))
+				.ColorAndOpacity(FSlateColor::UseForeground())
+			]);
+
 	// Show Log
 	AddToolbarWidget(SNew(SCheckBox)
 		.Style(&RCPanelStyle->ToggleButtonStyle)
@@ -373,25 +390,8 @@ void SRemoteControlPanel::Construct(const FArguments& InArgs, URemoteControlPres
 	}
 
 	GenerateToolbar();
-	GenerateAuxiliaryToolbar();
-
 	UpdateRebindButtonVisibility();
 
-	// Setup search filter.
-	SearchTextFilter = MakeShared<TTextFilter<const SRCPanelTreeNode&>>(TTextFilter<const SRCPanelTreeNode&>::FItemToStringArray::CreateSP(this, &SRemoteControlPanel::PopulateSearchStrings));
-	SearchedText = MakeShared<FText>(FText::GetEmpty());
-
-	// Create the Filter Widget
-	FilterPtr = SNew(SRCPanelFilter)
-		.OnFilterChanged(this, &SRemoteControlPanel::OnFilterChanged)
-		.AddMetaData<FTagMetaData>(FTagMetaData(TEXT("RemoteControlFilters")));
-
-	// Create the Filter Combo Button
-	TSharedPtr<SWidget> FilterComboButton = SRCPanelFilter::MakeAddFilterButton(FilterPtr.ToSharedRef());
-
-	TSharedPtr<ISlateMetaData> FilterComboButtonMetaData = MakeShared<FTagMetaData>(TEXT("ContentBrowserFiltersCombo"));
-	FilterComboButton->AddMetadata(FilterComboButtonMetaData.ToSharedRef());
-
 	EntityProtocolDetails = SNew(SBox);
 
 	EntityList = SNew(SRCPanelExposedEntitiesList, Preset.Get(), WidgetRegistry)
@@ -429,85 +429,6 @@ void SRemoteControlPanel::Construct(const FArguments& InArgs, URemoteControlPres
 		.AutoHeight()
 		[
 			ToolbarWidgetContent.ToSharedRef()
-		]
-		+ SVerticalBox::Slot()
-		.HAlign(HAlign_Fill)
-		.AutoHeight()
-		[
-			// Separator
-			SNew(SSeparator)
-			.SeparatorImage(FAppStyle::Get().GetBrush("Separator"))
-			.Thickness(5.f)
-			.Orientation(EOrientation::Orient_Horizontal)
-		]
-		+ SVerticalBox::Slot()
-		.Padding(3.f, 0.f)
-		.AutoHeight()
-		[
-			// Auxiliary Toolbar
-			SNew(SHorizontalBox)
-
-			// Search Box
-			+ SHorizontalBox::Slot()
-			.Padding(5.f, 3.f)
-			.VAlign(VAlign_Center)
-			.FillWidth(1.f)
-			[
-				SAssignNew(SearchBoxPtr, SSearchBox)
-				.HintText(LOCTEXT("SearchHint", "Search"))
-				.OnTextChanged(this, &SRemoteControlPanel::OnSearchTextChanged)
-				.OnTextCommitted(this, &SRemoteControlPanel::OnSearchTextCommitted)
-				.DelayChangeNotificationsWhileTyping(true)
-			]
-
-			// Filters
-			+ SHorizontalBox::Slot()
-			.VAlign(VAlign_Center)
-			.AutoWidth()
-			.Padding(5.f, 3.f, 3.f, 3.f)
-			[
-				FilterComboButton.ToSharedRef()
-			]
-
-			// Settings
-			+ SHorizontalBox::Slot()
-			.VAlign(VAlign_Center)
-			.AutoWidth()
-			.Padding(3.f, 3.f, 3.f, 3.f)
-			[
-				SNew(SButton)
-				.ButtonStyle(&RCPanelStyle->FlatButtonStyle)
-				.ContentPadding(2.0f)
-				.TextStyle(FRemoteControlPanelStyle::Get(), "RemoteControlPanel.Button.TextStyle")
-				.OnClicked(this, &SRemoteControlPanel::OnClickSettingsButton)
-				.ToolTipText(LOCTEXT("OpenRemoteControlSettings", "Open Remote Control settings."))
-				[
-					SNew(SImage)
-					.Image(FAppStyle::Get().GetBrush("Icons.Toolbar.Settings"))
-					.ColorAndOpacity(FSlateColor::UseForeground())
-				]
-			]
-
-			// Separator
-			+SHorizontalBox::Slot()
-			.HAlign(HAlign_Center)
-			.VAlign(VAlign_Fill)
-			.AutoWidth()
-			.Padding(3.f, 3.f, 0.f, 3.f)
-			[
-				SNew(SSeparator)
-				.SeparatorImage(FAppStyle::Get().GetBrush("Separator"))
-				.Thickness(2.f)
-				.Orientation(EOrientation::Orient_Vertical)
-			]
-
-			// Mini Toolbar Widget
-			+ SHorizontalBox::Slot()
-			.AutoWidth()
-			.Padding(3.f)
-			[
-				AuxiliaryToolbarWidgetContent.ToSharedRef()
-			]
 		];
 
 	TSharedRef<SWidget> FooterPanel = SNew(SVerticalBox)
@@ -1891,9 +1812,9 @@ void SRemoteControlPanel::UpdateEntityDetailsView(const TSharedPtr<SRCPanelTreeN
 	}
 
 	// Trigger search to list the search results specific to selected group.
-	if (SearchedText.IsValid() && !SearchedText->IsEmptyOrWhitespace() && SearchedText->ToString().Len() > 3)
+	if (EntityList.IsValid())
 	{
-		OnSearchTextChanged(*SearchedText);
+		EntityList->UpdateSearch();
 	}
 }
 
@@ -2651,72 +2572,21 @@ bool SRemoteControlPanel::CanDuplicateItem() const
 	return false;
 }
 
-void SRemoteControlPanel::OnSearchTextChanged(const FText& InFilterText)
-{
-	SearchTextFilter->SetRawFilterText(InFilterText);
-	SearchBoxPtr->SetError(SearchTextFilter->GetFilterErrorText());
-	*SearchedText = InFilterText;
-
-	const int32 Length = InFilterText.ToString().Len();
-
-	check(EntityList.IsValid());
-
-	if (Length > 3)
-	{
-		EntityList->TryRefreshingSearch(InFilterText);
-	}
-	else if (Length == 3 || Length == 0) // Avoid unnecessary refresh if search text is below the threshold.
-	{
-		EntityList->ResetSearch();
-
-		EntityList->Refresh();
-	}
-}
-
-void SRemoteControlPanel::OnSearchTextCommitted(const FText& InFilterText, ETextCommit::Type InCommitType)
-{
-	if (InCommitType == ETextCommit::OnCleared || InFilterText.IsEmpty())
-	{
-		EntityList->ResetSearch();
-
-		EntityList->Refresh();
-
-		return;
-	}
-
-	OnSearchTextChanged(InFilterText);
-}
-
-void SRemoteControlPanel::PopulateSearchStrings(const SRCPanelTreeNode& Item, TArray<FString>& OutSearchStrings) const
-{
-	if (Preset.IsValid())
-	{
-		if (TSharedPtr<FRemoteControlEntity> Entity = Preset->GetExposedEntity<FRemoteControlEntity>(Item.GetRCId()).Pin())
-		{
-			OutSearchStrings.Add(Entity->GetLabel().ToString());
-		}
-	}
-}
-
-void SRemoteControlPanel::OnFilterChanged()
-{
-	check(FilterPtr.IsValid());
-
-	FRCFilter Filter = FilterPtr->GetCombinedBackendFilter();
-
-	EntityList->SetBackendFilter(Filter);
-}
-
-void SRemoteControlPanel::LoadSettings(const FGuid& InInstanceId)
+void SRemoteControlPanel::LoadSettings(const FGuid& InInstanceId) const
 {
 	const FString SettingsString = InInstanceId.ToString();
 
 	// Load all our data using the settings string as a key in the user settings ini.
-	FilterPtr->LoadSettings(GEditorPerProjectIni, IRemoteControlUIModule::SettingsIniSection, SettingsString);
+	const TSharedPtr<SRCPanelFilter> FilterPtr = EntityList->GetFilterPtr();
+	if (FilterPtr.IsValid())
+	{
+		FilterPtr->LoadSettings(GEditorPerProjectIni, IRemoteControlUIModule::SettingsIniSection, SettingsString);
+	}
 }
 
 void SRemoteControlPanel::SaveSettings()
 {
+	const TSharedPtr<SRCPanelFilter> FilterPtr = EntityList->GetFilterPtr();
 	if (Preset.IsValid() && FilterPtr.IsValid())
 	{
 		const FString SettingsString = Preset->GetPresetId().ToString();
diff --git a/Engine/Plugins/VirtualProduction/RemoteControl/Source/RemoteControlUI/Private/UI/SRemoteControlPanel.h b/Engine/Plugins/VirtualProduction/RemoteControl/Source/RemoteControlUI/Private/UI/SRemoteControlPanel.h
--- a/Engine/Plugins/VirtualProduction/RemoteControl/Source/RemoteControlUI/Private/UI/SRemoteControlPanel.h
+++ b/Engine/Plugins/VirtualProduction/RemoteControl/Source/RemoteControlUI/Private/UI/SRemoteControlPanel.h
@@ -332,16 +332,8 @@ private:
 	/** Called to test if user is able to Duplicate a logic UI item. */
 	bool CanDuplicateItem() const;
 
-	// Exposed Entities filtering. (Filters the Exposed Entities view)
-	void OnSearchTextChanged(const FText& InFilterText);
-	void OnSearchTextCommitted(const FText& InFilterText, ETextCommit::Type InCommitType);
-	void PopulateSearchStrings(const SRCPanelTreeNode& Item, TArray<FString>& OutSearchStrings) const;
-
-	/** Handler for when a filter in the filter list has changed */
-	void OnFilterChanged();
-
 	/** Loads settings from config based on the preset identifier. */
-	void LoadSettings(const FGuid& InInstanceId);
+	void LoadSettings(const FGuid& InInstanceId) const;
 
 	/** Saves settings from config based on the preset identifier. */
 	void SaveSettings();
@@ -406,16 +398,8 @@ private:
 	TSharedPtr<SBorder> AuxiliaryToolbarWidgetContent;
 	/** Additional widgets to be added to the toolbar */
 	TArray<TSharedRef<SWidget>> ToolbarWidgets;
-	/** The text box used to search for tags. */
-	TSharedPtr<SSearchBox> SearchBoxPtr;
 	/** Holds a shared pointer reference to the active entity that is selected. */
 	TSharedPtr<SRCPanelTreeNode> SelectedEntity;
-	/** Text filter for the search text. */
-	TSharedPtr<TTextFilter<const SRCPanelTreeNode&>> SearchTextFilter;
-	/** The filter list */
-	TSharedPtr<SRCPanelFilter> FilterPtr;
-	/** Actively serached term. */
-	TSharedPtr<FText> SearchedText;
 	/** Panel Drawer widget holds all docked panels. */
 	TSharedPtr<SRCPanelDrawer> PanelDrawer;
 	/** Map of Opened Drawers. */
