[RemoteControl]
Fix memory leaks.
Both the Preset and the PropertyIdRegistry hold a TObjectPtr to each other.

diff --git a/Engine/Plugins/VirtualProduction/RemoteControl/Source/RemoteControl/Private/RemoteControlPreset.cpp b/Engine/Plugins/VirtualProduction/RemoteControl/Source/RemoteControl/Private/RemoteControlPreset.cpp
--- a/Engine/Plugins/VirtualProduction/RemoteControl/Source/RemoteControl/Private/RemoteControlPreset.cpp
+++ b/Engine/Plugins/VirtualProduction/RemoteControl/Source/RemoteControl/Private/RemoteControlPreset.cpp
@@ -520,7 +520,7 @@ URemoteControlPreset::URemoteControlPreset()
 	Registry = CreateDefaultSubobject<URemoteControlExposeRegistry>(FName("ExposeRegistry"));
 
 	PropertyIdRegistry = CreateDefaultSubobject<URemoteControlPropertyIdRegistry>(FName("PropertyIdRegistry"));
-	PropertyIdRegistry->Initialize(this);
+	PropertyIdRegistry->Initialize();
 }
 
 void URemoteControlPreset::PostInitProperties()
@@ -565,7 +565,7 @@ void URemoteControlPreset::PostLoad()
 
 	RemoveUnusedBindings();
 
-	PropertyIdRegistry->Initialize(this);
+	PropertyIdRegistry->Initialize();
 }
 
 void URemoteControlPreset::PostDuplicate(bool bDuplicateForPIE)
@@ -1609,7 +1609,7 @@ void URemoteControlPreset::RebindUnboundEntities()
 	Modify();
 	RebindingManager->Rebind(this);
 	Algo::Transform(Registry->GetExposedEntities(), PerFrameUpdatedEntities, [](const TSharedPtr<FRemoteControlEntity>& Entity) { return Entity->GetId(); });
-	PropertyIdRegistry->Initialize(this);
+	PropertyIdRegistry->Initialize();
 }
 
 void URemoteControlPreset::RebindAllEntitiesUnderSameActor(const FGuid& EntityId, AActor* NewActor, bool bUseRebindingContext)
diff --git a/Engine/Plugins/VirtualProduction/RemoteControl/Source/RemoteControl/Private/RemoteControlPropertyIdRegistry.cpp b/Engine/Plugins/VirtualProduction/RemoteControl/Source/RemoteControl/Private/RemoteControlPropertyIdRegistry.cpp
--- a/Engine/Plugins/VirtualProduction/RemoteControl/Source/RemoteControl/Private/RemoteControlPropertyIdRegistry.cpp
+++ b/Engine/Plugins/VirtualProduction/RemoteControl/Source/RemoteControl/Private/RemoteControlPropertyIdRegistry.cpp
@@ -203,9 +203,14 @@ void FRCPropertyIdWrapper::UpdateTypes(const TSharedRef<FRemoteControlProperty>&
 	}
 }
 
-void URemoteControlPropertyIdRegistry::Initialize(URemoteControlPreset* InSourcePreset)
+void URemoteControlPropertyIdRegistry::Initialize()
 {
-	SourcePreset = InSourcePreset;
+	URemoteControlPreset* SourcePreset = GetSourcePreset();
+	if (!SourcePreset)
+	{
+		return;
+	}
+
 	IdentifiedFields.Reset();
 	for (TWeakPtr<FRemoteControlEntity> RCEntity : SourcePreset->GetExposedEntities())
 	{
@@ -216,169 +221,179 @@ void URemoteControlPropertyIdRegistry::Initialize(URemoteControlPreset* InSource
 	}
 }
 
+URemoteControlPreset* URemoteControlPropertyIdRegistry::GetSourcePreset() const
+{
+	return GetTypedOuter<URemoteControlPreset>();
+}
+
 void URemoteControlPropertyIdRegistry::PerformChainReaction(const FRemoteControlPropertyIdArgs& InArgs)
 {
+	URemoteControlPreset* SourcePreset = GetSourcePreset();
+	if (!IsValid(SourcePreset))
+	{
+		return;
+	}
+
 	TSet<FGuid> TargetProperties;
+
 	Algo::TransformIf(IdentifiedFields, TargetProperties,
-		[&](const FRCPropertyIdWrapper& Wrapper)
+		[InArgs](const FRCPropertyIdWrapper& Wrapper)
 		{
 			return Wrapper.IsValid()
 			&& Wrapper.GetPropertyId() == InArgs.PropertyId
 			&& Wrapper.GetSuperType() == InArgs.SuperType
 			&& Wrapper.GetSubType() == InArgs.SubType;
 		},
-		[&](const FRCPropertyIdWrapper& Wrapper)
+		[](const FRCPropertyIdWrapper& Wrapper)
 		{
 			return Wrapper.GetEntityId();
 		}
 		);
-	if (!TargetProperties.IsEmpty() && IsValid(SourcePreset))
+
+	for (const FGuid& TargetProperty : TargetProperties)
 	{
-		for (const FGuid& TargetProperty : TargetProperties)
+		if (TSharedPtr<FRemoteControlProperty> TargetRCProperty = SourcePreset->GetExposedEntity<FRemoteControlProperty>(TargetProperty).Pin())
 		{
-			if (TSharedPtr<FRemoteControlProperty> TargetRCProperty = SourcePreset->GetExposedEntity<FRemoteControlProperty>(TargetProperty).Pin())
+			if (UObject* BoundObject = TargetRCProperty->GetBoundObject())
 			{
-				if (UObject* BoundObject = TargetRCProperty->GetBoundObject())
-				{
-					bool bCopyComplete = false;
+				bool bCopyComplete = false;
 
-					if (FProperty* Property = TargetRCProperty->GetProperty())
-					{
+				if (FProperty* Property = TargetRCProperty->GetProperty())
+				{
 #if WITH_EDITOR
-						BoundObject->PreEditChange(Property);
-						
-						BoundObject->Modify();
+					BoundObject->PreEditChange(Property);
+					
+					BoundObject->Modify();
 #endif // WITH_EDITOR
-						// Note : For all other object types except materials.
-						if (FObjectProperty* ObjectProperty = CastField<FObjectProperty>(Property))
+					// Note : For all other object types except materials.
+					if (FObjectProperty* ObjectProperty = CastField<FObjectProperty>(Property))
+					{
+						if (uint8* PropertyValuePtr = ObjectProperty->ContainerPtrToValuePtr<uint8>(BoundObject))
 						{
-							if (uint8* PropertyValuePtr = ObjectProperty->ContainerPtrToValuePtr<uint8>(BoundObject))
+							UObject* CurrentObject = ObjectProperty->GetObjectPropertyValue(PropertyValuePtr);
+							if (CurrentObject && CurrentObject->IsA(InArgs.SourceClass))
 							{
-								UObject* CurrentObject = ObjectProperty->GetObjectPropertyValue(PropertyValuePtr);
-								if (CurrentObject && CurrentObject->IsA(InArgs.SourceClass))
-								{
-									ObjectProperty->SetObjectPropertyValue(PropertyValuePtr, InArgs.SourceObject);
-									bCopyComplete = true;
-								}
+								ObjectProperty->SetObjectPropertyValue(PropertyValuePtr, InArgs.SourceObject);
+								bCopyComplete = true;
 							}
 						}
-						// Note : Specialization for Materials.
-						else if (FArrayProperty* ArrayProperty = CastField<FArrayProperty>(Property))
+					}
+					// Note : Specialization for Materials.
+					else if (FArrayProperty* ArrayProperty = CastField<FArrayProperty>(Property))
+					{
+						if (uint8* PropertyValuePtr = ArrayProperty->ContainerPtrToValuePtr<uint8>(BoundObject))
 						{
-							if (uint8* PropertyValuePtr = ArrayProperty->ContainerPtrToValuePtr<uint8>(BoundObject))
+							if (FObjectProperty* InnerObjectProperty = CastField<FObjectProperty>(ArrayProperty->Inner))
 							{
-								if (FObjectProperty* InnerObjectProperty = CastField<FObjectProperty>(ArrayProperty->Inner))
+								FScriptArrayHelper ArrayHelper(ArrayProperty, PropertyValuePtr);
+								if (TSharedPtr<FRemoteControlField> TargetRCField = StaticCastSharedPtr<FRemoteControlField>(TargetRCProperty); !TargetRCField->FieldPathInfo.Segments.IsEmpty())
 								{
-									FScriptArrayHelper ArrayHelper(ArrayProperty, PropertyValuePtr);
-									if (TSharedPtr<FRemoteControlField> TargetRCField = StaticCastSharedPtr<FRemoteControlField>(TargetRCProperty); !TargetRCField->FieldPathInfo.Segments.IsEmpty())
+									const int32 MaterialIndex = TargetRCField->FieldPathInfo.Segments[0].ArrayIndex;
+									uint8* ObjPtrContainer = ArrayHelper.GetRawPtr(MaterialIndex);
+									UObject* CurrentObject = InnerObjectProperty->GetObjectPropertyValue(ObjPtrContainer);
+									if (CurrentObject && CurrentObject->IsA(InArgs.SourceClass))
 									{
-										const int32 MaterialIndex = TargetRCField->FieldPathInfo.Segments[0].ArrayIndex;
-										uint8* ObjPtrContainer = ArrayHelper.GetRawPtr(MaterialIndex);
-										UObject* CurrentObject = InnerObjectProperty->GetObjectPropertyValue(ObjPtrContainer);
-										if (CurrentObject && CurrentObject->IsA(InArgs.SourceClass))
-										{
-											InnerObjectProperty->SetObjectPropertyValue(ObjPtrContainer, InArgs.SourceObject);
-											bCopyComplete = true;
-										}
+										InnerObjectProperty->SetObjectPropertyValue(ObjPtrContainer, InArgs.SourceObject);
+										bCopyComplete = true;
 									}
 								}
 							}
 						}
-						// Note : For primitive types not in container and that has a valid UClass Owner.
-						// ContainerPtrToValuePtr will crash if the GetOwner<UClass> is nullptr so we check before.
-						if (!bCopyComplete && Property->GetOwner<UClass>() != nullptr &&
-							!Property->IsA<FArrayProperty>() && !Property->IsA<FSetProperty>() && !Property->IsA<FMapProperty>())
+					}
+					// Note : For primitive types not in container and that has a valid UClass Owner.
+					// ContainerPtrToValuePtr will crash if the GetOwner<UClass> is nullptr so we check before.
+					if (!bCopyComplete && Property->GetOwner<UClass>() != nullptr &&
+						!Property->IsA<FArrayProperty>() && !Property->IsA<FSetProperty>() && !Property->IsA<FMapProperty>())
+					{
+						if (uint8* PropertyValuePtr = Property->ContainerPtrToValuePtr<uint8>(BoundObject))
 						{
-							if (uint8* PropertyValuePtr = Property->ContainerPtrToValuePtr<uint8>(BoundObject))
+							if (const FStructProperty* StructProperty = CastField<FStructProperty>(Property))
 							{
-								if (const FStructProperty* StructProperty = CastField<FStructProperty>(Property))
-								{
-									// FLinearColor are treated as FColor to avoid creating 2 different widget for them and avoid confusion
-									// Instead of the normal CopyCompleteValue we use this for FLinearColor
-									if (StructProperty->Struct->GetFName() == NAME_LinearColor)
-									{
-										FColor ColorValue;
-										InArgs.VirtualProperty->GetValueColor(ColorValue);
-										const FLinearColor RealValue(ColorValue);
-										Property->CopyCompleteValue(PropertyValuePtr, &RealValue);
-										bCopyComplete = true;
-									}
-								}
-								// We do this to avoid copying it above and here
-								if (bCopyComplete == false)
+								// FLinearColor are treated as FColor to avoid creating 2 different widget for them and avoid confusion
+								// Instead of the normal CopyCompleteValue we use this for FLinearColor
+								if (StructProperty->Struct->GetFName() == NAME_LinearColor)
 								{
-									bCopyComplete = InArgs.VirtualProperty->CopyCompleteValue(Property, PropertyValuePtr, true);
+									FColor ColorValue;
+									InArgs.VirtualProperty->GetValueColor(ColorValue);
+									const FLinearColor RealValue(ColorValue);
+									Property->CopyCompleteValue(PropertyValuePtr, &RealValue);
+									bCopyComplete = true;
 								}
 							}
+							// We do this to avoid copying it above and here
+							if (bCopyComplete == false)
+							{
+								bCopyComplete = InArgs.VirtualProperty->CopyCompleteValue(Property, PropertyValuePtr, true);
+							}
 						}
-						//Note : For all the other cases, Container and UStruct Owner.
-						else if (!bCopyComplete)
-						{
+					}
+					//Note : For all the other cases, Container and UStruct Owner.
+					else if (!bCopyComplete)
+					{
 #if !WITH_EDITOR
-							StructMemoryContainer = nullptr;
+						StructMemoryContainer = nullptr;
 #endif // !WITH_EDITOR
-							bCopyComplete = CopyNonUClassOwnerProperty(InArgs, TargetRCProperty, Property, BoundObject);
+						bCopyComplete = CopyNonUClassOwnerProperty(InArgs, TargetRCProperty, Property, BoundObject);
+					}
+					if (bCopyComplete)
+					{
+#if WITH_EDITOR
+						TargetRCProperty->FieldPathInfo.Resolve(BoundObject);
+						FPropertyChangedEvent PropertyEvent = TargetRCProperty->FieldPathInfo.ToPropertyChangedEvent(EPropertyChangeType::ValueSet);
+						FEditPropertyChain EditPropertyChain;
+						TargetRCProperty->FieldPathInfo.ToEditPropertyChain(EditPropertyChain);
+						if (EditPropertyChain.IsEmpty())
+						{
+							BoundObject->PostEditChangeProperty(PropertyEvent);
 						}
-						if (bCopyComplete)
+						else
 						{
-#if WITH_EDITOR
-							TargetRCProperty->FieldPathInfo.Resolve(BoundObject);
-							FPropertyChangedEvent PropertyEvent = TargetRCProperty->FieldPathInfo.ToPropertyChangedEvent(EPropertyChangeType::ValueSet);
-							FEditPropertyChain EditPropertyChain;
-							TargetRCProperty->FieldPathInfo.ToEditPropertyChain(EditPropertyChain);
-							if (EditPropertyChain.IsEmpty())
+							FPropertyChangedChainEvent ChainEvent(EditPropertyChain, PropertyEvent);
+							TArray<TMap<FString, int32>> ArrayIndicesPerObject;
 							{
-								BoundObject->PostEditChangeProperty(PropertyEvent);
-							}
-							else
-							{
-								FPropertyChangedChainEvent ChainEvent(EditPropertyChain, PropertyEvent);
-								TArray<TMap<FString, int32>> ArrayIndicesPerObject;
+								TMap<FString, int32> ArrayIndices;
+								ArrayIndices.Reserve(TargetRCProperty->FieldPathInfo.Segments.Num());
+								for (const FRCFieldPathSegment& Segment : TargetRCProperty->FieldPathInfo.Segments)
 								{
-									TMap<FString, int32> ArrayIndices;
-									ArrayIndices.Reserve(TargetRCProperty->FieldPathInfo.Segments.Num());
-									for (const FRCFieldPathSegment& Segment : TargetRCProperty->FieldPathInfo.Segments)
-									{
-										ArrayIndices.Add(Segment.Name.ToString(), Segment.ArrayIndex);
-									}
-									ArrayIndicesPerObject.Add(MoveTemp(ArrayIndices));
+									ArrayIndices.Add(Segment.Name.ToString(), Segment.ArrayIndex);
 								}
-								ChainEvent.ObjectIteratorIndex = 0;
-								ChainEvent.SetArrayIndexPerObject(ArrayIndicesPerObject);
-								BoundObject->PostEditChangeChainProperty(ChainEvent);
+								ArrayIndicesPerObject.Add(MoveTemp(ArrayIndices));
 							}
-							
+							ChainEvent.ObjectIteratorIndex = 0;
+							ChainEvent.SetArrayIndexPerObject(ArrayIndicesPerObject);
+							BoundObject->PostEditChangeChainProperty(ChainEvent);
+						}
+						
 #else // NOTE: During runtime use Serialization API to update the property properly.
-							TArray<uint8> Buffer;
-							FMemoryWriter Writer(Buffer);
-							FCborStructSerializerBackend WriterBackend(Writer, EStructSerializerBackendFlags::Default);
-							FStructSerializerPolicies Policies;
-							Policies.MapSerialization = EStructSerializerMapPolicies::Array;
-
-							// We do this at runtime because otherwise Property Owned by UStruct won't update.
-							// So during the copy we take the StructContainer and use that to serialize the property.
-							if (StructMemoryContainer != nullptr)
-							{
-								FStructSerializer::SerializeElement(StructMemoryContainer, Property, INDEX_NONE, WriterBackend, Policies);
-							}
-							else
-							{
-								FStructSerializer::SerializeElement(BoundObject, Property, INDEX_NONE, WriterBackend, Policies);
-							}
-							// Deserialization
-							FMemoryReader Reader(Buffer);
-							FCborStructDeserializerBackend ReaderBackend(Reader);
-							FRCObjectReference TargetObjectRef;
-							if (!GetObjectRef(TargetRCProperty, ERCAccess::WRITE_ACCESS, TargetObjectRef))
-							{
-								continue;
-							}
-							if (!IRemoteControlModule::Get().SetObjectProperties(TargetObjectRef, ReaderBackend, ERCPayloadType::Cbor, Buffer))
-							{
-								continue;
-							}
-#endif // WITH_EDITOR
+						TArray<uint8> Buffer;
+						FMemoryWriter Writer(Buffer);
+						FCborStructSerializerBackend WriterBackend(Writer, EStructSerializerBackendFlags::Default);
+						FStructSerializerPolicies Policies;
+						Policies.MapSerialization = EStructSerializerMapPolicies::Array;
+
+						// We do this at runtime because otherwise Property Owned by UStruct won't update.
+						// So during the copy we take the StructContainer and use that to serialize the property.
+						if (StructMemoryContainer != nullptr)
+						{
+							FStructSerializer::SerializeElement(StructMemoryContainer, Property, INDEX_NONE, WriterBackend, Policies);
+						}
+						else
+						{
+							FStructSerializer::SerializeElement(BoundObject, Property, INDEX_NONE, WriterBackend, Policies);
+						}
+						// Deserialization
+						FMemoryReader Reader(Buffer);
+						FCborStructDeserializerBackend ReaderBackend(Reader);
+						FRCObjectReference TargetObjectRef;
+						if (!GetObjectRef(TargetRCProperty, ERCAccess::WRITE_ACCESS, TargetObjectRef))
+						{
+							continue;
 						}
+						if (!IRemoteControlModule::Get().SetObjectProperties(TargetObjectRef, ReaderBackend, ERCPayloadType::Cbor, Buffer))
+						{
+							continue;
+						}
+#endif // WITH_EDITOR
 					}
 				}
 			}
@@ -414,9 +429,10 @@ void URemoteControlPropertyIdRegistry::UpdateIdentifiedField(const TSharedRef<FR
 		if (FRCPropertyIdWrapper* Wrapper = IdentifiedFields.FindByHash(Hash, InFieldToIdentify->GetId()))
 		{
 			Wrapper->SetPropertyId(InFieldToIdentify->PropertyId);
-			if (IsValid(SourcePreset))
+
+			if (URemoteControlPreset* SourcePreset = GetSourcePreset())
 			{
-				SourcePreset.Get()->GetPropertyIdRegistry()->OnPropertyIdUpdated().Broadcast();
+				SourcePreset->GetPropertyIdRegistry()->OnPropertyIdUpdated().Broadcast();
 			}
 		}
 	}
diff --git a/Engine/Plugins/VirtualProduction/RemoteControl/Source/RemoteControl/Public/RemoteControlPropertyIdRegistry.h b/Engine/Plugins/VirtualProduction/RemoteControl/Source/RemoteControl/Public/RemoteControlPropertyIdRegistry.h
--- a/Engine/Plugins/VirtualProduction/RemoteControl/Source/RemoteControl/Public/RemoteControlPropertyIdRegistry.h
+++ b/Engine/Plugins/VirtualProduction/RemoteControl/Source/RemoteControl/Public/RemoteControlPropertyIdRegistry.h
@@ -118,9 +118,10 @@ class REMOTECONTROL_API URemoteControlPropertyIdRegistry : public UObject
 public:
 	/**
 	 * Initialize the Registry with the given Preset.
-	 * @param InSourcePreset Preset to get the data from.
 	 */
-	void Initialize(URemoteControlPreset* InSourcePreset);
+	void Initialize();
+
+	URemoteControlPreset* GetSourcePreset() const;
 
 	/**
 	 * Update the value(s) of the property(ies) that are bound to a PropertyIdAction.
@@ -205,10 +206,6 @@ private:
 	UPROPERTY()
 	TSet<FRCPropertyIdWrapper> IdentifiedFields;
 
-	/** Holds the source remote control preset asset. */
-	UPROPERTY()
-	TObjectPtr<URemoteControlPreset> SourcePreset;
-
 	/** Delegate triggered when a property has its field Id changed. */
 	FOnPropertyIdUpdated OnPropertyIdUpdatedDelegate;
 
diff --git a/Engine/Plugins/VirtualProduction/RemoteControl/Source/RemoteControlUI/Private/UI/SRCPanelExposedEntitiesList.cpp b/Engine/Plugins/VirtualProduction/RemoteControl/Source/RemoteControlUI/Private/UI/SRCPanelExposedEntitiesList.cpp
--- a/Engine/Plugins/VirtualProduction/RemoteControl/Source/RemoteControlUI/Private/UI/SRCPanelExposedEntitiesList.cpp
+++ b/Engine/Plugins/VirtualProduction/RemoteControl/Source/RemoteControlUI/Private/UI/SRCPanelExposedEntitiesList.cpp
@@ -1733,7 +1733,7 @@ void SRCPanelExposedEntitiesList::ProcessRefresh()
 	if (Preset.IsValid())
 	{
 		//Refresh during PostUndo/Redo to keep the action updated.
-		Preset->GetPropertyIdRegistry()->Initialize(Preset.Get());
+		Preset->GetPropertyIdRegistry()->Initialize();
 		Preset->GetPropertyIdRegistry()->OnPropertyIdUpdated().Broadcast();
 
 		constexpr bool bForceMouseClick = true;
