[RemoteControl] RemoteControlUI
Shorten RemoteContro drag and drop widget

diff --git a/Engine/Plugins/VirtualProduction/RemoteControl/Source/RemoteControlUI/Private/UI/SRCPanelExposedEntity.cpp b/Engine/Plugins/VirtualProduction/RemoteControl/Source/RemoteControlUI/Private/UI/SRCPanelExposedEntity.cpp
--- a/Engine/Plugins/VirtualProduction/RemoteControl/Source/RemoteControlUI/Private/UI/SRCPanelExposedEntity.cpp
+++ b/Engine/Plugins/VirtualProduction/RemoteControl/Source/RemoteControlUI/Private/UI/SRCPanelExposedEntity.cpp
@@ -589,7 +589,7 @@ TSharedRef<SWidget> SRCPanelExposedEntity::CreateEntityWidget(TSharedPtr<SWidget
 
 	Args.ResetButton = ResetWidget;
 
-	Widget->SetContent(MakeNodeWidget(Args));
+	Widget->SetContent(MakeDragAndDropWidget(Args));
 	return Widget;
 }
 
diff --git a/Engine/Plugins/VirtualProduction/RemoteControl/Source/RemoteControlUI/Private/UI/SRCPanelTreeNode.cpp b/Engine/Plugins/VirtualProduction/RemoteControl/Source/RemoteControlUI/Private/UI/SRCPanelTreeNode.cpp
--- a/Engine/Plugins/VirtualProduction/RemoteControl/Source/RemoteControlUI/Private/UI/SRCPanelTreeNode.cpp
+++ b/Engine/Plugins/VirtualProduction/RemoteControl/Source/RemoteControlUI/Private/UI/SRCPanelTreeNode.cpp
@@ -175,6 +175,37 @@ TSharedRef<SWidget> SRCPanelTreeNode::MakeNodeWidget(const FMakeNodeWidgetArgs&
 	return MakeSplitRow(LeftColumn, RightColumn);
 }
 
+TSharedRef<SWidget> SRCPanelTreeNode::MakeDragAndDropWidget(const FMakeNodeWidgetArgs& Args)
+{
+	MakeNodeWidgets(Args);
+
+	const TSharedRef<SWidget> LeftColumn = SNew(SHorizontalBox)
+		.Clipping(EWidgetClipping::OnDemand)
+		// Field name
+		+ SHorizontalBox::Slot()
+		.VAlign(VAlign_Center)
+		.AutoWidth()
+		[
+			NodeNameWidget.ToSharedRef()
+		];
+
+	const TSharedRef<SWidget> RightColumn = SNew(SHorizontalBox)
+		.Clipping(EWidgetClipping::OnDemand)
+		// Node Value
+		+ SHorizontalBox::Slot()
+		.HAlign(HAlign_Fill)
+		[
+			SNew(SBox)
+			.HAlign(HAlign_Left)
+			.VAlign(VAlign_Center)
+			[
+				NodeValueWidget.ToSharedRef()
+			]
+		];
+
+	return MakeSplitRow(LeftColumn, RightColumn);
+}
+
 void SRCPanelTreeNode::MakeNodeWidgets(const FMakeNodeWidgetArgs& Args)
 {
 	auto WidgetOrNull = [](const TSharedPtr<SWidget>& Widget) {return Widget ? Widget.ToSharedRef() : SNullWidget::NullWidget; };
diff --git a/Engine/Plugins/VirtualProduction/RemoteControl/Source/RemoteControlUI/Private/UI/SRCPanelTreeNode.h b/Engine/Plugins/VirtualProduction/RemoteControl/Source/RemoteControlUI/Private/UI/SRCPanelTreeNode.h
--- a/Engine/Plugins/VirtualProduction/RemoteControl/Source/RemoteControlUI/Private/UI/SRCPanelTreeNode.h
+++ b/Engine/Plugins/VirtualProduction/RemoteControl/Source/RemoteControlUI/Private/UI/SRCPanelTreeNode.h
@@ -101,6 +101,8 @@ protected:
 	TSharedRef<SWidget> MakeSplitRow(TSharedRef<SWidget> LeftColumn, TSharedRef<SWidget> RightColumn);
 	/** Create a widget that represents a node in the panel tree hierarchy. */
 	TSharedRef<SWidget> MakeNodeWidget(const FMakeNodeWidgetArgs& Args);
+	/** Create a widget that represents a node for the drag and drop. */
+	TSharedRef<SWidget> MakeDragAndDropWidget(const FMakeNodeWidgetArgs& Args);
 	/** Creates cached copies of underlying widgets. */
 	void MakeNodeWidgets(const FMakeNodeWidgetArgs& Args);
 
