Fixed other Remote Control Presets overwriting the Embedded preset's values in a PIE World

#rnx

diff --git a/Engine/Plugins/VirtualProduction/RemoteControl/Source/RemoteControl/Private/RemoteControlBinding.cpp b/Engine/Plugins/VirtualProduction/RemoteControl/Source/RemoteControl/Private/RemoteControlBinding.cpp
--- a/Engine/Plugins/VirtualProduction/RemoteControl/Source/RemoteControl/Private/RemoteControlBinding.cpp
+++ b/Engine/Plugins/VirtualProduction/RemoteControl/Source/RemoteControl/Private/RemoteControlBinding.cpp
@@ -268,7 +268,15 @@ void URemoteControlLevelDependantBinding::UnbindObject(const TSoftObjectPtr<UObj
 
 UObject* URemoteControlLevelDependantBinding::Resolve() const
 {
-	constexpr bool bAllowPIE = true;
+	bool bAllowPIE = true;
+
+	// Disallow PIE for Embedded Presets as there should be an Embedded Preset counterpart in PIE taking care of this
+	URemoteControlPreset* Preset = Cast<URemoteControlPreset>(GetOuter());
+	if (Preset && Preset->IsEmbeddedPreset())
+	{
+		bAllowPIE = false;
+	}
+
 	UObject* Object = ResolveForCurrentWorld(bAllowPIE).Get();
 
 	if (Object)
@@ -299,7 +307,11 @@ UObject* URemoteControlLevelDependantBinding::Resolve() const
 		}
 	}
 
-	return FindObjectInCounterpartWorld(Object, ECounterpartWorldTarget::PIE);
+	if (bAllowPIE)
+	{
+		return FindObjectInCounterpartWorld(Object, ECounterpartWorldTarget::PIE);
+	}
+	return Object;
 }
 
 bool URemoteControlLevelDependantBinding::IsValid() const
@@ -390,7 +402,7 @@ TSoftObjectPtr<UObject> URemoteControlLevelDependantBinding::ResolveForCurrentWo
 		FSoftObjectPath LevelPath = World;
 #if WITH_EDITOR
 		// Try finding the object using the sub level selection map first.
-		if (World->WorldType == EWorldType::PIE)
+		if (bAllowPIE && World->WorldType == EWorldType::PIE)
 		{
 			FString PIEPackageName = World->GetPackage()->GetPathName();
 			FString OriginalPackage = UWorld::StripPIEPrefixFromPackageName(PIEPackageName, World->StreamingLevelsPrefix);
@@ -410,7 +422,7 @@ TSoftObjectPtr<UObject> URemoteControlLevelDependantBinding::ResolveForCurrentWo
 				}
 
 #if WITH_EDITOR
-				if (World->WorldType == EWorldType::PIE)
+				if (bAllowPIE && World->WorldType == EWorldType::PIE)
 				{
 					bool bPathToActor = true;
 
