[RemoteControl] RemoteControlUI
Reduce Header Height in RemoteControlPanel

diff --git a/Engine/Plugins/VirtualProduction/RemoteControl/Source/RemoteControlUI/Private/UI/Panels/SRCDockPanel.cpp b/Engine/Plugins/VirtualProduction/RemoteControl/Source/RemoteControlUI/Private/UI/Panels/SRCDockPanel.cpp
--- a/Engine/Plugins/VirtualProduction/RemoteControl/Source/RemoteControlUI/Private/UI/Panels/SRCDockPanel.cpp
+++ b/Engine/Plugins/VirtualProduction/RemoteControl/Source/RemoteControlUI/Private/UI/Panels/SRCDockPanel.cpp
@@ -255,7 +255,7 @@ void SRCMinorPanel::Construct(const SRCMinorPanel::FArguments& InArgs)
 				+SHorizontalBox::Slot()
 				.AutoWidth()
 				.VAlign(VAlign_Center)
-				.Padding(4.f, 4.f, 2.f, 4.f)
+				.Padding(4.f, 0.f, 2.f, 0.f)
 				[
 					SAssignNew(LeftHeaderToolbar, SHorizontalBox)
 				]
@@ -264,7 +264,7 @@ void SRCMinorPanel::Construct(const SRCMinorPanel::FArguments& InArgs)
 				+SHorizontalBox::Slot()
 				.AutoWidth()
 				.VAlign(VAlign_Center)
-				.Padding(2.f, 4.f)
+				.Padding(2.f, 0.f)
 				[
 					SNew(STextBlock)
 					.TextStyle(&RCPanelStyle->HeaderTextStyle)
@@ -275,7 +275,7 @@ void SRCMinorPanel::Construct(const SRCMinorPanel::FArguments& InArgs)
 				+SHorizontalBox::Slot()
 				.FillWidth(1.0f)
 				.VAlign(VAlign_Center)
-				.Padding(2.f, 4.f, 5.f, 4.f)
+				.Padding(2.f, 0.f, 5.f, 0.f)
 				[
 					SAssignNew(CenterHeaderToolbar, SHorizontalBox)
 				]
@@ -283,7 +283,7 @@ void SRCMinorPanel::Construct(const SRCMinorPanel::FArguments& InArgs)
 				// Right Toolbar
 				+SHorizontalBox::Slot()
 				.AutoWidth()
-				.Padding(2.f, 4.f, 4.f, 4.f)
+				.Padding(2.f, 0.f, 4.f, 0.f)
 				.VAlign(VAlign_Center)
 				[
 					SAssignNew(RightHeaderToolbar, SHorizontalBox)
diff --git a/Engine/Plugins/VirtualProduction/RemoteControl/Source/RemoteControlUI/Private/UI/SRCModeSwitcher.cpp b/Engine/Plugins/VirtualProduction/RemoteControl/Source/RemoteControlUI/Private/UI/SRCModeSwitcher.cpp
--- a/Engine/Plugins/VirtualProduction/RemoteControl/Source/RemoteControlUI/Private/UI/SRCModeSwitcher.cpp
+++ b/Engine/Plugins/VirtualProduction/RemoteControl/Source/RemoteControlUI/Private/UI/SRCModeSwitcher.cpp
@@ -127,7 +127,7 @@ void SRCModeSwitcher::RegenerateWidgets()
 					.ToolTip(SomeMode.ToolTip)
 					.ToolTipText(SomeMode.ToolTip.IsSet() ? TAttribute<FText>() :
 						SomeMode.DefaultTooltip.IsSet() ? SomeMode.DefaultTooltip : ModeLabel)
-					.Padding(4.f)
+					.Padding(FMargin(4.f, 0.f))
 					[
 						ContentWidget
 					];
diff --git a/Engine/Plugins/VirtualProduction/RemoteControl/Source/RemoteControlUI/Private/UI/SRemoteControlPanel.cpp b/Engine/Plugins/VirtualProduction/RemoteControl/Source/RemoteControlUI/Private/UI/SRemoteControlPanel.cpp
--- a/Engine/Plugins/VirtualProduction/RemoteControl/Source/RemoteControlUI/Private/UI/SRemoteControlPanel.cpp
+++ b/Engine/Plugins/VirtualProduction/RemoteControl/Source/RemoteControlUI/Private/UI/SRemoteControlPanel.cpp
@@ -1928,6 +1928,7 @@ void SRemoteControlPanel::RegisterDefaultToolBar()
 	if (!ToolMenus->IsMenuRegistered(DefaultRemoteControlPanelToolBarName))
 	{
 		UToolMenu* ToolbarBuilder = ToolMenus->RegisterMenu(DefaultRemoteControlPanelToolBarName, NAME_None, EMultiBoxType::SlimHorizontalToolBar);
+		ToolbarBuilder->StyleName = "ContentBrowser.ToolBar";
 
 #if 0
 		ToolbarBuilder->StyleName = "AssetEditorToolbar";
@@ -1975,7 +1976,7 @@ void SRemoteControlPanel::GenerateToolbar()
 				[
 					SNew(SSeparator)
 					.SeparatorImage(FAppStyle::Get().GetBrush("Separator"))
-					.Thickness(2.f)
+					.Thickness(1.5f)
 					.Orientation(EOrientation::Orient_Vertical)
 				];
 
@@ -1988,7 +1989,12 @@ void SRemoteControlPanel::GenerateToolbar()
 				];
 		}
 
-		MiscWidgets = MiscWidgetsHBox;
+		MiscWidgets = SNew(SVerticalBox)
+			+ SVerticalBox::Slot()
+			.AutoHeight()
+			[
+				MiscWidgetsHBox
+			];
 	}
 
 	const URemoteControlSettings* Settings = GetMutableDefault<URemoteControlSettings>();
@@ -2155,6 +2161,7 @@ void SRemoteControlPanel::RegisterAuxiliaryToolBar()
 	if (!ToolMenus->IsMenuRegistered(AuxiliaryRemoteControlPanelToolBarName))
 	{
 		UToolMenu* ToolbarBuilder = ToolMenus->RegisterMenu(AuxiliaryRemoteControlPanelToolBarName, NAME_None, EMultiBoxType::SlimHorizontalToolBar);
+		ToolbarBuilder->StyleName = "ContentBrowser.ToolBar";
 
 #if 0
 		ToolbarBuilder->StyleName = "AssetEditorToolbar";
@@ -2165,19 +2172,29 @@ void SRemoteControlPanel::RegisterAuxiliaryToolBar()
 			const FRemoteControlCommands& Commands = FRemoteControlCommands::Get();
 
 			ToolsSection.AddEntry(FToolMenuEntry::InitWidget("Protocols"
-				, SNew(SAutoResizeButton)
-					.UICommand(FRemoteControlCommands::Get().ToggleProtocolMappings)
-					.ForceSmallIcons_Static(SRemoteControlPanel::ShouldForceSmallIcons)
-					.IconOverride(FSlateIcon(FAppStyle::GetAppStyleSetName(), "LevelEditor.Tabs.StatsViewer"))
-				, Commands.ToggleProtocolMappings->GetLabel()
+				, SNew(SVerticalBox)
+					+ SVerticalBox::Slot()
+					.AutoHeight()
+					[
+						SNew(SAutoResizeButton)
+							.UICommand(FRemoteControlCommands::Get().ToggleProtocolMappings)
+							.ForceSmallIcons_Static(SRemoteControlPanel::ShouldForceSmallIcons)
+							.IconOverride(FSlateIcon(FAppStyle::GetAppStyleSetName(), "LevelEditor.Tabs.StatsViewer"))
+					]
+					, Commands.ToggleProtocolMappings->GetLabel()
 			)
 			);
 
 			ToolsSection.AddEntry(FToolMenuEntry::InitWidget("Logic"
-				, SNew(SAutoResizeButton)
+			, SNew(SVerticalBox)
+				+ SVerticalBox::Slot()
+				.AutoHeight()
+				[
+					SNew(SAutoResizeButton)
 					.UICommand(FRemoteControlCommands::Get().ToggleLogicEditor)
 					.ForceSmallIcons_Static(SRemoteControlPanel::ShouldForceSmallIcons)
-					.IconOverride(FSlateIcon(FAppStyle::GetAppStyleSetName(), TEXT("GraphEditor.StateMachine_24x")))
+					.IconOverride(FSlateIcon(FAppStyle::GetAppStyleSetName(), TEXT("GraphEditor.StateMachine_16x")))
+				]
 				, Commands.ToggleLogicEditor->GetLabel()
 			)
 			);
