[RemoteControl]
Centering PropertyId Text

diff --git a/Engine/Plugins/VirtualProduction/RemoteControl/Source/RemoteControlUI/Private/UI/SRCPanelExposedEntitiesList.cpp b/Engine/Plugins/VirtualProduction/RemoteControl/Source/RemoteControlUI/Private/UI/SRCPanelExposedEntitiesList.cpp
--- a/Engine/Plugins/VirtualProduction/RemoteControl/Source/RemoteControlUI/Private/UI/SRCPanelExposedEntitiesList.cpp
+++ b/Engine/Plugins/VirtualProduction/RemoteControl/Source/RemoteControlUI/Private/UI/SRCPanelExposedEntitiesList.cpp
@@ -285,7 +285,7 @@ void SRCPanelExposedEntitiesList::Construct(const FArguments& InArgs, URemoteCon
 			+ SRCHeaderRow::Column(RemoteControlPresetColumns::PropertyIdentifier)
 			.DefaultLabel(LOCTEXT("RCPresetPropertyIdColumnHeader", "Property ID"))
 			.HAlignHeader(HAlign_Center)
-			.FillWidth(0.1f)
+			.FillWidth(0.15f)
 			.HeaderContentPadding(RCPanelStyle->HeaderRowPadding)
 
 			+ SRCHeaderRow::Column(RemoteControlPresetColumns::OwnerName)
@@ -303,13 +303,13 @@ void SRCPanelExposedEntitiesList::Construct(const FArguments& InArgs, URemoteCon
 			+ SRCHeaderRow::Column(RemoteControlPresetColumns::Description)
 			.DefaultLabel(LOCTEXT("RCPresetDescColumnHeader", "Description"))
 			.HAlignHeader(HAlign_Center)
-			.FillWidth(0.35f)
+			.FillWidth(0.25f)
 			.HeaderContentPadding(RCPanelStyle->HeaderRowPadding)
 
 			+ SRCHeaderRow::Column(RemoteControlPresetColumns::Value)
 			.DefaultLabel(LOCTEXT("RCPresetValueColumnHeader", "Value"))
 			.HAlignHeader(HAlign_Center)
-			.FillWidth(0.4f)
+			.FillWidth(0.35f)
 			.HeaderContentPadding(RCPanelStyle->HeaderRowPadding)
 
 			+ SRCHeaderRow::Column(RemoteControlPresetColumns::Reset)
diff --git a/Engine/Plugins/VirtualProduction/RemoteControl/Source/RemoteControlUI/Private/UI/SRCPanelExposedEntity.cpp b/Engine/Plugins/VirtualProduction/RemoteControl/Source/RemoteControlUI/Private/UI/SRCPanelExposedEntity.cpp
--- a/Engine/Plugins/VirtualProduction/RemoteControl/Source/RemoteControlUI/Private/UI/SRCPanelExposedEntity.cpp
+++ b/Engine/Plugins/VirtualProduction/RemoteControl/Source/RemoteControlUI/Private/UI/SRCPanelExposedEntity.cpp
@@ -511,6 +511,7 @@ TSharedRef<SWidget> SRCPanelExposedEntity::CreateEntityWidget(TSharedPtr<SWidget
 	Args.PropertyIdWidget = SNew(SBox)
 		[
 			SNew(SEditableTextBox)
+			.Justification(ETextJustify::Center)
 			.MinDesiredWidth(50.f)
 			.SelectAllTextWhenFocused(true)
 			.RevertTextOnEscape(true)
