[RemoteControl]
Adding re-ordering inside the All group

diff --git a/Engine/Plugins/VirtualProduction/RemoteControl/Source/RemoteControl/Private/RemoteControlPreset.cpp b/Engine/Plugins/VirtualProduction/RemoteControl/Source/RemoteControl/Private/RemoteControlPreset.cpp
--- a/Engine/Plugins/VirtualProduction/RemoteControl/Source/RemoteControl/Private/RemoteControlPreset.cpp
+++ b/Engine/Plugins/VirtualProduction/RemoteControl/Source/RemoteControl/Private/RemoteControlPreset.cpp
@@ -378,6 +378,76 @@ void FRemoteControlPresetLayout::SwapFields(const FFieldSwapArgs& FieldSwapArgs)
 	}
 }
 
+void FRemoteControlPresetLayout::SwapFieldsDefaultGroup(const FFieldSwapArgs& InFieldSwapArgs, const FGuid InFieldRealGroup, TArray<FGuid> InEntities)
+{
+	FRemoteControlPresetGroup* DragOriginGroup = GetGroup(InFieldRealGroup);
+	FRemoteControlPresetGroup* DragTargetGroup = GetGroup(InFieldSwapArgs.TargetGroupId);
+
+	if (!DragOriginGroup || !DragTargetGroup)
+	{
+		return;
+	}
+
+	int32 DragOriginFieldIndex = InEntities.IndexOfByKey(InFieldSwapArgs.DraggedFieldId);
+	int32 DragTargetFieldIndex = InEntities.IndexOfByKey(InFieldSwapArgs.TargetFieldId);
+
+	if (DragOriginFieldIndex == INDEX_NONE)
+	{
+		return;
+	}
+
+	if (DragTargetFieldIndex != INDEX_NONE)
+	{
+		if (DragTargetFieldIndex > DragOriginFieldIndex)
+		{
+			DragTargetFieldIndex += 1;
+		}
+		else
+		{
+			DragOriginFieldIndex += 1;
+		}
+
+		// Here we don't want to trigger add/delete delegates since the fields just get moved around.
+		TArray<FGuid>& Fields = InEntities;
+
+		FRCTransactionListenerHelper<const FGuid&, const TArray<FGuid>&>(ERCTransaction::Undo, Owner->GetPresetId(), OnFieldOrderChanged(), InFieldSwapArgs.TargetGroupId, Fields);
+
+		Fields.Insert(InFieldSwapArgs.DraggedFieldId, DragTargetFieldIndex);
+		Fields.Swap(DragTargetFieldIndex, DragOriginFieldIndex);
+		Fields.RemoveAt(DragOriginFieldIndex);
+	
+		Owner->CacheLayoutData();
+		OnFieldOrderChangedDelegate.Broadcast(GetDefaultGroup().Id, Fields);
+		Owner->OnPresetLayoutModified().Broadcast(Owner.Get());
+
+		FRCTransactionListenerHelper<const FGuid&, const TArray<FGuid>&>(ERCTransaction::Redo, Owner->GetPresetId(), OnFieldOrderChanged(), InFieldSwapArgs.TargetGroupId, Fields);
+
+		FRCTransactionListenerHelper<URemoteControlPreset*>(ERCTransaction::Undo, Owner->GetPresetId(), Owner->OnPresetLayoutModified(), Owner.Get());
+		FRCTransactionListenerHelper<URemoteControlPreset*>(ERCTransaction::Redo, Owner->GetPresetId(), Owner->OnPresetLayoutModified(), Owner.Get());
+	}
+	else if (!IsDefaultGroup(DragTargetGroup->Id))
+	{
+		DragOriginGroup->AccessFields().RemoveSwap(InFieldSwapArgs.DraggedFieldId, true);
+
+		DragTargetFieldIndex = DragTargetFieldIndex == INDEX_NONE ? 0 : DragTargetFieldIndex;
+		DragTargetGroup->AccessFields().Insert(InFieldSwapArgs.DraggedFieldId, DragTargetFieldIndex);
+
+		FRCTransactionListenerHelper<const FGuid&, const FGuid&, int32>(ERCTransaction::Undo, Owner->GetPresetId(), OnFieldDeleted(), InFieldSwapArgs.TargetGroupId, InFieldSwapArgs.DraggedFieldId, DragTargetFieldIndex);
+		FRCTransactionListenerHelper<const FGuid&, const FGuid&, int32>(ERCTransaction::Undo, Owner->GetPresetId(), OnFieldAdded(), InFieldSwapArgs.OriginGroupId, InFieldSwapArgs.DraggedFieldId, DragOriginFieldIndex);
+
+		FRCTransactionListenerHelper<const FGuid&, const FGuid&, int32>(ERCTransaction::Redo, Owner->GetPresetId(), OnFieldDeleted(), InFieldSwapArgs.OriginGroupId, InFieldSwapArgs.DraggedFieldId, DragOriginFieldIndex);
+		FRCTransactionListenerHelper<const FGuid&, const FGuid&, int32>(ERCTransaction::Redo, Owner->GetPresetId(), OnFieldAdded(), InFieldSwapArgs.TargetGroupId, InFieldSwapArgs.DraggedFieldId, DragTargetFieldIndex);
+
+		Owner->CacheLayoutData();
+		OnFieldDeletedDelegate.Broadcast(InFieldSwapArgs.OriginGroupId, InFieldSwapArgs.DraggedFieldId, DragOriginFieldIndex);
+		OnFieldAddedDelegate.Broadcast(InFieldSwapArgs.TargetGroupId, InFieldSwapArgs.DraggedFieldId, DragTargetFieldIndex);
+		Owner->OnPresetLayoutModified().Broadcast(Owner.Get());
+
+		FRCTransactionListenerHelper<URemoteControlPreset*>(ERCTransaction::Undo, Owner->GetPresetId(), Owner->OnPresetLayoutModified(), Owner.Get());
+		FRCTransactionListenerHelper<URemoteControlPreset*>(ERCTransaction::Redo, Owner->GetPresetId(), Owner->OnPresetLayoutModified(), Owner.Get());
+	}
+}
+
 void FRemoteControlPresetLayout::DeleteGroup(FGuid GroupId)
 {
 
diff --git a/Engine/Plugins/VirtualProduction/RemoteControl/Source/RemoteControl/Public/RemoteControlPreset.h b/Engine/Plugins/VirtualProduction/RemoteControl/Source/RemoteControl/Public/RemoteControlPreset.h
--- a/Engine/Plugins/VirtualProduction/RemoteControl/Source/RemoteControl/Public/RemoteControlPreset.h
+++ b/Engine/Plugins/VirtualProduction/RemoteControl/Source/RemoteControl/Public/RemoteControlPreset.h
@@ -207,7 +207,10 @@ struct REMOTECONTROL_API FRemoteControlPresetLayout
 
 	/** Get or create the default group. */
 	FRemoteControlPresetGroup& GetDefaultGroup();
-
+	/** Return the DefaultGroupOrder */
+	TArray<FGuid>& GetDefaultGroupOrder() { return DefaultGroupOrder; }
+	/** Set the All group order */
+	void SetDefaultGroupOrder(const TArray<FGuid>& InNewDefaultGroupOrder) { DefaultGroupOrder = InNewDefaultGroupOrder; }
 	/** Returns true when the given group id is a default one. */
 	bool IsDefaultGroup(FGuid GroupId) const;
 	
@@ -258,6 +261,9 @@ struct REMOTECONTROL_API FRemoteControlPresetLayout
 	/** Swap fields across groups or in the same one. */
 	void SwapFields(const FFieldSwapArgs& FieldSwapArgs);
 
+	/** Swap fields across groups or in the same one for the default group. */
+	void SwapFieldsDefaultGroup(const FFieldSwapArgs& FieldSwapArgs, const FGuid InFieldRealGroup, TArray<FGuid> InEntities);
+
 	/** Delete a group from the layout. */
 	void DeleteGroup(FGuid GroupId);
 
@@ -314,6 +320,10 @@ private:
 	/** Create a group by providing a name and ID. */
 	FRemoteControlPresetGroup& CreateGroupInternal(FName GroupName, FGuid GroupId);
 
+	/** Keep the ALL group order since its not saved in the normal workflow */
+	UPROPERTY()
+	TArray<FGuid> DefaultGroupOrder;
+
 	/** The list of groups under this layout. */
 	UPROPERTY()
 	TArray<FRemoteControlPresetGroup> Groups;
diff --git a/Engine/Plugins/VirtualProduction/RemoteControl/Source/RemoteControlUI/Private/UI/SRCPanelExposedEntitiesList.cpp b/Engine/Plugins/VirtualProduction/RemoteControl/Source/RemoteControlUI/Private/UI/SRCPanelExposedEntitiesList.cpp
--- a/Engine/Plugins/VirtualProduction/RemoteControl/Source/RemoteControlUI/Private/UI/SRCPanelExposedEntitiesList.cpp
+++ b/Engine/Plugins/VirtualProduction/RemoteControl/Source/RemoteControlUI/Private/UI/SRCPanelExposedEntitiesList.cpp
@@ -69,10 +69,10 @@ public:
 		SLATE_STYLE_ARGUMENT(FTableRowStyle, Style)
 
 		// Low level DragAndDrop
+		SLATE_EVENT(FOnAcceptDrop, OnAcceptDrop)
 		SLATE_EVENT(FOnDragDetected, OnDragDetected)
 		SLATE_EVENT(FOnTableRowDragEnter, OnDragEnter)
 		SLATE_EVENT(FOnTableRowDragLeave, OnDragLeave)
-		SLATE_EVENT(FOnTableRowDrop, OnDrop)
 
 	SLATE_END_ARGS()
 
@@ -83,10 +83,15 @@ public:
 
 		FSuperRowType::FArguments SuperArgs = FSuperRowType::FArguments();
 
+		SuperArgs.OnCanAcceptDrop_Lambda([this] (const FDragDropEvent& InDragDropEvent, EItemDropZone InDropZone, TSharedPtr<SRCPanelTreeNode> Node)
+		{
+			return InDropZone;
+		});
+
+		SuperArgs.OnAcceptDrop(InArgs._OnAcceptDrop);
 		SuperArgs.OnDragDetected(InArgs._OnDragDetected);
 		SuperArgs.OnDragEnter(InArgs._OnDragEnter);
 		SuperArgs.OnDragLeave(InArgs._OnDragLeave);
-		SuperArgs.OnDrop(InArgs._OnDrop);
 
 		SuperArgs.ExpanderStyleSet(&FCoreStyle::Get());
 		SuperArgs.Padding(InArgs._Padding);
@@ -275,13 +280,6 @@ void SRCPanelExposedEntitiesList::Construct(const FArguments& InArgs, URemoteCon
 			.CanSelectGeneratedColumn(true) //To show/hide columns
 			.HiddenColumnsList(HiddenColumns) // List of columns to hide by default. User can un-hide via context menu list
 
-			+ SRCHeaderRow::Column(RemoteControlPresetColumns::DragDropHandle)
-			.DefaultLabel(LOCTEXT("RCPresetDragDropHandleColumnHeader", ""))
-			.FixedWidth(25.f)
-			.HeaderContentPadding(RCPanelStyle->HeaderRowPadding)
-			.ShouldGenerateWidget(true)
-			.ShouldGenerateSubMenuEntry(false)
-
 			+ SRCHeaderRow::Column(RemoteControlPresetColumns::PropertyIdentifier)
 			.DefaultLabel(LOCTEXT("RCPresetPropertyIdColumnHeader", "Property ID"))
 			.HAlignHeader(HAlign_Center)
@@ -737,6 +735,7 @@ void SRCPanelExposedEntitiesList::GenerateListWidgets()
 {
 	TRACE_CPUPROFILER_EVENT_SCOPE(SRCPanelExposedEntitiesList::GenerateListWidgets);
 
+	TArray<FGuid> OrderMap = Preset->Layout.GetDefaultGroupOrder();
 	FieldWidgetMap.Reset();
 
 	for (TWeakPtr<FRemoteControlEntity> WeakEntity : Preset->GetExposedEntities())
@@ -753,6 +752,25 @@ void SRCPanelExposedEntitiesList::GenerateListWidgets()
 			FieldWidgetMap.Add(Entity->GetId(), FRemoteControlUIModule::Get().GenerateEntityWidget(Args));
 		}
 	}
+
+	// We order the ALL group here otherwise the order would be the same order of the ExposedEntities
+	if (OrderMap.Num() && FieldEntities.Num())
+	{
+		FieldWidgetMap.KeySort(
+			[&OrderMap]
+			(const FGuid& A, const FGuid& B)
+			{
+				if (!OrderMap.Contains(A))
+				{
+					return false;
+				}
+				if (!OrderMap.Contains(B))
+				{
+					return true;
+				}
+				return OrderMap.IndexOfByKey(A) < OrderMap.IndexOfByKey(B);
+			});
+	}
 }
 
 void SRCPanelExposedEntitiesList::GenerateListWidgets(const FRemoteControlPresetGroup& FromGroup)
@@ -813,16 +831,16 @@ TSharedRef<ITableRow> SRCPanelExposedEntitiesList::OnGenerateRow(TSharedPtr<SRCP
 {
 	const TSharedRef<SWidget> NodeWidget = Node->AsShared();
 	
-	auto OnDropLambda = [this, Node]
-	(const FDragDropEvent& Event)
+	auto OnAcceptDropLambda = [this]
+	(const FDragDropEvent& InDragDropEvent, EItemDropZone InDropZone, const TSharedPtr<SRCPanelTreeNode>& InNode)
 	{
-		if (const TSharedPtr<FExposedEntityDragDrop> DragDropOp = Event.GetOperationAs<FExposedEntityDragDrop>())
+		if (const TSharedPtr<FExposedEntityDragDrop> DragDropOp = InDragDropEvent.GetOperationAs<FExposedEntityDragDrop>())
 		{
-			if (const TSharedPtr<SRCPanelGroup> Group = FindGroupById(GetGroupId(Node->GetRCId())))
+			if (const TSharedPtr<SRCPanelGroup> Group = FindGroupById(GetGroupId(InNode->GetRCId())))
 			{
 				if (DragDropOp->IsOfType<FExposedEntityDragDrop>())
 				{
-					return OnDropOnGroup(DragDropOp, Node, Group);
+					return OnDropOnGroup(DragDropOp, InNode, Group);
 				}
 				else if (DragDropOp->IsOfType<FFieldGroupDragDropOp>())
 				{
@@ -848,9 +866,19 @@ TSharedRef<ITableRow> SRCPanelExposedEntitiesList::OnGenerateRow(TSharedPtr<SRCP
 		constexpr float LeftPadding = 3.f;
 		const FMargin Margin = Node->GetRCType() == SRCPanelTreeNode::FieldChild ? FMargin(LeftPadding + 10.f, 1.f, 1.f, 1.f) : FMargin(LeftPadding, 1.f, 1.f, 1.f);
 		return SNew(SEntityRow, OwnerTable)
+			.OnDragDetected_Lambda([this, Node] (const FGeometry&, const FPointerEvent&)
+			{
+				if (Node && Node->GetRCType() == SRCPanelTreeNode::Field)
+				{
+					const TSharedRef<FExposedEntityDragDrop> DragDropOp = MakeShared<FExposedEntityDragDrop>(Node->GetDragAndDropWidget(), Node->GetRCId());
+					DragDropOp->Construct();
+					return FReply::Handled().BeginDragDrop(DragDropOp);
+				}
+				return FReply::Unhandled();
+			})
 			.OnDragEnter_Lambda([Node](const FDragDropEvent& Event) { if (Node && Node->GetRCType() == SRCPanelTreeNode::Field) StaticCastSharedPtr<SRCPanelExposedField>(Node)->SetIsHovered(true); })
 			.OnDragLeave_Lambda([Node](const FDragDropEvent& Event) { if (Node && Node->GetRCType() == SRCPanelTreeNode::Field) StaticCastSharedPtr<SRCPanelExposedField>(Node)->SetIsHovered(false); })
-			.OnDrop_Lambda(OnDropLambda)
+			.OnAcceptDrop_Lambda(OnAcceptDropLambda)
 			.Padding(Margin)
 			.Style(&RCPanelStyle->TableRowStyle)
 			.ActiveProtocol_Lambda([this]() { return ActiveProtocol; })
@@ -981,14 +1009,35 @@ FReply SRCPanelExposedEntitiesList::OnDropOnGroup(const TSharedPtr<FDragDropOper
 
 	if (DragDropOperation->IsOfType<FExposedEntityDragDrop>())
 	{
-		if (Preset->Layout.IsDefaultGroup(DragTargetGroup->GetRCId()))
-		{
-			// We do not add fields to the default group.
-			return FReply::Unhandled();
-		}
-
 		if (TSharedPtr<FExposedEntityDragDrop> DragDropOp = StaticCastSharedPtr<FExposedEntityDragDrop>(DragDropOperation))
 		{
+			if (Preset->Layout.IsDefaultGroup(GetSelectedGroup()->GetRCId()))
+			{
+				FRemoteControlPresetLayout::FFieldSwapArgs Args;
+				Args.OriginGroupId = GetSelectedGroup()->GetRCId();
+				Args.TargetGroupId = DragTargetGroup->GetRCId();
+				Args.DraggedFieldId = DragDropOp->GetId();
+
+				if (TargetEntity)
+				{
+					Args.TargetFieldId = TargetEntity->GetRCId();
+				}
+
+				FScopedTransaction Transaction(LOCTEXT("MoveFieldDefaultGroup", "Move exposed field"));
+				Preset->Modify();
+				TArray<FGuid> OrderEntities;
+
+				for (const TSharedPtr<SRCPanelTreeNode>& Entity : FieldEntities)
+				{
+					OrderEntities.Add(Entity->GetRCId());
+				}
+
+				Preset->Layout.SwapFieldsDefaultGroup(Args, GetGroupId(DragDropOp->GetId()), OrderEntities);
+				constexpr bool bForceMouseClick = true;
+				SetSelection(FindGroupById(Args.OriginGroupId), bForceMouseClick);
+				return FReply::Handled();
+			}
+
 			FGuid DragOriginGroupId = GetGroupId(DragDropOp->GetId());
 			if (!DragOriginGroupId.IsValid())
 			{
@@ -1017,7 +1066,7 @@ FReply SRCPanelExposedEntitiesList::OnDropOnGroup(const TSharedPtr<FDragDropOper
 			Preset->Modify();
 			Preset->Layout.SwapFields(Args);
 			constexpr bool bForceMouseClick = true;
-			SetSelection(FindGroupById(Args.TargetGroupId), bForceMouseClick);
+			SetSelection(FindGroupById(Args.OriginGroupId), bForceMouseClick);
 			return FReply::Handled();
 		}
 	}
@@ -1371,12 +1420,28 @@ void SRCPanelExposedEntitiesList::OnFieldOrderChanged(const FGuid& GroupId, cons
 
 		if (*Group)
 		{
-			(*Group)->GetNodes().Sort(
-			[&OrderMap]
-			(const TSharedPtr<SRCPanelTreeNode>& A, const TSharedPtr<SRCPanelTreeNode>& B)
-				{
-					return OrderMap.FindChecked(A->GetRCId()) < OrderMap.FindChecked(B->GetRCId());
-				});
+			if (Preset && Preset->Layout.IsDefaultGroup((*Group)->GetRCId()))
+			{
+				FieldWidgetMap.ValueSort(
+				[&OrderMap]
+					(const TSharedPtr<SRCPanelTreeNode>& A, const TSharedPtr<SRCPanelTreeNode>& B)
+					{
+						return OrderMap.FindChecked(A->GetRCId()) < OrderMap.FindChecked(B->GetRCId());
+					});
+				TArray<FGuid> OrderedGuid;
+				FieldWidgetMap.GenerateKeyArray(OrderedGuid);
+				// We save the ALL group order here because it is not saved in the normal workflow
+				Preset->Layout.SetDefaultGroupOrder(OrderedGuid);
+			}
+			else
+			{
+				(*Group)->GetNodes().Sort(
+				[&OrderMap]
+				(const TSharedPtr<SRCPanelTreeNode>& A, const TSharedPtr<SRCPanelTreeNode>& B)
+					{
+						return OrderMap.FindChecked(A->GetRCId()) < OrderMap.FindChecked(B->GetRCId());
+					});
+			}
 		}
 	}
 
@@ -1554,11 +1619,7 @@ SHeaderRow::FColumn::FArguments SRCPanelExposedEntitiesList::CreateColumn(const
 
 int32 SRCPanelExposedEntitiesList::GetColumnIndex(const FName& ForColumn) const
 {
-	if (ForColumn == RemoteControlPresetColumns::PropertyIdentifier)
-	{
-		return GetColumnIndex_Internal(ForColumn, RemoteControlPresetColumns::DragDropHandle, ERCColumn::ERC_After);
-	}
-	else if (ForColumn == RemoteControlPresetColumns::Mask)
+	if (ForColumn == RemoteControlPresetColumns::Mask)
 	{
 		return GetColumnIndex_Internal(ForColumn, RemoteControlPresetColumns::Description, ERCColumn::ERC_After);
 	}
diff --git a/Engine/Plugins/VirtualProduction/RemoteControl/Source/RemoteControlUI/Private/UI/SRCPanelExposedEntity.cpp b/Engine/Plugins/VirtualProduction/RemoteControl/Source/RemoteControlUI/Private/UI/SRCPanelExposedEntity.cpp
--- a/Engine/Plugins/VirtualProduction/RemoteControl/Source/RemoteControlUI/Private/UI/SRCPanelExposedEntity.cpp
+++ b/Engine/Plugins/VirtualProduction/RemoteControl/Source/RemoteControlUI/Private/UI/SRCPanelExposedEntity.cpp
@@ -501,13 +501,6 @@ TSharedRef<SWidget> SRCPanelExposedEntity::CreateEntityWidget(TSharedPtr<SWidget
 		.Padding(0.0f)
 		.BorderImage(this, &SRCPanelExposedEntity::GetBorderImage);
 	
-	Args.DragHandle = SNew(SBox)
-		.Visibility(this, &SRCPanelExposedEntity::GetVisibilityAccordingToLiveMode, EVisibility::Collapsed)
-		[
-			SNew(SRCPanelDragHandle<FExposedEntityDragDrop>, GetRCId())
-			.Widget(Widget)
-		];
-
 	Args.PropertyIdWidget = SNew(SBox)
 		[
 			SNew(SEditableTextBox)
@@ -590,7 +583,7 @@ TSharedRef<SWidget> SRCPanelExposedEntity::CreateEntityWidget(TSharedPtr<SWidget
 
 	Args.ResetButton = ResetWidget;
 
-	Widget->SetContent(MakeDragAndDropWidget(Args));
+	Widget->SetContent(MakeNodeWidget(Args));
 	return Widget;
 }
 
diff --git a/Engine/Plugins/VirtualProduction/RemoteControl/Source/RemoteControlUI/Private/UI/SRCPanelTreeNode.cpp b/Engine/Plugins/VirtualProduction/RemoteControl/Source/RemoteControlUI/Private/UI/SRCPanelTreeNode.cpp
--- a/Engine/Plugins/VirtualProduction/RemoteControl/Source/RemoteControlUI/Private/UI/SRCPanelTreeNode.cpp
+++ b/Engine/Plugins/VirtualProduction/RemoteControl/Source/RemoteControlUI/Private/UI/SRCPanelTreeNode.cpp
@@ -14,7 +14,6 @@
 #define LOCTEXT_NAMESPACE "RemoteControlPanelNode"
 
 TSet<FName> SRCPanelTreeNode::DefaultColumns = {
-	RemoteControlPresetColumns::DragDropHandle,
 	RemoteControlPresetColumns::PropertyIdentifier,
 	RemoteControlPresetColumns::OwnerName,
 	RemoteControlPresetColumns::SubobjectPath,
@@ -51,11 +50,7 @@ const bool SRCPanelTreeNode::SupportsProtocol(const FName& InProtocolName) const
 
 TSharedRef<SWidget> SRCPanelTreeNode::GetWidget(const FName ForColumnName, const FName InActiveProtocol)
 {
-	if (ForColumnName == RemoteControlPresetColumns::DragDropHandle)
-	{
-		return DragHandleWidget.ToSharedRef();
-	}
-	else if (ForColumnName == RemoteControlPresetColumns::PropertyIdentifier)
+	if (ForColumnName == RemoteControlPresetColumns::PropertyIdentifier)
 	{
 		return PropertyIdWidget.ToSharedRef();
 	}
@@ -83,6 +78,40 @@ TSharedRef<SWidget> SRCPanelTreeNode::GetWidget(const FName ForColumnName, const
 	return SNullWidget::NullWidget;
 }
 
+TSharedRef<SWidget> SRCPanelTreeNode::GetDragAndDropWidget()
+{
+	if (NodeNameWidget && NodeValueWidget)
+	{
+		const TSharedRef<SWidget> LeftColumn = SNew(SHorizontalBox)
+			.Clipping(EWidgetClipping::OnDemand)
+			// Field name
+			+ SHorizontalBox::Slot()
+			.VAlign(VAlign_Center)
+			.AutoWidth()
+			[
+				NodeNameWidget.ToSharedRef()
+			];
+
+		const TSharedRef<SWidget> RightColumn = SNew(SHorizontalBox)
+			.Clipping(EWidgetClipping::OnDemand)
+			// Node Value
+			+ SHorizontalBox::Slot()
+			.HAlign(HAlign_Fill)
+			[
+				SNew(SBox)
+				.HAlign(HAlign_Left)
+				.VAlign(VAlign_Center)
+				[
+					NodeValueWidget.ToSharedRef()
+				]
+			];
+
+		return MakeSplitRow(LeftColumn, RightColumn);
+	}
+
+	return SNullWidget::NullWidget;
+}
+
 TSharedRef<SWidget> SRCPanelTreeNode::MakeSplitRow(TSharedRef<SWidget> LeftColumn, TSharedRef<SWidget> RightColumn)
 {
 	TAttribute<float> LeftColumnAttribute;
@@ -114,14 +143,6 @@ TSharedRef<SWidget> SRCPanelTreeNode::MakeNodeWidget(const FMakeNodeWidgetArgs&
 
 	TSharedRef<SWidget> LeftColumn = SNew(SHorizontalBox)
 		.Clipping(EWidgetClipping::OnDemand)
-		// Drag and drop handle
-		+ SHorizontalBox::Slot()
-		.HAlign(HAlign_Left)
-		.VAlign(VAlign_Center)
-		.AutoWidth()
-		[
-			DragHandleWidget.ToSharedRef()
-		]
 		// Link identifier widget
 		+ SHorizontalBox::Slot()
 		.VAlign(VAlign_Center)
@@ -175,43 +196,10 @@ TSharedRef<SWidget> SRCPanelTreeNode::MakeNodeWidget(const FMakeNodeWidgetArgs&
 	return MakeSplitRow(LeftColumn, RightColumn);
 }
 
-TSharedRef<SWidget> SRCPanelTreeNode::MakeDragAndDropWidget(const FMakeNodeWidgetArgs& Args)
-{
-	MakeNodeWidgets(Args);
-
-	const TSharedRef<SWidget> LeftColumn = SNew(SHorizontalBox)
-		.Clipping(EWidgetClipping::OnDemand)
-		// Field name
-		+ SHorizontalBox::Slot()
-		.VAlign(VAlign_Center)
-		.AutoWidth()
-		[
-			NodeNameWidget.ToSharedRef()
-		];
-
-	const TSharedRef<SWidget> RightColumn = SNew(SHorizontalBox)
-		.Clipping(EWidgetClipping::OnDemand)
-		// Node Value
-		+ SHorizontalBox::Slot()
-		.HAlign(HAlign_Fill)
-		[
-			SNew(SBox)
-			.HAlign(HAlign_Left)
-			.VAlign(VAlign_Center)
-			[
-				NodeValueWidget.ToSharedRef()
-			]
-		];
-
-	return MakeSplitRow(LeftColumn, RightColumn);
-}
-
 void SRCPanelTreeNode::MakeNodeWidgets(const FMakeNodeWidgetArgs& Args)
 {
 	auto WidgetOrNull = [](const TSharedPtr<SWidget>& Widget) {return Widget ? Widget.ToSharedRef() : SNullWidget::NullWidget; };
 
-	DragHandleWidget = WidgetOrNull(Args.DragHandle);
-
 	PropertyIdWidget = WidgetOrNull(Args.PropertyIdWidget);
 
 	NodeOwnerNameWidget = WidgetOrNull(Args.OwnerNameWidget);
diff --git a/Engine/Plugins/VirtualProduction/RemoteControl/Source/RemoteControlUI/Private/UI/SRCPanelTreeNode.h b/Engine/Plugins/VirtualProduction/RemoteControl/Source/RemoteControlUI/Private/UI/SRCPanelTreeNode.h
--- a/Engine/Plugins/VirtualProduction/RemoteControl/Source/RemoteControlUI/Private/UI/SRCPanelTreeNode.h
+++ b/Engine/Plugins/VirtualProduction/RemoteControl/Source/RemoteControlUI/Private/UI/SRCPanelTreeNode.h
@@ -17,7 +17,6 @@ struct SRCPanelExposedMaterial;
 
 namespace RemoteControlPresetColumns
 {
-	static FName DragDropHandle = TEXT("DragDropHandle");
 	static FName PropertyIdentifier = TEXT("PropertyID");
 	static FName OwnerName = TEXT("OwnerName");
 	static FName SubobjectPath = TEXT("Subobject Path");
@@ -84,11 +83,12 @@ struct SRCPanelTreeNode : public SCompoundWidget, public IHasProtocolExtensibili
 	virtual void SetHighlightText(const FText& InHightlightText = FText::GetEmpty()) {};
 	/** Retrieves the referenced widget corresponding to the given column name. */
 	virtual TSharedRef<SWidget> GetWidget(const FName ForColumnName, const FName InActiveProtocol);
+	/** Retrieves the DragAndDropWidget if possible otherwise returns a NullWidget */
+	TSharedRef<SWidget> GetDragAndDropWidget();
 
 protected:
 	struct FMakeNodeWidgetArgs
 	{
-		TSharedPtr<SWidget> DragHandle;
 		TSharedPtr<SWidget> PropertyIdWidget;
 		TSharedPtr<SWidget> OwnerNameWidget;
 		TSharedPtr<SWidget> SubObjectPathWidget;
@@ -101,8 +101,6 @@ protected:
 	TSharedRef<SWidget> MakeSplitRow(TSharedRef<SWidget> LeftColumn, TSharedRef<SWidget> RightColumn);
 	/** Create a widget that represents a node in the panel tree hierarchy. */
 	TSharedRef<SWidget> MakeNodeWidget(const FMakeNodeWidgetArgs& Args);
-	/** Create a widget that represents a node for the drag and drop. */
-	TSharedRef<SWidget> MakeDragAndDropWidget(const FMakeNodeWidgetArgs& Args);
 	/** Creates cached copies of underlying widgets. */
 	void MakeNodeWidgets(const FMakeNodeWidgetArgs& Args);
 
@@ -126,9 +124,6 @@ protected:
 
 private:
 
-	/** Cached widget of drag handle. */
-	TSharedPtr<SWidget> DragHandleWidget;
-
 	/** Cached widget of PropertyId identifier. */
 	TSharedPtr<SWidget> PropertyIdWidget;
 
