[RemoteControl]
Avoid ensure when adding action with the same controller value

diff --git a/Engine/Plugins/VirtualProduction/RemoteControl/Source/RemoteControlUI/Private/UI/Action/SRCActionPanel.cpp b/Engine/Plugins/VirtualProduction/RemoteControl/Source/RemoteControlUI/Private/UI/Action/SRCActionPanel.cpp
--- a/Engine/Plugins/VirtualProduction/RemoteControl/Source/RemoteControlUI/Private/UI/Action/SRCActionPanel.cpp
+++ b/Engine/Plugins/VirtualProduction/RemoteControl/Source/RemoteControlUI/Private/UI/Action/SRCActionPanel.cpp
@@ -321,15 +321,18 @@ URCAction* SRCActionPanel::AddAction(const TSharedRef<const FRemoteControlField>
 
 			URCAction* NewAction = BehaviourItem->AddAction(InRemoteControlField);
 
-			AddNewActionToList(NewAction);
-
-			// Broadcast new Action to other panels
-			if (const TSharedPtr<SRemoteControlPanel> RemoteControlPanel = GetRemoteControlPanel())
+			if (NewAction)
 			{
-				RemoteControlPanel->OnActionAdded.Broadcast(NewAction);
-			}
+				AddNewActionToList(NewAction);
 
-			return NewAction;
+				// Broadcast new Action to other panels
+				if (const TSharedPtr<SRemoteControlPanel> RemoteControlPanel = GetRemoteControlPanel())
+				{
+					RemoteControlPanel->OnActionAdded.Broadcast(NewAction);
+				}
+
+				return NewAction;
+			}
 		}
 	}
 
