[RemoteControl]
Fix Color Controller lag with a bind behaviour and a FLinearColor property in the Action

diff --git a/Engine/Plugins/VirtualProduction/RemoteControl/Source/RemoteControlUI/Private/UI/Controller/RCControllerModel.cpp b/Engine/Plugins/VirtualProduction/RemoteControl/Source/RemoteControlUI/Private/UI/Controller/RCControllerModel.cpp
--- a/Engine/Plugins/VirtualProduction/RemoteControl/Source/RemoteControlUI/Private/UI/Controller/RCControllerModel.cpp
+++ b/Engine/Plugins/VirtualProduction/RemoteControl/Source/RemoteControlUI/Private/UI/Controller/RCControllerModel.cpp
@@ -46,6 +46,18 @@ TSharedRef<SWidget> FRCControllerModel::GetWidget() const
 {
 	const FNodeWidgets NodeWidgets = DetailTreeNodeWeakPtr.Pin()->CreateNodeWidgets();
 
+	// We need to add this metadata to the ColorController to avoid the update when dragging causing a lot of lag
+	const TSharedPtr<IPropertyHandle> PropertyHandle = DetailTreeNodeWeakPtr.Pin()->CreatePropertyHandle();
+	if (PropertyHandle.IsValid())
+	{
+		FStructProperty* StructProperty = CastField<FStructProperty>(PropertyHandle->GetProperty());
+		if (StructProperty && StructProperty->Struct &&
+			StructProperty->Struct.GetFName() == FName("Color") &&
+			!StructProperty->HasMetaData("OnlyUpdateOnInteractionEnd"))
+		{
+			StructProperty->AppendMetaData({{FName("OnlyUpdateOnInteractionEnd"), TEXT("true")}});
+		}
+	}
 	const TSharedRef<SHorizontalBox> FieldWidget = SNew(SHorizontalBox);
 	if (VirtualPropertyWeakPtr.IsValid())
 	{
