[RemoteControl]
Fix callback bind everytime the ControllerPanelList call the Reset, causing the callback being called multiple times.

diff --git a/Engine/Plugins/VirtualProduction/RemoteControl/Source/RemoteControlUI/Private/UI/Controller/SRCControllerPanelList.cpp b/Engine/Plugins/VirtualProduction/RemoteControl/Source/RemoteControlUI/Private/UI/Controller/SRCControllerPanelList.cpp
--- a/Engine/Plugins/VirtualProduction/RemoteControl/Source/RemoteControlUI/Private/UI/Controller/SRCControllerPanelList.cpp
+++ b/Engine/Plugins/VirtualProduction/RemoteControl/Source/RemoteControlUI/Private/UI/Controller/SRCControllerPanelList.cpp
@@ -436,7 +436,10 @@ void SRCControllerPanelList::Reset()
 	check(Preset);
 
 	PropertyRowGenerator->SetStructure(Preset->GetControllerContainerStructOnScope());
-	PropertyRowGenerator->OnFinishedChangingProperties().AddSP(this, &SRCControllerPanelList::OnFinishedChangingProperties);
+	if (!PropertyRowGenerator->OnFinishedChangingProperties().IsBoundToObject(this))
+	{
+		PropertyRowGenerator->OnFinishedChangingProperties().AddSP(this, &SRCControllerPanelList::OnFinishedChangingProperties);
+	}
 
 	// Generator should be moved to separate class
 	TArray<TSharedRef<IDetailTreeNode>> RootTreeNodes = PropertyRowGenerator->GetRootTreeNodes();
