[RemoteControl/AvalancheMediaEditor]
Fixing Texture Controller not updating correctly, using custom update instead of the PropertyHandle to update it.

diff --git a/Engine/Plugins/VirtualProduction/RemoteControl/Source/RemoteControlUI/Private/RemoteControlUIModule.cpp b/Engine/Plugins/VirtualProduction/RemoteControl/Source/RemoteControlUI/Private/RemoteControlUIModule.cpp
--- a/Engine/Plugins/VirtualProduction/RemoteControl/Source/RemoteControlUI/Private/RemoteControlUIModule.cpp
+++ b/Engine/Plugins/VirtualProduction/RemoteControl/Source/RemoteControlUI/Private/RemoteControlUIModule.cpp
@@ -1014,13 +1014,13 @@ void FRemoteControlUIModule::SelectObjects(const TArray<UObject*>& Objects) cons
 	}
 }
 
-TSharedPtr<SWidget> FRemoteControlUIModule::CreateCustomControllerWidget(URCVirtualPropertyBase* InController) const
+TSharedPtr<SWidget> FRemoteControlUIModule::CreateCustomControllerWidget(URCVirtualPropertyBase* InController, TSharedPtr<IPropertyHandle> InOriginalPropertyHandle) const
 {
 	TSharedPtr<SWidget> CustomControllerWidget;
 	const FString& CustomName = UE::RCCustomControllers::GetCustomControllerTypeName(InController);
 	if (CustomName == UE::RCCustomControllers::CustomTextureControllerName)
 	{
-		CustomControllerWidget = SNew(SCustomTextureControllerWidget, InController);
+		CustomControllerWidget = SNew(SCustomTextureControllerWidget, InOriginalPropertyHandle);
 	}
 
 	return CustomControllerWidget;
diff --git a/Engine/Plugins/VirtualProduction/RemoteControl/Source/RemoteControlUI/Private/RemoteControlUIModule.h b/Engine/Plugins/VirtualProduction/RemoteControl/Source/RemoteControlUI/Private/RemoteControlUIModule.h
--- a/Engine/Plugins/VirtualProduction/RemoteControl/Source/RemoteControlUI/Private/RemoteControlUIModule.h
+++ b/Engine/Plugins/VirtualProduction/RemoteControl/Source/RemoteControlUI/Private/RemoteControlUIModule.h
@@ -60,7 +60,7 @@ public:
 	virtual void UnregisterWidgetFactoryForType(UScriptStruct* RemoteControlEntityType) override;
 	virtual void HighlightPropertyInDetailsPanel(const FPropertyPath& Path) const override;
 	virtual void SelectObjects(const TArray<UObject*>& Objects) const override;
-	virtual TSharedPtr<SWidget> CreateCustomControllerWidget(URCVirtualPropertyBase* InController) const override;
+	virtual TSharedPtr<SWidget> CreateCustomControllerWidget(URCVirtualPropertyBase* InController, TSharedPtr<IPropertyHandle> InOriginalPropertyHandle) const override;
 	//~ End IRemoteControlUIModule interface
 
 	/**
diff --git a/Engine/Plugins/VirtualProduction/RemoteControl/Source/RemoteControlUI/Private/UI/Controller/CustomControllers/SCustomTextureControllerWidget.cpp b/Engine/Plugins/VirtualProduction/RemoteControl/Source/RemoteControlUI/Private/UI/Controller/CustomControllers/SCustomTextureControllerWidget.cpp
--- a/Engine/Plugins/VirtualProduction/RemoteControl/Source/RemoteControlUI/Private/UI/Controller/CustomControllers/SCustomTextureControllerWidget.cpp
+++ b/Engine/Plugins/VirtualProduction/RemoteControl/Source/RemoteControlUI/Private/UI/Controller/CustomControllers/SCustomTextureControllerWidget.cpp
@@ -1,6 +1,7 @@
 ﻿// Copyright Epic Games, Inc. All Rights Reserved.
 
 #include "SCustomTextureControllerWidget.h"
+#include "Action/Bind/RCCustomBindActionUtilities.h"
 #include "AssetThumbnail.h"
 #include "Controller/RCController.h"
 #include "EditorDirectories.h"
@@ -8,9 +9,9 @@
 #include "Misc/MessageDialog.h"
 #include "Misc/PackageName.h"
 #include "PropertyCustomizationHelpers.h"
+#include "PropertyHandle.h"
 #include "RCVirtualProperty.h"
 #include "SlateOptMacros.h"
-#include "Action/Bind/RCCustomBindActionUtilities.h"
 #include "Styling/AppStyle.h"
 #include "ThumbnailRendering/ThumbnailManager.h"
 #include "Widgets/Input/SFilePathPicker.h"
@@ -20,17 +21,17 @@
 
 #define LOCTEXT_NAMESPACE "ExternalTextureControllerWidget"
 
-void SCustomTextureControllerWidget::Construct(const FArguments& InArgs, URCVirtualPropertyBase* InController)
+void SCustomTextureControllerWidget::Construct(const FArguments& InArgs, const TSharedPtr<IPropertyHandle>& InOriginalPropertyHandle)
 {
-	if (!InController || InController->GetValueType() != EPropertyBagPropertyType::String)
+	if (!InOriginalPropertyHandle.IsValid())
 	{
 		return;
 	}
-	
-	ControllerWeakPtr = InController;
+
+	OriginalPropertyHandle = InOriginalPropertyHandle;
 
 	FString ControllerString;
-	InController->GetValueString(ControllerString);
+	OriginalPropertyHandle->GetValueAsFormattedString(ControllerString);
 	bInternal = FPackageName::IsValidTextForLongPackageName(ControllerString);
 	if (FPackageName::IsValidTextForLongPackageName(ControllerString))
 	{
@@ -101,11 +102,6 @@ FString SCustomTextureControllerWidget::GetAssetPathName() const
 	return TEXT("");
 }
 
-URCVirtualPropertyBase* SCustomTextureControllerWidget::GetVirtualProperty() const
-{
-	return ControllerWeakPtr.Get();
-}
-
 void SCustomTextureControllerWidget::OnControllerTypeChanged(TSharedPtr<FString, ESPMode::ThreadSafe> InString, ESelectInfo::Type InArg)
 {
 	const FString& Selection = *InString.Get();
@@ -155,15 +151,12 @@ FString SCustomTextureControllerWidget::GetCurrentPath() const
 void SCustomTextureControllerWidget::UpdateControllerValue()
 {
 	const FString& Path = GetCurrentPath();
-	
-	if (URCController* Controller = Cast<URCController>(GetVirtualProperty()))
+
+	if (OriginalPropertyHandle.IsValid())
 	{
-		Controller->SetValueString(Path);
-		RefreshThumbnailImage();
-		
-		// Make sure behaviors/actions get executed
-		Controller->ExecuteBehaviours();
+		OriginalPropertyHandle->SetValueFromFormattedString(Path);
 	}
+	RefreshThumbnailImage();
 }
 
 void SCustomTextureControllerWidget::HandleFilePathPickerPathPicked(const FString& InPickedPath)
diff --git a/Engine/Plugins/VirtualProduction/RemoteControl/Source/RemoteControlUI/Private/UI/Controller/CustomControllers/SCustomTextureControllerWidget.h b/Engine/Plugins/VirtualProduction/RemoteControl/Source/RemoteControlUI/Private/UI/Controller/CustomControllers/SCustomTextureControllerWidget.h
--- a/Engine/Plugins/VirtualProduction/RemoteControl/Source/RemoteControlUI/Private/UI/Controller/CustomControllers/SCustomTextureControllerWidget.h
+++ b/Engine/Plugins/VirtualProduction/RemoteControl/Source/RemoteControlUI/Private/UI/Controller/CustomControllers/SCustomTextureControllerWidget.h
@@ -5,10 +5,11 @@
 #include "Widgets/Layout/SBox.h"
 #include "Widgets/SCompoundWidget.h"
 
+class IPropertyHandle;
 class URCVirtualPropertyBase;
 class UTexture2D;
 
-class REMOTECONTROLUI_API SCustomTextureControllerWidget : public SCompoundWidget
+class SCustomTextureControllerWidget : public SCompoundWidget
 {
 public:
 	SLATE_BEGIN_ARGS(SCustomTextureControllerWidget)
@@ -19,14 +20,11 @@ public:
 	
 	/**
 	 * Constructs this widget with InArgs
-	 * @param InController: the controller handled by this widget
+	 * @param InOriginalPropertyHandle Original PropertyHandle of this CustomWidget
 	 */
-	void Construct(const FArguments& InArgs, URCVirtualPropertyBase* InController);
+	void Construct(const FArguments& InArgs, const TSharedPtr<IPropertyHandle>& InOriginalPropertyHandle);
 
 private:
-	/** Returns the handled controller as URCVirtualPropertyBase */
-	URCVirtualPropertyBase* GetVirtualProperty() const;
-
 	/** Handle toggling Texture Controller type change: External vs. Internal */
 	void OnControllerTypeChanged(TSharedPtr<FString, ESPMode::ThreadSafe> InString, ESelectInfo::Type InArg);
 
@@ -64,9 +62,6 @@ private:
 	TSharedRef<SWidget> GetExternalTextureValueWidget();
 	TSharedRef<SWidget> GetAssetTextureValueWidget();
 
-	/** The controller handled by this widget */
-	TWeakObjectPtr<URCVirtualPropertyBase> ControllerWeakPtr;
-
 	/** The texture selected by the user via the asset or external file path */
 	TObjectPtr<UTexture2D> Texture;
 
@@ -87,4 +82,7 @@ private:
 
 	/** List of available Texture Controller types, used by combo box selection. Supports Asset and External*/
 	TArray<TSharedPtr<FString>> ControllerTypes;
+
+	/** Original PropertyHandle, used to update correctly the Widget and the Controller */
+	TSharedPtr<IPropertyHandle> OriginalPropertyHandle;
 };
diff --git a/Engine/Plugins/VirtualProduction/RemoteControl/Source/RemoteControlUI/Private/UI/Controller/RCControllerModel.cpp b/Engine/Plugins/VirtualProduction/RemoteControl/Source/RemoteControlUI/Private/UI/Controller/RCControllerModel.cpp
--- a/Engine/Plugins/VirtualProduction/RemoteControl/Source/RemoteControlUI/Private/UI/Controller/RCControllerModel.cpp
+++ b/Engine/Plugins/VirtualProduction/RemoteControl/Source/RemoteControlUI/Private/UI/Controller/RCControllerModel.cpp
@@ -64,7 +64,7 @@ TSharedRef<SWidget> FRCControllerModel::GetWidget() const
 		static const FMargin SlotMargin(10.0f, 2.0f);
 		
 		// If this is a custom controller, we will use its custom widget
-		if (const TSharedPtr<SWidget>& CustomControllerWidget = IRemoteControlUIModule::Get().CreateCustomControllerWidget(VirtualPropertyWeakPtr.Get()))
+		if (const TSharedPtr<SWidget>& CustomControllerWidget = IRemoteControlUIModule::Get().CreateCustomControllerWidget(VirtualPropertyWeakPtr.Get(), DetailTreeNodeWeakPtr.Pin()->CreatePropertyHandle()))
 		{
 			FieldWidget->AddSlot()
 				.Padding(SlotMargin)
@@ -135,16 +135,6 @@ TSharedRef<SWidget> FRCControllerModel::GetTypeSelectionWidget()
 	return SNullWidget::NullWidget;
 }
 
-TSharedRef<SWidget> FRCControllerModel::GetCustomControllerWidget(const FString& InCustomControllerTypeName) const
-{
-	if (InCustomControllerTypeName == UE::RCCustomControllers::CustomTextureControllerName)
-	{		
-		return SNew(SCustomTextureControllerWidget, VirtualPropertyWeakPtr.Get());		
-	}
-
-	return SNullWidget::NullWidget;
-}
-
 TSharedRef<SWidget> FRCControllerModel::GetControllerExtensionWidget(const FName& InColumnName) const
 {
 	FRCControllerExtensionWidgetsInfo ExtensionWidgetsInfo(GetVirtualProperty());
diff --git a/Engine/Plugins/VirtualProduction/RemoteControl/Source/RemoteControlUI/Private/UI/Controller/RCControllerModel.h b/Engine/Plugins/VirtualProduction/RemoteControl/Source/RemoteControlUI/Private/UI/Controller/RCControllerModel.h
--- a/Engine/Plugins/VirtualProduction/RemoteControl/Source/RemoteControlUI/Private/UI/Controller/RCControllerModel.h
+++ b/Engine/Plugins/VirtualProduction/RemoteControl/Source/RemoteControlUI/Private/UI/Controller/RCControllerModel.h
@@ -48,9 +48,6 @@ public:
 	/** The widget showing the Value Type selection dropdown for this Controller (used for MultiControllers) */
 	TSharedRef<SWidget> GetTypeSelectionWidget();
 
-	/** Retrieve the widget for a custom controller */
-	TSharedRef<SWidget> GetCustomControllerWidget(const FString& InCustomControllerTypeName) const;
-
 	/** The widget allowing controller columns customization */
 	TSharedRef<SWidget> GetControllerExtensionWidget(const FName& InColumnName) const;
 	
diff --git a/Engine/Plugins/VirtualProduction/RemoteControl/Source/RemoteControlUI/Public/IRemoteControlUIModule.h b/Engine/Plugins/VirtualProduction/RemoteControl/Source/RemoteControlUI/Public/IRemoteControlUIModule.h
--- a/Engine/Plugins/VirtualProduction/RemoteControl/Source/RemoteControlUI/Public/IRemoteControlUIModule.h
+++ b/Engine/Plugins/VirtualProduction/RemoteControl/Source/RemoteControlUI/Public/IRemoteControlUIModule.h
@@ -259,5 +259,5 @@ public:
 	/**
 	 * Tries to retrieve a Custom Controller Widget for the specified Controller
 	 */
-	virtual TSharedPtr<SWidget> CreateCustomControllerWidget(URCVirtualPropertyBase* InController) const = 0;
+	virtual TSharedPtr<SWidget> CreateCustomControllerWidget(URCVirtualPropertyBase* InController, TSharedPtr<IPropertyHandle> InOriginalPropertyHandle) const = 0;
 };
