[Remote Control] SRCPanelExposedEntitiesList::OnObjectPropertyChange refreshing all nodes widgets every time any object property is updated

- Refresh is now done on Tick
- It is possible to disable this refresh from Remote Control Config options

diff --git a/Engine/Plugins/VirtualProduction/RemoteControl/Source/RemoteControlCommon/Public/RemoteControlSettings.h b/Engine/Plugins/VirtualProduction/RemoteControl/Source/RemoteControlCommon/Public/RemoteControlSettings.h
--- a/Engine/Plugins/VirtualProduction/RemoteControl/Source/RemoteControlCommon/Public/RemoteControlSettings.h
+++ b/Engine/Plugins/VirtualProduction/RemoteControl/Source/RemoteControlCommon/Public/RemoteControlSettings.h
@@ -378,6 +378,10 @@ public:
 	
 	UPROPERTY(Config, EditAnywhere, Category = "Remote Control Preset", DisplayName = "Default mode (e.g. Setup/Operation)")
 	FName DefaultPanelMode = TEXT("Setup");
+
+	/** Refresh all widgets in the exposed properties list when object properties are updated */
+	UPROPERTY(config, EditAnywhere, Category = "Remote Control Preset", DisplayName = "Refresh all exposed entities widgets on object properties update")
+	bool bRefreshExposedEntitiesOnObjectPropertyUpdate = true;
 	
 private:
 	UPROPERTY(config)
diff --git a/Engine/Plugins/VirtualProduction/RemoteControl/Source/RemoteControlUI/Private/UI/SRCPanelExposedEntitiesList.cpp b/Engine/Plugins/VirtualProduction/RemoteControl/Source/RemoteControlUI/Private/UI/SRCPanelExposedEntitiesList.cpp
--- a/Engine/Plugins/VirtualProduction/RemoteControl/Source/RemoteControlUI/Private/UI/SRCPanelExposedEntitiesList.cpp
+++ b/Engine/Plugins/VirtualProduction/RemoteControl/Source/RemoteControlUI/Private/UI/SRCPanelExposedEntitiesList.cpp
@@ -467,6 +467,8 @@ void SRCPanelExposedEntitiesList::Tick(const FGeometry& AllottedGeometry, const
 		ProcessRefresh();
 		bRefreshRequested = false;
 	}
+
+	ExposedEntitiesNodesRefresh();
 }
 
 TSharedPtr<SRCPanelTreeNode> SRCPanelExposedEntitiesList::GetSelectedGroup() const
@@ -617,9 +619,22 @@ FText SRCPanelExposedEntitiesList::HandleEntityListHeaderLabel() const
 	return LOCTEXT("PropertiesLabel", "Properties");
 }
 
+void SRCPanelExposedEntitiesList::ExposedEntitiesNodesRefresh()
+{
+	if (bNodesRefreshRequested)
+	{
+		for (const TPair <FGuid, TSharedPtr<SRCPanelTreeNode>>& Node : FieldWidgetMap)
+		{
+			Node.Value->Refresh();
+		}
+
+		bNodesRefreshRequested = false;
+	}
+}
+
 void SRCPanelExposedEntitiesList::OnObjectPropertyChange(UObject* InObject, FPropertyChangedEvent& InChangeEvent)
 {
-	EPropertyChangeType::Type TypesNeedingRefresh = EPropertyChangeType::ArrayAdd | EPropertyChangeType::ArrayClear | EPropertyChangeType::ArrayRemove | EPropertyChangeType::ValueSet | EPropertyChangeType::Duplicate;
+	static constexpr EPropertyChangeType::Type TypesNeedingRefresh = EPropertyChangeType::ArrayAdd | EPropertyChangeType::ArrayClear | EPropertyChangeType::ArrayRemove | EPropertyChangeType::ValueSet | EPropertyChangeType::Duplicate;
 	auto IsRelevantProperty = [](FFieldClass* PropertyClass)
 	{
 		return PropertyClass && (PropertyClass == FArrayProperty::StaticClass() || PropertyClass == FSetProperty::StaticClass() || PropertyClass == FMapProperty::StaticClass());
@@ -665,10 +680,16 @@ void SRCPanelExposedEntitiesList::OnObjectPropertyChange(UObject* InObject, FPro
 				}
 			}
 		}
+		
+		bool bShouldRefreshNodes = true;
+		if (const URemoteControlSettings* Settings = GetDefault<URemoteControlSettings>())
+		{
+			bShouldRefreshNodes = Settings->bRefreshExposedEntitiesOnObjectPropertyUpdate;
+		}
 
-		for (const TPair <FGuid, TSharedPtr<SRCPanelTreeNode>>& Node : FieldWidgetMap)
+		if (bShouldRefreshNodes)
 		{
-			Node.Value->Refresh();
+			bNodesRefreshRequested = true;
 		}
 	}
 
diff --git a/Engine/Plugins/VirtualProduction/RemoteControl/Source/RemoteControlUI/Private/UI/SRCPanelExposedEntitiesList.h b/Engine/Plugins/VirtualProduction/RemoteControl/Source/RemoteControlUI/Private/UI/SRCPanelExposedEntitiesList.h
--- a/Engine/Plugins/VirtualProduction/RemoteControl/Source/RemoteControlUI/Private/UI/SRCPanelExposedEntitiesList.h
+++ b/Engine/Plugins/VirtualProduction/RemoteControl/Source/RemoteControlUI/Private/UI/SRCPanelExposedEntitiesList.h
@@ -203,6 +203,11 @@ private:
 
 	void ProcessRefresh();
 
+	/**
+	 * If necessary, will perform a refresh of the exposed entities nodes widgets.
+	 */
+	void ExposedEntitiesNodesRefresh();
+	
 private:
 	/** Holds the Groups list view. */
 	TSharedPtr<SListView<TSharedPtr<SRCPanelTreeNode>>> GroupsListView;
@@ -262,4 +267,7 @@ private:
 	FGuid CurrentlySelectedGroup;
 
 	bool bRefreshRequested = false;
+
+	/** When true, widgets of the Exposed Entities List will be refreshed on Tick */
+	bool bNodesRefreshRequested = false;
 };
\ No newline at end of file
