[RemoteControl]
Fix FText hitting an ensure because the StructOnScope is not setting the Package in the PropertyIdAction

diff --git a/Engine/Plugins/VirtualProduction/RemoteControl/Source/RemoteControl/Private/RCVirtualProperty.cpp b/Engine/Plugins/VirtualProduction/RemoteControl/Source/RemoteControl/Private/RCVirtualProperty.cpp
--- a/Engine/Plugins/VirtualProduction/RemoteControl/Source/RemoteControl/Private/RCVirtualProperty.cpp
+++ b/Engine/Plugins/VirtualProduction/RemoteControl/Source/RemoteControl/Private/RCVirtualProperty.cpp
@@ -4,6 +4,7 @@
 
 #include "IStructDeserializerBackend.h"
 #include "RCVirtualPropertyContainer.h"
+#include "RemoteControlPreset.h"
 #include "StructDeserializer.h"
 #include "StructSerializer.h"
 #include "UObject/TextProperty.h"
@@ -867,5 +868,10 @@ const FInstancedPropertyBag* URCVirtualPropertySelfContainer::GetPropertyBagInst
 
 TSharedPtr<FStructOnScope> URCVirtualPropertySelfContainer::CreateStructOnScope()
 {
-	return MakeShared<FStructOnScope>(Bag.GetPropertyBagStruct(), Bag.GetMutableValue().GetMemory());
+	TSharedRef<FStructOnScope> StructOnScope = MakeShared<FStructOnScope>(Bag.GetPropertyBagStruct(), Bag.GetMutableValue().GetMemory());
+	if (PresetWeakPtr.IsValid() && PresetWeakPtr->GetPackage())
+	{
+		StructOnScope->SetPackage(PresetWeakPtr->GetPackage());
+	}
+	return StructOnScope;
 }
diff --git a/Engine/Plugins/VirtualProduction/RemoteControl/Source/RemoteControlLogic/Private/Action/RCPropertyIdAction.cpp b/Engine/Plugins/VirtualProduction/RemoteControl/Source/RemoteControlLogic/Private/Action/RCPropertyIdAction.cpp
--- a/Engine/Plugins/VirtualProduction/RemoteControl/Source/RemoteControlLogic/Private/Action/RCPropertyIdAction.cpp
+++ b/Engine/Plugins/VirtualProduction/RemoteControl/Source/RemoteControlLogic/Private/Action/RCPropertyIdAction.cpp
@@ -150,6 +150,7 @@ void URCPropertyIdAction::UpdatePropertyId()
 										{
 											CachedPropertySelfContainer.Add(PropertyClassName, NewObject<URCVirtualPropertySelfContainer>(this));
 											CachedPropertySelfContainer[PropertyClassName]->DuplicateProperty(NewPropertyIdName, *FieldIt);
+											CachedPropertySelfContainer[PropertyClassName]->PresetWeakPtr = PresetWeakPtr;
 											PropertySelfContainer.Add(PropertyClassName, CachedPropertySelfContainer[PropertyClassName]);
 										}
 										break;
@@ -169,6 +170,7 @@ void URCPropertyIdAction::UpdatePropertyId()
 								CachedPropertySelfContainer[PropertyClassName]->AddProperty(NewPropertyIdName,
 									PropertyIdHandler->GetPropertyType(Property),
 									PropertyIdHandler->GetPropertyTypeObject(Property));
+								CachedPropertySelfContainer[PropertyClassName]->PresetWeakPtr = PresetWeakPtr;
 								
 								PropertySelfContainer.Add(PropertyClassName, CachedPropertySelfContainer[PropertyClassName]);
 							}
