Removed AndOrNot.h usage and other traits from misc files.

diff --git a/Engine/Plugins/Experimental/AVCodecs/AVCodecsCore/Source/AVCodecsCore/Public/AVExtension.h b/Engine/Plugins/Experimental/AVCodecs/AVCodecsCore/Source/AVCodecsCore/Public/AVExtension.h
--- a/Engine/Plugins/Experimental/AVCodecs/AVCodecsCore/Source/AVCodecsCore/Public/AVExtension.h
+++ b/Engine/Plugins/Experimental/AVCodecs/AVCodecsCore/Source/AVCodecsCore/Public/AVExtension.h
@@ -84,36 +84,26 @@ public:
 		return EAVResult::Success;
 	}
 
-	// Versions of TAnd and TOr that work with std::is_same return values, this is incompatible with TEMPLATE_REQUIRES
-	// so we need to still use this construct for now until std::is_same_v is better supported
-	template<typename A, typename B>
-	struct TIsSameHack
-	{
-		enum { Value = std::is_same_v<A, B>	};
-	};
-
 	/**
 	 * Transform one configuration type into another. Naive implementation is provided that uses the assignment operator, but it can be explicitly specialized.
 	 *
-	 * @tparam TOutput Config type to convert to.
-	 * @tparam TInput Config type to convert from.
+	 * @tparam OutputType Config type to convert to.
+	 * @tparam InputType Config type to convert from.
 	 * @param OutConfig The config to write to.
 	 * @param InConfig The config to read from.
 	 * @return Result of the operation, @see FAVResult.
 	 */
-	template <typename TOutput, typename TInput, TEMPLATE_REQUIRES(TOr<TIsSameHack<TInput, TOutput>, TNot<TIsDerivedFrom<TOutput, TInput>>>::Value)>
-	static FAVResult TransformConfig(TOutput& OutConfig, TInput const& InConfig)
-	{
-		OutConfig = InConfig;
-
-		return EAVResult::Success;
-	}
-
-	// Alternate path for upcasting.
-	template <typename TOutput, typename TInput, TEMPLATE_REQUIRES(TAnd<TNot<TIsSameHack<TInput, TOutput>>, TIsDerivedFrom<TOutput, TInput>>::Value)>
-	static FAVResult TransformConfig(TOutput& OutConfig, TInput const& InConfig)
+	template <typename OutputType, typename InputType>
+	static FAVResult TransformConfig(OutputType& OutConfig, InputType const& InConfig)
 	{
-		static_cast<TInput&>(OutConfig) = InConfig;
+		if constexpr (std::is_same_v<InputType, OutputType> || !std::is_base_of_v<InputType, OutputType>)
+		{
+			OutConfig = InConfig;
+		}
+		else
+		{
+			static_cast<InputType&>(OutConfig) = InConfig;
+		}
 
 		return EAVResult::Success;
 	}
diff --git a/Engine/Plugins/Experimental/Web/WebAPI/Source/WebAPIEditor/Public/WebAPIJsonUtilities.h b/Engine/Plugins/Experimental/Web/WebAPI/Source/WebAPIEditor/Public/WebAPIJsonUtilities.h
--- a/Engine/Plugins/Experimental/Web/WebAPI/Source/WebAPIEditor/Public/WebAPIJsonUtilities.h
+++ b/Engine/Plugins/Experimental/Web/WebAPI/Source/WebAPIEditor/Public/WebAPIJsonUtilities.h
@@ -35,10 +35,9 @@ namespace UE::Json
 		struct CMap 
 		{
 			template <typename T>
-			auto Requires() -> decltype(
-				TOr<
-					TIsTMap<T>,
-					TIsDerivedFrom<T, TMap<typename T::KeyType, typename T::ValueType>>>::Value);
+			auto Requires(const TMap<typename T::KeyType, typename T::ValueType>* MapPtr, T* Ptr) -> decltype(
+				MapPtr = Ptr
+			);
 		};
 	}
 	
diff --git a/Engine/Plugins/VirtualProduction/RemoteControl/Source/RemoteControlCommon/Public/RCTypeUtilities.h b/Engine/Plugins/VirtualProduction/RemoteControl/Source/RemoteControlCommon/Public/RCTypeUtilities.h
--- a/Engine/Plugins/VirtualProduction/RemoteControl/Source/RemoteControlCommon/Public/RCTypeUtilities.h
+++ b/Engine/Plugins/VirtualProduction/RemoteControl/Source/RemoteControlCommon/Public/RCTypeUtilities.h
@@ -4,6 +4,8 @@
 
 #include "CoreMinimal.h"
 
+#include <type_traits>
+
 #include "RCTypeTraits.h"
 #include "Math/NumericLimits.h"
 #include "Math/UnrealMathUtility.h"
@@ -157,25 +159,29 @@ namespace RemoteControlTypeUtilities
 	static auto MaxOp = [](auto&& InValues) { return FMath::Max(InValues); };
 
 	/** Returns typed metadata for numeric property if it exists, otherwise returns the DefaultValue. */
-	template <typename PropertyType, typename ValueType>
-	constexpr static typename TEnableIf<
-		std::is_base_of_v<FProperty, PropertyType> &&
-		RemoteControlTypeTraits::TNumericValueConstraint_V<ValueType>, ValueType>::Type
-	GetMetadataValue(const PropertyType* InProperty, const FName& InKey, const ValueType& InDefaultValue)
+	template <
+		typename PropertyType,
+		typename ValueType
+		UE_REQUIRES(std::is_base_of_v<FProperty, PropertyType>)
+	>
+	constexpr static ValueType GetMetadataValue(const PropertyType* InProperty, const FName& InKey, const ValueType& InDefaultValue)
 	{
 #if WITH_EDITORONLY_DATA
-		if(InProperty->HasMetaData(InKey))
+		if constexpr (RemoteControlTypeTraits::TNumericValueConstraint_V<ValueType>)
 		{
-			if(const FNumericProperty* NumericProperty = CastField<FNumericProperty>(InProperty))
+			if (InProperty->HasMetaData(InKey))
 			{
-				if(NumericProperty->IsInteger())
-				{
-					return static_cast<ValueType>(NumericProperty->GetIntMetaData(InKey));
-				}
-			
-				if(NumericProperty->IsFloatingPoint())
+				if (const FNumericProperty* NumericProperty = CastField<FNumericProperty>(InProperty))
 				{
-					return static_cast<ValueType>(NumericProperty->GetFloatMetaData(InKey));
+					if (NumericProperty->IsInteger())
+					{
+						return static_cast<ValueType>(NumericProperty->GetIntMetaData(InKey));
+					}
+
+					if (NumericProperty->IsFloatingPoint())
+					{
+						return static_cast<ValueType>(NumericProperty->GetFloatMetaData(InKey));
+					}
 				}
 			}
 		}
@@ -184,16 +190,6 @@ namespace RemoteControlTypeUtilities
 		return InDefaultValue;
 	}
 
-	/** Returns typed metadata if it exists, otherwise returns the DefaultValue. */
-	template <typename PropertyType, typename ValueType>
-	constexpr static typename TEnableIf<
-		std::is_base_of_v<FProperty, PropertyType> &&
-		!RemoteControlTypeTraits::TNumericValueConstraint_V<ValueType>, ValueType>::Type
-	GetMetadataValue(const PropertyType* InProperty, const FName& InKey, const ValueType InDefaultValue)
-	{
-		return InDefaultValue;
-	}
-
 	template <class Func, class PropertyType, class ValueType>
 	static ValueType GetClampedValue(Func InFunc, const PropertyType* InProperty, ValueType InValue, const TArray<FName>& InMetaKeys)
 	{
@@ -347,30 +343,27 @@ namespace RemoteControlTypeUtilities
 	}
 
 	template <typename ValueType>
-	static typename TEnableIf<!RemoteControlTypeTraits::TNumericValueConstraint_V<ValueType>, ValueType>::Type
-	GetDefaultRangeValueMin(const FProperty* InProperty)
+	static ValueType GetDefaultRangeValueMin(const FProperty* InProperty)
 	{
 		check(InProperty);
 
-		return TRemoteControlTypeTraits<ValueType>::DefaultRangeValueMin();
-	}
-
-	template <typename ValueType>
-	static typename TEnableIf<RemoteControlTypeTraits::TNumericValueConstraint_V<ValueType>, ValueType>::Type
-	GetDefaultRangeValueMin(const FProperty* InProperty)
-	{
-		check(InProperty);
+		if constexpr (RemoteControlTypeTraits::TNumericValueConstraint_V<ValueType>)
+		{
+			// Choose greater than the 3 values: (Traits defined) Default, ClampMin, or UIMin
+			ValueType ValueResult = GetClampedValue(MaxOp<ValueType>, InProperty, TRemoteControlTypeTraits<ValueType>::DefaultRangeValueMin(), { ClampMinKey, UIMinKey });
 
-		// Choose greater than the 3 values: (Traits defined) Default, ClampMin, or UIMin
-		ValueType ValueResult = GetClampedValue(MaxOp<ValueType>, InProperty, TRemoteControlTypeTraits<ValueType>::DefaultRangeValueMin(), { ClampMinKey, UIMinKey });
+			// ValueType might not match property, so clamp to min/max of property type
+			if (const FNumericProperty* NumericProperty = CastField<FNumericProperty>(InProperty))
+			{
+				ClampToPropertyType<ValueType>(NumericProperty, ValueResult);
+			}
 
-		// ValueType might not match property, so clamp to min/max of property type
-		if(const FNumericProperty* NumericProperty = CastField<FNumericProperty>(InProperty))
+			return ValueResult;
+		}
+		else
 		{
-			ClampToPropertyType<ValueType>(NumericProperty, ValueResult);
+			return TRemoteControlTypeTraits<ValueType>::DefaultRangeValueMin();
 		}
-
-		return ValueResult;
 	}
 
 	template <typename ValueType>
@@ -380,30 +373,27 @@ namespace RemoteControlTypeUtilities
 	}
 
 	template <typename ValueType>
-	static typename TEnableIf<!RemoteControlTypeTraits::TNumericValueConstraint_V<ValueType>, ValueType>::Type
-	GetDefaultRangeValueMax(const FProperty* InProperty)
-	{
-		check(InProperty);
-
-		return TRemoteControlTypeTraits<ValueType>::DefaultRangeValueMax();
-	}
-
-	template <typename ValueType>
-	static typename TEnableIf<RemoteControlTypeTraits::TNumericValueConstraint_V<ValueType>, ValueType>::Type
-	GetDefaultRangeValueMax(const FProperty* InProperty)
+	static ValueType GetDefaultRangeValueMax(const FProperty* InProperty)
 	{
 		check(InProperty);
 		
-		// Choose lesser than the 3 values: (Traits defined) Default, ClampMin, or UIMin
-		ValueType ValueResult =	GetClampedValue(MinOp<ValueType>, InProperty, TRemoteControlTypeTraits<ValueType>::DefaultRangeValueMax(), { ClampMaxKey, UIMaxKey });
+		if constexpr (RemoteControlTypeTraits::TNumericValueConstraint_V<ValueType>)
+		{
+			// Choose lesser than the 3 values: (Traits defined) Default, ClampMin, or UIMin
+			ValueType ValueResult = GetClampedValue(MinOp<ValueType>, InProperty, TRemoteControlTypeTraits<ValueType>::DefaultRangeValueMax(), { ClampMaxKey, UIMaxKey });
+
+			// ValueType might not match property, so clamp to min/max of property type
+			if (const FNumericProperty* NumericProperty = CastField<FNumericProperty>(InProperty))
+			{
+				ClampToPropertyType<ValueType>(NumericProperty, ValueResult);
+			}
 
-		// ValueType might not match property, so clamp to min/max of property type
-		if(const FNumericProperty* NumericProperty = CastField<FNumericProperty>(InProperty))
+			return ValueResult;
+		}
+		else
 		{
-			ClampToPropertyType<ValueType>(NumericProperty, ValueResult);
+			return TRemoteControlTypeTraits<ValueType>::DefaultRangeValueMax();
 		}
-
-		return ValueResult;
 	}
 
 	template <typename ValueType>
@@ -413,34 +403,33 @@ namespace RemoteControlTypeUtilities
 	}
 
 	template <typename ValueType>
-	static typename TEnableIf<!RemoteControlTypeTraits::TNumericValueConstraint_V<ValueType>, ValueType>::Type
-	GetDefaultMappingValueMin(const FProperty* InProperty)
+	static ValueType GetDefaultMappingValueMin(const FProperty* InProperty)
 	{
 		check(InProperty);
 
-		return TRemoteControlTypeTraits<ValueType>::DefaultMappingValueMin();
-	}
+		if constexpr (RemoteControlTypeTraits::TNumericValueConstraint_V<ValueType>)
+		{
+			ValueType ValueResult = GetClampedValue(MaxOp<ValueType>, InProperty, TRemoteControlTypeTraits<ValueType>::DefaultMappingValueMin(), { ClampMinKey, UIMinKey });
 
-	template <typename ValueType>
-	static typename TEnableIf<RemoteControlTypeTraits::TNumericValueConstraint_V<ValueType>, ValueType>::Type
-	GetDefaultMappingValueMin(const FProperty* InProperty)
-	{
-		check(InProperty);
-		
-		ValueType ValueResult = GetClampedValue(MaxOp<ValueType>, InProperty, TRemoteControlTypeTraits<ValueType>::DefaultMappingValueMin(), { ClampMinKey, UIMinKey });
+			// ValueType might not match property, so clamp to min/max of property type
+			if (const FNumericProperty* NumericProperty = CastField<FNumericProperty>(InProperty))
+			{
+				ClampToPropertyType<ValueType>(NumericProperty, ValueResult);
+			}
 
-		// ValueType might not match property, so clamp to min/max of property type
-		if(const FNumericProperty* NumericProperty = CastField<FNumericProperty>(InProperty))
+			return ValueResult;
+		}
+		else
 		{
-			ClampToPropertyType<ValueType>(NumericProperty, ValueResult);
+			return TRemoteControlTypeTraits<ValueType>::DefaultMappingValueMin();
 		}
-
-		return ValueResult;
 	}
 
-	template <typename ValueType>
-	static typename TEnableIf<std::is_same_v<ValueType, FStructOnScope>, TSharedPtr<ValueType>>::Type
-	GetDefaultMappingValueMin(const FStructProperty* InProperty)
+	template <
+		typename ValueType
+		UE_REQUIRES(std::is_same_v<ValueType, FStructOnScope>)
+	>
+	static TSharedPtr<ValueType> GetDefaultMappingValueMin(const FStructProperty* InProperty)
 	{
 		check(InProperty);
 
@@ -465,34 +454,33 @@ namespace RemoteControlTypeUtilities
 	}
 
 	template <typename ValueType>
-	static typename TEnableIf<!RemoteControlTypeTraits::TNumericValueConstraint_V<ValueType>, ValueType>::Type
-	GetDefaultMappingValueMax(const FProperty* InProperty)
+	static ValueType GetDefaultMappingValueMax(const FProperty* InProperty)
 	{
 		check(InProperty);
 
-		return TRemoteControlTypeTraits<ValueType>::DefaultMappingValueMax();
-	}
+		if constexpr (RemoteControlTypeTraits::TNumericValueConstraint_V<ValueType>)
+		{
+			ValueType ValueResult = GetClampedValue(MinOp<ValueType>, InProperty, TRemoteControlTypeTraits<ValueType>::DefaultMappingValueMax(), { ClampMaxKey, UIMaxKey });
 
-	template <typename ValueType>
-	static typename TEnableIf<RemoteControlTypeTraits::TNumericValueConstraint_V<ValueType>, ValueType>::Type
-	GetDefaultMappingValueMax(const FProperty* InProperty)
-	{
-		check(InProperty);
-		
-		ValueType ValueResult = GetClampedValue(MinOp<ValueType>, InProperty, TRemoteControlTypeTraits<ValueType>::DefaultMappingValueMax(), { ClampMaxKey, UIMaxKey });
+			// ValueType might not match property, so clamp to min/max of property type
+			if (const FNumericProperty* NumericProperty = CastField<FNumericProperty>(InProperty))
+			{
+				ClampToPropertyType<ValueType>(NumericProperty, ValueResult);
+			}
 
-		// ValueType might not match property, so clamp to min/max of property type
-		if(const FNumericProperty* NumericProperty = CastField<FNumericProperty>(InProperty))
+			return ValueResult;
+		}
+		else
 		{
-			ClampToPropertyType<ValueType>(NumericProperty, ValueResult);
+			return TRemoteControlTypeTraits<ValueType>::DefaultMappingValueMax();
 		}
-
-		return ValueResult;
 	}
 
-	template <typename ValueType>
-	static typename TEnableIf<std::is_same_v<ValueType, FStructOnScope>, TSharedPtr<ValueType>>::Type
-	GetDefaultMappingValueMax(const FStructProperty* InProperty)
+	template <
+		typename ValueType
+		UE_REQUIRES(std::is_same_v<ValueType, FStructOnScope>)
+	>
+	static TSharedPtr<ValueType> GetDefaultMappingValueMax(const FStructProperty* InProperty)
 	{
 		check(InProperty);
 
@@ -501,13 +489,15 @@ namespace RemoteControlTypeUtilities
 
 	/** Various RemoteControl property type Traits (wraps internal implementations). This resolves the property type for you. */
 
-	template <typename PropertyType, typename ValueType>
-	using TPropertyConstraint = typename TEnableIf<
-		TAnd<TIsDerivedFrom<PropertyType, FProperty>>::Value, ValueType>::Type;
+	template <typename PropertyType>
+	constexpr bool TPropertyConstraint_V = std::is_base_of_v<FProperty, PropertyType>;
 	
-	template <typename PropertyType, typename ValueType>
-	static TPropertyConstraint<PropertyType, ValueType>
-	GetDefaultRangeValueMin(const PropertyType* InProperty)
+	template <
+		typename PropertyType,
+		typename ValueType
+		UE_REQUIRES(TPropertyConstraint_V<PropertyType>)
+	>
+	static ValueType GetDefaultRangeValueMin(const PropertyType* InProperty)
 	{
 		check(InProperty);
 		
@@ -516,9 +506,12 @@ namespace RemoteControlTypeUtilities
 		return GetClampedValue(MaxOp<ValueType>, InProperty, TRemoteControlTypeTraits<ValueType>::DefaultRangeValueMin(), { ClampMinKey, UIMinKey });
 	}
 
-	template <typename PropertyType, typename ValueType>
-	static TPropertyConstraint<PropertyType, ValueType>
-	GetDefaultRangeValueMax(const FProperty* InProperty)
+	template <
+		typename PropertyType,
+		typename ValueType
+		UE_REQUIRES(TPropertyConstraint_V<PropertyType>)
+	>
+	static ValueType GetDefaultRangeValueMax(const FProperty* InProperty)
 	{
 		check(InProperty);
 
@@ -527,9 +520,12 @@ namespace RemoteControlTypeUtilities
 		return GetClampedValue(MinOp<ValueType>, InProperty, TRemoteControlTypeTraits<ValueType>::DefaultRangeValueMax(), { ClampMaxKey, UIMaxKey });
 	}
 
-	template <typename PropertyType, typename ValueType>
-	static TPropertyConstraint<PropertyType, ValueType>
-	GetDefaultMappingValueMin(const FProperty* InProperty)
+	template <
+		typename PropertyType,
+		typename ValueType
+		UE_REQUIRES(TPropertyConstraint_V<PropertyType>)
+	>
+	static ValueType GetDefaultMappingValueMin(const FProperty* InProperty)
 	{
 		check(InProperty);
 		
@@ -538,9 +534,12 @@ namespace RemoteControlTypeUtilities
 		return GetClampedValue(MaxOp<ValueType>, InProperty, TRemoteControlTypeTraits<ValueType>::DefaultMappingValueMin(), { ClampMinKey, UIMinKey });
 	}
 
-	template <typename PropertyType, typename ValueType>
-	static TPropertyConstraint<PropertyType, ValueType>
-	GetDefaultMappingValueMax(const FProperty* InProperty)
+	template <
+		typename PropertyType,
+		typename ValueType
+		UE_REQUIRES(TPropertyConstraint_V<PropertyType>)
+	>
+	static ValueType GetDefaultMappingValueMax(const FProperty* InProperty)
 	{
 		check(InProperty);
 		
diff --git a/Engine/Source/Runtime/ClothingSystemRuntimeCommon/Public/PointWeightMap.h b/Engine/Source/Runtime/ClothingSystemRuntimeCommon/Public/PointWeightMap.h
--- a/Engine/Source/Runtime/ClothingSystemRuntimeCommon/Public/PointWeightMap.h
+++ b/Engine/Source/Runtime/ClothingSystemRuntimeCommon/Public/PointWeightMap.h
@@ -1,6 +1,8 @@
 // Copyright Epic Games, Inc. All Rights Reserved.
 #pragma once
 
+#include <type_traits>
+
 #include "UObject/ObjectMacros.h"
 #include "PointWeightMap.generated.h"
 
@@ -73,7 +75,10 @@ struct FPointWeightMap
 	 * Initialize a weight map from another weight map while enabling and setting a new target.
 	 * @param Source the source weight map to copy the values from.
 	 * @param Target the new weight map target. */
-	template <typename T, typename = typename TEnableIf<TOr<TIsEnum<T>, TIsArithmetic<T>>::Value>::Type>
+	template <
+		typename T
+		UE_REQUIRES(std::is_enum_v<T> || std::is_arithmetic_v<T>)
+	>
 	void Initialize(const FPointWeightMap& Source, T Target)
 	{
 		Values = Source.Values;
