Swap FHttpRequestHandler to TDelegate from TFunction. Allows for more convenient callback handling.
Fixed up usage of FHttpRequestHandler across engine and projects.

#tests Ran some RPC tests successfully, which execute RPCs against the engine utilizing HTTPServer.

diff --git a/Engine/Plugins/VirtualProduction/RemoteControl/Source/WebRemoteControl/Private/RemoteControlDefaultPreprocessors.h b/Engine/Plugins/VirtualProduction/RemoteControl/Source/WebRemoteControl/Private/RemoteControlDefaultPreprocessors.h
--- a/Engine/Plugins/VirtualProduction/RemoteControl/Source/WebRemoteControl/Private/RemoteControlDefaultPreprocessors.h
+++ b/Engine/Plugins/VirtualProduction/RemoteControl/Source/WebRemoteControl/Private/RemoteControlDefaultPreprocessors.h
@@ -83,7 +83,7 @@ namespace UE::WebRemoteControl
 	/** Utility function to wrap a preprocessor handler to a http request handler than the HttpRouter can take. */
 	FHttpRequestHandler MakeHttpRequestHandler(FRCPreprocessorHandler Handler)
 	{
-		return [WrappedHandler = MoveTemp(Handler)](const FHttpServerRequest& Request, const FHttpResultCallback& OnComplete)
+		return FHttpRequestHandler::CreateLambda([WrappedHandler = MoveTemp(Handler)](const FHttpServerRequest& Request, const FHttpResultCallback& OnComplete)
 		{
 			FPreprocessorResult Result = WrappedHandler(Request);
 			if (Result.Result == EPreprocessorResult::RequestPassthrough)
@@ -95,7 +95,7 @@ namespace UE::WebRemoteControl
 				OnComplete(MoveTemp(Result.OptionalResponse));
 				return true; // Request handled.
 			}
-		};
+		});
 	}
 
 #if WITH_EDITOR
diff --git a/Engine/Plugins/VirtualProduction/RemoteControl/Source/WebRemoteControl/Private/RemoteControlRoute.h b/Engine/Plugins/VirtualProduction/RemoteControl/Source/WebRemoteControl/Private/RemoteControlRoute.h
--- a/Engine/Plugins/VirtualProduction/RemoteControl/Source/WebRemoteControl/Private/RemoteControlRoute.h
+++ b/Engine/Plugins/VirtualProduction/RemoteControl/Source/WebRemoteControl/Private/RemoteControlRoute.h
@@ -10,11 +10,9 @@
 
 #include "RemoteControlRoute.generated.h"
 
-DECLARE_DELEGATE_RetVal_TwoParams(bool, FRequestHandlerDelegate, const FHttpServerRequest&, const FHttpResultCallback&);
-
 struct FRemoteControlRoute
 {
-	FRemoteControlRoute(FString InRouteDescription, FHttpPath InPath, EHttpServerRequestVerbs InVerb, FRequestHandlerDelegate InHandler)
+	FRemoteControlRoute(FString InRouteDescription, FHttpPath InPath, EHttpServerRequestVerbs InVerb, FHttpRequestHandler InHandler)
 		: RouteDescription(MoveTemp(InRouteDescription))
 		, Path(MoveTemp(InPath))
 		, Verb(InVerb)
@@ -28,7 +26,7 @@ struct FRemoteControlRoute
 	/** The desired HTTP verb (ie. GET, PUT..) */
 	EHttpServerRequestVerbs Verb = EHttpServerRequestVerbs::VERB_GET;
 	/** The handler called when the route is accessed. */
-	FRequestHandlerDelegate Handler;
+	FHttpRequestHandler Handler;
 
 	friend uint32 GetTypeHash(const FRemoteControlRoute& Route) { return HashCombine(GetTypeHash(Route.Path), GetTypeHash(Route.Verb)); }
 	friend bool operator==(const FRemoteControlRoute& LHS, const FRemoteControlRoute& RHS) { return LHS.Path == RHS.Path && LHS.Verb == RHS.Verb; }
diff --git a/Engine/Plugins/VirtualProduction/RemoteControl/Source/WebRemoteControl/Private/WebRemoteControl.cpp b/Engine/Plugins/VirtualProduction/RemoteControl/Source/WebRemoteControl/Private/WebRemoteControl.cpp
--- a/Engine/Plugins/VirtualProduction/RemoteControl/Source/WebRemoteControl/Private/WebRemoteControl.cpp
+++ b/Engine/Plugins/VirtualProduction/RemoteControl/Source/WebRemoteControl/Private/WebRemoteControl.cpp
@@ -577,7 +577,7 @@ void FWebRemoteControlModule::SetExternalRemoteWebSocketLoggerConnection(TShared
 void FWebRemoteControlModule::StartRoute(const FRemoteControlRoute& Route)
 {
 	// The handler is wrapped in a lambda since HttpRouter::BindRoute only accepts TFunctions
-	ActiveRouteHandles.Add(GetTypeHash(Route), HttpRouter->BindRoute(Route.Path, Route.Verb, [this, Handler = Route.Handler](const FHttpServerRequest& Request, const FHttpResultCallback& OnComplete) { return Handler.Execute(Request, OnComplete); }));
+	ActiveRouteHandles.Add(GetTypeHash(Route), HttpRouter->BindRoute(Route.Path, Route.Verb, Route.Handler));
 }
 
 void FWebRemoteControlModule::RegisterRoutes()
@@ -587,14 +587,14 @@ void FWebRemoteControlModule::RegisterRoutes()
 		TEXT("Get information about different routes available on this API."),
 		FHttpPath(TEXT("/remote/info")),
 		EHttpServerRequestVerbs::VERB_GET,
-		FRequestHandlerDelegate::CreateRaw(this, &FWebRemoteControlModule::HandleInfoRoute)
+		FHttpRequestHandler::CreateRaw(this, &FWebRemoteControlModule::HandleInfoRoute)
 		});
 
 	RegisterRoute({
 		TEXT("Allows cross-origin http requests to the API."),
 		FHttpPath(TEXT("/remote")),
 		EHttpServerRequestVerbs::VERB_OPTIONS,
-		FRequestHandlerDelegate::CreateRaw(this, &FWebRemoteControlModule::HandleOptionsRoute)
+		FHttpRequestHandler::CreateRaw(this, &FWebRemoteControlModule::HandleOptionsRoute)
 		});
 
 	// Raw API
@@ -602,49 +602,49 @@ void FWebRemoteControlModule::RegisterRoutes()
 		TEXT("Allows batching multiple calls into one request."),
 		FHttpPath(TEXT("/remote/batch")),
 		EHttpServerRequestVerbs::VERB_PUT,
-		FRequestHandlerDelegate::CreateRaw(this, &FWebRemoteControlModule::HandleBatchRequest)
+		FHttpRequestHandler::CreateRaw(this, &FWebRemoteControlModule::HandleBatchRequest)
 		});
 
 	RegisterRoute({
 		TEXT("Call a function on a remote object."),
 		FHttpPath(TEXT("/remote/object/call")),
 		EHttpServerRequestVerbs::VERB_PUT,
-		FRequestHandlerDelegate::CreateRaw(this, &FWebRemoteControlModule::HandleObjectCallRoute)
+		FHttpRequestHandler::CreateRaw(this, &FWebRemoteControlModule::HandleObjectCallRoute)
 		});
 
 	RegisterRoute({
 		TEXT("Read or write a property on a remote object."),
 		FHttpPath(TEXT("/remote/object/property")),
 		EHttpServerRequestVerbs::VERB_PUT,
-		FRequestHandlerDelegate::CreateRaw(this, &FWebRemoteControlModule::HandleObjectPropertyRoute)
+		FHttpRequestHandler::CreateRaw(this, &FWebRemoteControlModule::HandleObjectPropertyRoute)
 		});
 
 	RegisterRoute({
 		TEXT("Append a new item to an array property on a remote object."),
 		FHttpPath(TEXT("/remote/object/property/append")),
 		EHttpServerRequestVerbs::VERB_PUT,
-		FRequestHandlerDelegate::CreateRaw(this, &FWebRemoteControlModule::HandleObjectPropertyAppendRoute)
+		FHttpRequestHandler::CreateRaw(this, &FWebRemoteControlModule::HandleObjectPropertyAppendRoute)
 		});
 
 	RegisterRoute({
 		TEXT("Insert a new item into an array property on a remote object."),
 		FHttpPath(TEXT("/remote/object/property/insert")),
 		EHttpServerRequestVerbs::VERB_PUT,
-		FRequestHandlerDelegate::CreateRaw(this, &FWebRemoteControlModule::HandleObjectPropertyInsertRoute)
+		FHttpRequestHandler::CreateRaw(this, &FWebRemoteControlModule::HandleObjectPropertyInsertRoute)
 		});
 
 	RegisterRoute({
 		TEXT("Remove an item from an array property on a remote object."),
 		FHttpPath(TEXT("/remote/object/property/remove")),
 		EHttpServerRequestVerbs::VERB_PUT,
-		FRequestHandlerDelegate::CreateRaw(this, &FWebRemoteControlModule::HandleObjectPropertyRemoveRoute)
+		FHttpRequestHandler::CreateRaw(this, &FWebRemoteControlModule::HandleObjectPropertyRemoveRoute)
 		});
 
 	RegisterRoute({
 		TEXT("Describe an object."),
 		FHttpPath(TEXT("/remote/object/describe")),
 		EHttpServerRequestVerbs::VERB_PUT,
-		FRequestHandlerDelegate::CreateRaw(this, &FWebRemoteControlModule::HandleDescribeObjectRoute)
+		FHttpRequestHandler::CreateRaw(this, &FWebRemoteControlModule::HandleDescribeObjectRoute)
 		});
 
 	// Passphrase Checking
@@ -652,7 +652,7 @@ void FWebRemoteControlModule::RegisterRoutes()
 		TEXT("Check whether or no the given Passphrase is correct"),
 		FHttpPath(TEXT("/remote/passphrase/")),
 		EHttpServerRequestVerbs::VERB_GET,
-		FRequestHandlerDelegate::CreateRaw(this, &FWebRemoteControlModule::HandlePassphraseRoute)
+		FHttpRequestHandler::CreateRaw(this, &FWebRemoteControlModule::HandlePassphraseRoute)
 		});
 
 	// Preset API
@@ -660,77 +660,77 @@ void FWebRemoteControlModule::RegisterRoutes()
 		TEXT("Get a remote control preset's content."),
 		FHttpPath(TEXT("/remote/preset/:preset")),
 		EHttpServerRequestVerbs::VERB_GET,
-		FRequestHandlerDelegate::CreateRaw(this, &FWebRemoteControlModule::HandleGetPresetRoute)
+		FHttpRequestHandler::CreateRaw(this, &FWebRemoteControlModule::HandleGetPresetRoute)
 		});
 
 	RegisterRoute({
 		TEXT("Get a list of available remote control presets."),
 		FHttpPath(TEXT("/remote/presets")),
 		EHttpServerRequestVerbs::VERB_GET,
-		FRequestHandlerDelegate::CreateRaw(this, &FWebRemoteControlModule::HandleGetPresetsRoute)
+		FHttpRequestHandler::CreateRaw(this, &FWebRemoteControlModule::HandleGetPresetsRoute)
 		});
 
 	RegisterRoute({
 		TEXT("Call a function on a preset."),
 		FHttpPath(TEXT("/remote/preset/:preset/function/:functionname")),
 		EHttpServerRequestVerbs::VERB_PUT,
-		FRequestHandlerDelegate::CreateRaw(this, &FWebRemoteControlModule::HandlePresetCallFunctionRoute)
+		FHttpRequestHandler::CreateRaw(this, &FWebRemoteControlModule::HandlePresetCallFunctionRoute)
 		});
 
 	RegisterRoute({
 		TEXT("Set a property on a preset."),
 		FHttpPath(TEXT("/remote/preset/:preset/property/:propertyname")),
 		EHttpServerRequestVerbs::VERB_PUT,
-		FRequestHandlerDelegate::CreateRaw(this, &FWebRemoteControlModule::HandlePresetSetPropertyRoute)
+		FHttpRequestHandler::CreateRaw(this, &FWebRemoteControlModule::HandlePresetSetPropertyRoute)
 		});
 
 	RegisterRoute({
 		TEXT("Get a property on a preset."),
 		FHttpPath(TEXT("/remote/preset/:preset/property/:propertyname")),
 		EHttpServerRequestVerbs::VERB_GET,
-		FRequestHandlerDelegate::CreateRaw(this, &FWebRemoteControlModule::HandlePresetGetPropertyRoute)
+		FHttpRequestHandler::CreateRaw(this, &FWebRemoteControlModule::HandlePresetGetPropertyRoute)
 		});
 
 	RegisterRoute({
 		TEXT("Expose a property on a preset."),
 		FHttpPath(TEXT("/remote/preset/:preset/expose/property")),
 		EHttpServerRequestVerbs::VERB_PUT,
-		FRequestHandlerDelegate::CreateRaw(this, &FWebRemoteControlModule::HandlePresetExposePropertyRoute)
+		FHttpRequestHandler::CreateRaw(this, &FWebRemoteControlModule::HandlePresetExposePropertyRoute)
 		});
 
 	RegisterRoute({
 		TEXT("Unexpose a property on a preset."),
 		FHttpPath(TEXT("/remote/preset/:preset/unexpose/property/:property")),
 		EHttpServerRequestVerbs::VERB_PUT,
-		FRequestHandlerDelegate::CreateRaw(this, &FWebRemoteControlModule::HandlePresetUnexposePropertyRoute)
+		FHttpRequestHandler::CreateRaw(this, &FWebRemoteControlModule::HandlePresetUnexposePropertyRoute)
 		});
 
 	RegisterRoute({
 		TEXT("Get an exposed actor's properties."),
 		FHttpPath(TEXT("/remote/preset/:preset/actor/:actor")),
 		EHttpServerRequestVerbs::VERB_GET,
-		FRequestHandlerDelegate::CreateRaw(this, &FWebRemoteControlModule::HandlePresetGetExposedActorPropertiesRoute)
+		FHttpRequestHandler::CreateRaw(this, &FWebRemoteControlModule::HandlePresetGetExposedActorPropertiesRoute)
 		});
 
 	RegisterRoute({
 		TEXT("Get an exposed actor's property."),
 		FHttpPath(TEXT("/remote/preset/:preset/actor/:actor/property/:propertyname")),
 		EHttpServerRequestVerbs::VERB_GET,
-		FRequestHandlerDelegate::CreateRaw(this, &FWebRemoteControlModule::HandlePresetGetExposedActorPropertyRoute)
+		FHttpRequestHandler::CreateRaw(this, &FWebRemoteControlModule::HandlePresetGetExposedActorPropertyRoute)
 		});
 
 	RegisterRoute({
 		TEXT("Set an exposed actor's property."),
 		FHttpPath(TEXT("/remote/preset/:preset/actor/:actor/property/:propertyname")),
 		EHttpServerRequestVerbs::VERB_PUT,
-		FRequestHandlerDelegate::CreateRaw(this, &FWebRemoteControlModule::HandlePresetSetExposedActorPropertyRoute)
+		FHttpRequestHandler::CreateRaw(this, &FWebRemoteControlModule::HandlePresetSetExposedActorPropertyRoute)
 		});
 
 	RegisterRoute({
 		TEXT("Get a preset"),
 		FHttpPath(TEXT("/remote/preset/:preset")),
 		EHttpServerRequestVerbs::VERB_GET,
-		FRequestHandlerDelegate::CreateRaw(this, &FWebRemoteControlModule::HandleGetPresetRoute)
+		FHttpRequestHandler::CreateRaw(this, &FWebRemoteControlModule::HandleGetPresetRoute)
 		});
 	
 	// Search
@@ -738,7 +738,7 @@ void FWebRemoteControlModule::RegisterRoutes()
 		TEXT("Search for assets"),
 		FHttpPath(TEXT("/remote/search/assets")),
 		EHttpServerRequestVerbs::VERB_PUT,
-		FRequestHandlerDelegate::CreateRaw(this, &FWebRemoteControlModule::HandleSearchAssetRoute)
+		FHttpRequestHandler::CreateRaw(this, &FWebRemoteControlModule::HandleSearchAssetRoute)
 		});
 
 	// Metadata
@@ -746,70 +746,70 @@ void FWebRemoteControlModule::RegisterRoutes()
 		TEXT("Get a preset's metadata"),
 		FHttpPath(TEXT("/remote/preset/:preset/metadata")),
 		EHttpServerRequestVerbs::VERB_GET,
-		FRequestHandlerDelegate::CreateRaw(this, &FWebRemoteControlModule::HandleGetMetadataRoute)
+		FHttpRequestHandler::CreateRaw(this, &FWebRemoteControlModule::HandleGetMetadataRoute)
 		});
 
 	RegisterRoute({
 		TEXT("Get/Set/Delete a preset metadata field"),
 		FHttpPath(TEXT("/remote/preset/:preset/metadata/:metadatafield")),
 		EHttpServerRequestVerbs::VERB_GET | EHttpServerRequestVerbs::VERB_PUT | EHttpServerRequestVerbs::VERB_DELETE,
-		FRequestHandlerDelegate::CreateRaw(this, &FWebRemoteControlModule::HandleMetadataFieldOperationsRoute)
+		FHttpRequestHandler::CreateRaw(this, &FWebRemoteControlModule::HandleMetadataFieldOperationsRoute)
 		});
 
 	RegisterRoute({
 	    TEXT("Get/Set/Delete an exposed property's metadata field"),
 	    FHttpPath(TEXT("/remote/preset/:preset/property/:label/metadata/:key")),
 	    EHttpServerRequestVerbs::VERB_GET | EHttpServerRequestVerbs::VERB_PUT | EHttpServerRequestVerbs::VERB_DELETE,
-	    FRequestHandlerDelegate::CreateRaw(this, &FWebRemoteControlModule::HandleEntityMetadataOperationsRoute)
+	    FHttpRequestHandler::CreateRaw(this, &FWebRemoteControlModule::HandleEntityMetadataOperationsRoute)
     });
 
 	RegisterRoute({
         TEXT("Get/Set/Delete an exposed function's metadata field"),
 		FHttpPath(TEXT("/remote/preset/:preset/function/:label/metadata/:key")),
         EHttpServerRequestVerbs::VERB_GET | EHttpServerRequestVerbs::VERB_PUT | EHttpServerRequestVerbs::VERB_DELETE,
-        FRequestHandlerDelegate::CreateRaw(this, &FWebRemoteControlModule::HandleEntityMetadataOperationsRoute)
+        FHttpRequestHandler::CreateRaw(this, &FWebRemoteControlModule::HandleEntityMetadataOperationsRoute)
     });
 
 	RegisterRoute({
         TEXT("Get/Set/Delete an exposed property's metadata field"),
 		FHttpPath(TEXT("/remote/preset/:preset/actor/:label/metadata/:key")),
         EHttpServerRequestVerbs::VERB_GET | EHttpServerRequestVerbs::VERB_PUT | EHttpServerRequestVerbs::VERB_DELETE,
-        FRequestHandlerDelegate::CreateRaw(this, &FWebRemoteControlModule::HandleEntityMetadataOperationsRoute)
+        FHttpRequestHandler::CreateRaw(this, &FWebRemoteControlModule::HandleEntityMetadataOperationsRoute)
     });
 
 	RegisterRoute({
 		TEXT("Set an exposed property's label"),
 		FHttpPath(TEXT("/remote/preset/:preset/property/:label/label")),
 		EHttpServerRequestVerbs::VERB_PUT,
-		FRequestHandlerDelegate::CreateRaw(this, &FWebRemoteControlModule::HandleEntitySetLabelRoute)
+		FHttpRequestHandler::CreateRaw(this, &FWebRemoteControlModule::HandleEntitySetLabelRoute)
 	});
 	
 	RegisterRoute({
 		TEXT("Set an exposed function's label"),
 		FHttpPath(TEXT("/remote/preset/:preset/function/:label/label")),
 		EHttpServerRequestVerbs::VERB_PUT,
-		FRequestHandlerDelegate::CreateRaw(this, &FWebRemoteControlModule::HandleEntitySetLabelRoute)
+		FHttpRequestHandler::CreateRaw(this, &FWebRemoteControlModule::HandleEntitySetLabelRoute)
 	});
 	
 	RegisterRoute({
 		TEXT("Set an exposed actor's label"),
 		FHttpPath(TEXT("/remote/preset/:preset/actor/:label/label")),
 		EHttpServerRequestVerbs::VERB_PUT,
-		FRequestHandlerDelegate::CreateRaw(this, &FWebRemoteControlModule::HandleEntitySetLabelRoute)
+		FHttpRequestHandler::CreateRaw(this, &FWebRemoteControlModule::HandleEntitySetLabelRoute)
 	});
 
 	RegisterRoute({
 		TEXT("Create a temporary preset which won't be visible in the editor or saved as an asset"),
 		FHttpPath(TEXT("/remote/preset/transient")),
 		EHttpServerRequestVerbs::VERB_PUT,
-		FRequestHandlerDelegate::CreateRaw(this, &FWebRemoteControlModule::HandleCreateTransientPresetRoute)
+		FHttpRequestHandler::CreateRaw(this, &FWebRemoteControlModule::HandleCreateTransientPresetRoute)
 	});
 
 	RegisterRoute({
 		TEXT("Delete a transient preset"),
 		FHttpPath(TEXT("/remote/preset/transient/:preset")),
 		EHttpServerRequestVerbs::VERB_DELETE,
-		FRequestHandlerDelegate::CreateRaw(this, &FWebRemoteControlModule::HandleDeleteTransientPresetRoute)
+		FHttpRequestHandler::CreateRaw(this, &FWebRemoteControlModule::HandleDeleteTransientPresetRoute)
 	});
 
 	// Remote Control Logic - Controllers
@@ -817,14 +817,14 @@ void FWebRemoteControlModule::RegisterRoutes()
 	TEXT("Set a controller value on a preset."),
 	FHttpPath(TEXT("/remote/preset/:preset/controller/:controller")),
 	EHttpServerRequestVerbs::VERB_PUT,
-	FRequestHandlerDelegate::CreateRaw(this, &FWebRemoteControlModule::HandlePresetSetControllerRoute)
+	FHttpRequestHandler::CreateRaw(this, &FWebRemoteControlModule::HandlePresetSetControllerRoute)
 		});
 
 	RegisterRoute({
 		TEXT("Get a controller value from a preset."),
 		FHttpPath(TEXT("/remote/preset/:preset/controller/:controller")),
 		EHttpServerRequestVerbs::VERB_GET,
-		FRequestHandlerDelegate::CreateRaw(this, &FWebRemoteControlModule::HandlePresetGetControllerRoute)
+		FHttpRequestHandler::CreateRaw(this, &FWebRemoteControlModule::HandlePresetGetControllerRoute)
 		});
 
 	//**************************************
diff --git a/Engine/Plugins/VirtualProduction/RemoteControl/Source/WebRemoteControl/Private/WebRemoteControlEditorRoutes.cpp b/Engine/Plugins/VirtualProduction/RemoteControl/Source/WebRemoteControl/Private/WebRemoteControlEditorRoutes.cpp
--- a/Engine/Plugins/VirtualProduction/RemoteControl/Source/WebRemoteControl/Private/WebRemoteControlEditorRoutes.cpp
+++ b/Engine/Plugins/VirtualProduction/RemoteControl/Source/WebRemoteControl/Private/WebRemoteControlEditorRoutes.cpp
@@ -43,7 +43,7 @@ void FWebRemoteControlEditorRoutes::RegisterRoutes(FWebRemoteControlModule* WebR
 			TEXT("Create a connection until an event is triggered."),
 			FHttpPath(TEXT("/remote/object/event")),
 			EHttpServerRequestVerbs::VERB_PUT,
-			FRequestHandlerDelegate::CreateRaw(this, &FWebRemoteControlEditorRoutes::HandleObjectEventRoute)
+			FHttpRequestHandler::CreateRaw(this, &FWebRemoteControlEditorRoutes::HandleObjectEventRoute)
 		));
 
 		EventDispatchers.AddDefaulted((int32)ERemoteControlEvent::EventCount);
@@ -54,7 +54,7 @@ void FWebRemoteControlEditorRoutes::RegisterRoutes(FWebRemoteControlModule* WebR
 		TEXT("Get an object's thumbnail"),
 		FHttpPath(TEXT("/remote/object/thumbnail")),
 		EHttpServerRequestVerbs::VERB_PUT,
-		FRequestHandlerDelegate::CreateRaw(this, &FWebRemoteControlEditorRoutes::HandleGetThumbnailRoute)
+		FHttpRequestHandler::CreateRaw(this, &FWebRemoteControlEditorRoutes::HandleGetThumbnailRoute)
 	));
 }
 
diff --git a/Engine/Source/Runtime/ExternalRPCRegistry/Private/ExternalRpcRegistry.cpp b/Engine/Source/Runtime/ExternalRPCRegistry/Private/ExternalRpcRegistry.cpp
--- a/Engine/Source/Runtime/ExternalRPCRegistry/Private/ExternalRpcRegistry.cpp
+++ b/Engine/Source/Runtime/ExternalRPCRegistry/Private/ExternalRpcRegistry.cpp
@@ -92,13 +92,7 @@ UExternalRpcRegistry* UExternalRpcRegistry::GetInstance()
 		}
 		FParse::Value(FCommandLine::Get(), TEXT("rpcport="), ObjectInstance->PortToUse);
 		
-		TWeakObjectPtr<UExternalRpcRegistry> WeakThis(ObjectInstance);
-
-		const FHttpRequestHandler& ListRoutesRequestHandler = [WeakThis](const FHttpServerRequest& Request, const FHttpResultCallback& OnComplete)
-		{
-			if (!WeakThis.IsValid()) { return false; }
-			return WeakThis->HttpListOpenRoutes(Request, OnComplete);
-		};
+		FHttpRequestHandler ListRoutesRequestHandler = FHttpRequestHandler::CreateUObject(ObjectInstance, &ThisClass::HttpListOpenRoutes);
 		TArray<FExternalRpcArgumentDesc*> ArgumentArray;
 		// We always want the ListRegisteredRpcs route bound, no matter what.
 		UExternalRpcRegistry::GetInstance()->RegisterNewRouteWithArguments(TEXT("ListRegisteredRpcs"), FHttpPath("/listrpcs"), EHttpServerRequestVerbs::VERB_GET,
diff --git a/Engine/Source/Runtime/Online/HTTPServer/Private/HttpConnection.cpp b/Engine/Source/Runtime/Online/HTTPServer/Private/HttpConnection.cpp
--- a/Engine/Source/Runtime/Online/HTTPServer/Private/HttpConnection.cpp
+++ b/Engine/Source/Runtime/Online/HTTPServer/Private/HttpConnection.cpp
@@ -202,10 +202,15 @@ void FHttpConnection::ProcessRequest(const TSharedPtr<FHttpServerRequest>& Reque
 
 	bool bRequestHandled = false;
 	auto RequestHandlerIterator = Router->CreateRequestHandlerIterator(Request);
-	while (const auto& RequestHandlerPtr = RequestHandlerIterator.Next())
+	while (const FHttpRequestHandler* RequestHandlerPtr = RequestHandlerIterator.Next())
 	{
-		(*RequestHandlerPtr).CheckCallable();
-		bRequestHandled = (*RequestHandlerPtr)(*Request, OnProcessingComplete);
+		if (!RequestHandlerPtr->IsBound())
+		{
+			UE_LOG(LogHttpConnection, Verbose, TEXT("Skipping an unbound FHttpRequestHandler."));
+			continue;
+		}
+
+		bRequestHandled = RequestHandlerPtr->Execute(*Request, OnProcessingComplete);
 		if (bRequestHandled)
 		{
 			break;
diff --git a/Engine/Source/Runtime/Online/HTTPServer/Private/HttpRouter.cpp b/Engine/Source/Runtime/Online/HTTPServer/Private/HttpRouter.cpp
--- a/Engine/Source/Runtime/Online/HTTPServer/Private/HttpRouter.cpp
+++ b/Engine/Source/Runtime/Online/HTTPServer/Private/HttpRouter.cpp
@@ -9,6 +9,8 @@
 #include "HttpRequestHandler.h"
 #include "HttpRequestHandlerIterator.h"
 
+DEFINE_LOG_CATEGORY_STATIC(LogHttpRouter, Log, All);
+
 bool FHttpRouter::Query(const TSharedPtr<FHttpServerRequest>& Request, const FHttpResultCallback& OnProcessingComplete)
 {
 	bool bRequestHandled = false;
@@ -18,8 +20,13 @@ bool FHttpRouter::Query(const TSharedPtr<FHttpServerRequest>& Request, const FHt
 	FHttpRequestHandlerIterator Iterator(Request, RequestHandlerRegistrar, MoveTemp(PreprocessorsArray));
 	while (const FHttpRequestHandler* RequestHandlerPtr = Iterator.Next())
 	{
-		(*RequestHandlerPtr).CheckCallable();
-		bRequestHandled = (*RequestHandlerPtr)(*Request, OnProcessingComplete);
+		if (!RequestHandlerPtr->IsBound())
+		{
+			UE_LOG(LogHttpRouter, Verbose, TEXT("Skipping an unbound FHttpRequestHandler."));
+			continue;
+		}
+
+		bRequestHandled = RequestHandlerPtr->Execute(*Request, OnProcessingComplete);
 		if (bRequestHandled)
 		{
 			break;
diff --git a/Engine/Source/Runtime/Online/HTTPServer/Private/Tests/HttpServerTests.cpp b/Engine/Source/Runtime/Online/HTTPServer/Private/Tests/HttpServerTests.cpp
--- a/Engine/Source/Runtime/Online/HTTPServer/Private/Tests/HttpServerTests.cpp
+++ b/Engine/Source/Runtime/Online/HTTPServer/Private/Tests/HttpServerTests.cpp
@@ -34,11 +34,10 @@ bool FHttpServerIntegrationTest::RunTest(const FString& Parameters)
 	TestTrue(TEXT("HttpRouter is NOT null on bind failure by default"), ValidHttpRouterOnFail.IsValid());
 
 	// Ensure we can create route bindings
-	const FHttpRequestHandler RequestHandler = [this]
-	(const FHttpServerRequest& Request, const FHttpResultCallback& OnComplete)
+	const FHttpRequestHandler RequestHandler = FHttpRequestHandler::CreateLambda([](const FHttpServerRequest& Request, const FHttpResultCallback& OnComplete)
 	{
 		return true;
-	};
+	});
 	FHttpRouteHandle HttpRouteHandle = HttpRouter->BindRoute(HttpPath, EHttpServerRequestVerbs::VERB_GET, RequestHandler);
 	TestTrue(TEXT("HttpRouteHandle.IsValid()"), HttpRouteHandle.IsValid());
 
@@ -93,11 +92,11 @@ bool FHttpServerPathParametersTest::RunTest(const FString& Parameters)
 
 		void SetHandler()
 		{
-			Handler = [this](const FHttpServerRequest&, const FHttpResultCallback&)
+			Handler = FHttpRequestHandler::CreateLambda([this](const FHttpServerRequest&, const FHttpResultCallback&)
 			{
 				bRouteQueried = true;
 				return true;
-			};
+			});
 		}
 
 		EVerb QueryVerb;
@@ -157,7 +156,7 @@ bool FHttpServerPathParametersTest::RunTest(const FString& Parameters)
 		FHttpRequestHandlerIterator Iterator(Request, Registrar);
 		if (const FHttpRequestHandler* RequestHandlerPtr = Iterator.Next())
 		{
-			(*RequestHandlerPtr)(*Request, Callback);
+			[[maybe_unused]] bool bHandled = RequestHandlerPtr->Execute(*Request, Callback);
 		}
 
 		if ((Test.bExpectedResult == ShouldMatch && !Test.bRouteQueried) || (Test.bExpectedResult == ShouldNotMatch && Test.bRouteQueried))
diff --git a/Engine/Source/Runtime/Online/HTTPServer/Public/HttpRequestHandler.h b/Engine/Source/Runtime/Online/HTTPServer/Public/HttpRequestHandler.h
--- a/Engine/Source/Runtime/Online/HTTPServer/Public/HttpRequestHandler.h
+++ b/Engine/Source/Runtime/Online/HTTPServer/Public/HttpRequestHandler.h
@@ -16,5 +16,5 @@ struct FHttpServerRequest;
  * @param OnComplete The callback to invoke to write an http response
  * @return True if the request has been handled, false otherwise
  */
-typedef TFunction<bool(const FHttpServerRequest& Request, const FHttpResultCallback& OnComplete)> FHttpRequestHandler;
+using FHttpRequestHandler = TDelegate<bool(const FHttpServerRequest& Request, const FHttpResultCallback& OnComplete)>;
 
diff --git a/Engine/Source/Runtime/PerfCounters/Private/PerfCounters.cpp b/Engine/Source/Runtime/PerfCounters/Private/PerfCounters.cpp
--- a/Engine/Source/Runtime/PerfCounters/Private/PerfCounters.cpp
+++ b/Engine/Source/Runtime/PerfCounters/Private/PerfCounters.cpp
@@ -74,12 +74,7 @@ bool FPerfCounters::Initialize()
 	// Register a handler for /stats
 	TWeakPtr<FPerfCounters> WeakThisPtr(AsShared());
     StatsRouteHandle = HttpRouter->BindRoute(FHttpPath("/stats"), EHttpServerRequestVerbs::VERB_GET,
-		[WeakThisPtr](const FHttpServerRequest& Request, const FHttpResultCallback& OnComplete)
-	{
-		auto SharedThis = WeakThisPtr.Pin();
-		if (!SharedThis.IsValid()) { return false; }
-		return SharedThis->ProcessStatsRequest(Request, OnComplete);
-	});
+		FHttpRequestHandler::CreateSP(this, &FPerfCounters::ProcessStatsRequest));
 	if(!StatsRouteHandle.IsValid())
 	{
 		UE_LOG(LogPerfCounters, Error,
@@ -89,12 +84,7 @@ bool FPerfCounters::Initialize()
 
 	// Register a handler for /exec
 	ExecRouteHandle = HttpRouter->BindRoute(FHttpPath("/exec"), EHttpServerRequestVerbs::VERB_GET,
-		[WeakThisPtr] (const FHttpServerRequest& Request, const FHttpResultCallback& OnComplete)
-	{
-		auto SharedThis = WeakThisPtr.Pin();
-		if (!SharedThis.IsValid()) { return false; }
-		return SharedThis->ProcessExecRequest(Request, OnComplete);
-	});
+		FHttpRequestHandler::CreateSP(this, &FPerfCounters::ProcessExecRequest));
 	if(!ExecRouteHandle.IsValid())
 	{
 		UE_LOG(LogPerfCounters, Error,
