[RemoteControl]
PropertyId makes RemoteControl crash if exposing an array not owned by a UObject.

diff --git a/Engine/Plugins/VirtualProduction/RemoteControl/Source/RemoteControl/Private/RemoteControlPropertyIdRegistry.cpp b/Engine/Plugins/VirtualProduction/RemoteControl/Source/RemoteControl/Private/RemoteControlPropertyIdRegistry.cpp
--- a/Engine/Plugins/VirtualProduction/RemoteControl/Source/RemoteControl/Private/RemoteControlPropertyIdRegistry.cpp
+++ b/Engine/Plugins/VirtualProduction/RemoteControl/Source/RemoteControl/Private/RemoteControlPropertyIdRegistry.cpp
@@ -15,7 +15,7 @@
 #include "StructSerializer.h"
 #endif
 
-namespace
+namespace RemoteControlPropertyIdUtilities
 {
 	bool GetObjectRef(const TSharedPtr<FRemoteControlProperty>& Field, const ERCAccess Access, FRCObjectReference& OutObjectRef)
 	{
@@ -36,6 +36,87 @@ namespace
 		}
 		return false;
 	}
+
+	bool CopyNonUClassOwnerArray(const TObjectPtr<URCVirtualPropertySelfContainer>& InVirtualPropertySelfContainer,
+		const FArrayProperty* InArrayProperty,
+		const TSharedPtr<FRemoteControlProperty>& InTargetRCProperty,
+		FProperty* InProperty,
+		FRCFieldPathInfo InSegment,
+		const int32 InIndexSegment,
+		const uint8* InPropertyContainer,
+		const bool bCheckByteEnumComparison)
+	{
+		FScriptArrayHelper ArrayHelper(InArrayProperty, InPropertyContainer);
+		if (const TSharedPtr<FRemoteControlField> TargetRCField = StaticCastSharedPtr<FRemoteControlField>(InTargetRCProperty);
+			!TargetRCField->FieldPathInfo.Segments.IsEmpty())
+		{
+			const int32 Index = InSegment.Segments[InIndexSegment].ArrayIndex;
+			uint8* PtrContainer = ArrayHelper.GetElementPtr(Index);
+			if (InIndexSegment < InSegment.GetSegmentCount() - 2)
+			{
+				return false;
+			}
+			if (const FArrayProperty* ArrayRealProperty = CastField<FArrayProperty>(InProperty))
+			{
+				return InVirtualPropertySelfContainer->CopyCompleteValue(ArrayRealProperty->Inner, PtrContainer, bCheckByteEnumComparison);
+			}
+		}
+		return false;
+	}
+	
+	bool CopyNonUClassOwnerMap(const TObjectPtr<URCVirtualPropertySelfContainer>& InVirtualPropertySelfContainer,
+		const FMapProperty* InMapProperty,
+		const TSharedPtr<FRemoteControlProperty>& InTargetRCProperty,
+		FProperty* InProperty,
+		FRCFieldPathInfo InSegment,
+		const int32 InIndexSegment,
+		const uint8* InPropertyContainer,
+		const bool bCheckByteEnumComparison)
+	{
+		FScriptMapHelper MapHelper(InMapProperty, InPropertyContainer);
+		if (const TSharedPtr<FRemoteControlField> TargetRCField = StaticCastSharedPtr<FRemoteControlField>(InTargetRCProperty);
+			!TargetRCField->FieldPathInfo.Segments.IsEmpty())
+		{
+			const int32 Index = InSegment.Segments[InIndexSegment].ArrayIndex;
+			uint8* PtrContainer = MapHelper.GetValuePtr(Index);
+			if (InIndexSegment < InSegment.GetSegmentCount() - 2)
+			{
+				return false;
+			}
+			if (const FMapProperty* MapRealProperty = CastField<FMapProperty>(InProperty))
+			{
+				return InVirtualPropertySelfContainer->CopyCompleteValue(MapRealProperty->ValueProp, PtrContainer, bCheckByteEnumComparison);
+			}
+		}
+		return false;
+	}
+
+	bool CopyNonUClassOwnerSet(const TObjectPtr<URCVirtualPropertySelfContainer>& InVirtualPropertySelfContainer,
+		const FSetProperty* InSetProperty,
+		const TSharedPtr<FRemoteControlProperty>& InTargetRCProperty,
+		FProperty* InProperty,
+		FRCFieldPathInfo InSegment,
+		const int32 InIndexSegment,
+		const uint8* InPropertyContainer,
+		const bool bCheckByteEnumComparison)
+	{
+		FScriptSetHelper SetHelper(InSetProperty, InPropertyContainer);
+		if (const TSharedPtr<FRemoteControlField> TargetRCField = StaticCastSharedPtr<FRemoteControlField>(InTargetRCProperty);
+			!TargetRCField->FieldPathInfo.Segments.IsEmpty())
+		{
+			const int32 Index = InSegment.Segments[InIndexSegment].ArrayIndex;
+			uint8* PtrContainer = SetHelper.GetElementPtr(Index);
+			if (InIndexSegment < InSegment.GetSegmentCount() - 2)
+			{
+				return false;
+			}
+			if (const FSetProperty* SetRealProperty = CastField<FSetProperty>(InProperty))
+			{
+				return InVirtualPropertySelfContainer->CopyCompleteValue(SetRealProperty->ElementProp, PtrContainer, bCheckByteEnumComparison);
+			}
+		}
+		return false;
+	}
 }
 
 FRCPropertyIdWrapper::FRCPropertyIdWrapper(const TSharedRef<FRemoteControlProperty>& InRCProperty)
@@ -174,32 +255,48 @@ void FRCPropertyIdWrapper::UpdateTypes(const TSharedRef<FRemoteControlProperty>&
 	else if (const FArrayProperty* ArrayProperty = CastField<FArrayProperty>(InProperty))
 	{
 		SuperType = ArrayProperty->Inner->GetClass()->GetFName();
-		if (const uint8* PropertyValuePtr = ArrayProperty->ContainerPtrToValuePtr<uint8>(InOwner))
+
+		// ContainerPtrToValuePtr will crash if the GetOwner<UClass> is nullptr so we check before.
+		if (ArrayProperty->GetOwner<UClass>())
 		{
-			if (const FObjectProperty* InnerObjectProperty = CastField<FObjectProperty>(ArrayProperty->Inner))
+			if (const uint8* PropertyValuePtr = ArrayProperty->ContainerPtrToValuePtr<uint8>(InOwner))
 			{
-				FScriptArrayHelper ArrayHelper(ArrayProperty, PropertyValuePtr);
-				
-				if (const TSharedRef<FRemoteControlField> TargetRCField = StaticCastSharedRef<FRemoteControlField>(InRCProperty); !TargetRCField->FieldPathInfo.Segments.IsEmpty())
+				if (const FObjectProperty* InnerObjectProperty = CastField<FObjectProperty>(ArrayProperty->Inner))
 				{
-					const int32 MaterialIndex = TargetRCField->FieldPathInfo.Segments[0].ArrayIndex;
-					const uint8* ObjPtrContainer = ArrayHelper.GetRawPtr(MaterialIndex);
-					if (const UObject* CurrentObject = InnerObjectProperty->GetObjectPropertyValue(ObjPtrContainer))
+					FScriptArrayHelper ArrayHelper(ArrayProperty, PropertyValuePtr);
+					
+					if (const TSharedRef<FRemoteControlField> TargetRCField = StaticCastSharedRef<FRemoteControlField>(InRCProperty); !TargetRCField->FieldPathInfo.Segments.IsEmpty())
 					{
-						if (CurrentObject->IsA(UMaterialInterface::StaticClass()))
-						{
-							SubType = UMaterialInterface::StaticClass()->GetFName();
-							ClassToCreate = UMaterialInterface::StaticClass();
-						}
-						else
+						const int32 MaterialIndex = TargetRCField->FieldPathInfo.Segments[0].ArrayIndex;
+						const uint8* ObjPtrContainer = ArrayHelper.GetRawPtr(MaterialIndex);
+						if (const UObject* CurrentObject = InnerObjectProperty->GetObjectPropertyValue(ObjPtrContainer))
 						{
-							SubType = CurrentObject->GetClass()->GetFName();
-							ClassToCreate = CurrentObject->GetClass();
+							if (CurrentObject->IsA(UMaterialInterface::StaticClass()))
+							{
+								SubType = UMaterialInterface::StaticClass()->GetFName();
+								ClassToCreate = UMaterialInterface::StaticClass();
+							}
+							else
+							{
+								SubType = CurrentObject->GetClass()->GetFName();
+								ClassToCreate = CurrentObject->GetClass();
+							}
 						}
 					}
 				}
 			}
 		}
+		else if (const FStructProperty* InnerStructProperty = CastField<FStructProperty>(ArrayProperty->Inner))
+		{
+			if (InnerStructProperty->Struct->GetFName() == NAME_LinearColor)
+			{
+				SubType = NAME_Color;
+			}
+			else
+			{
+				SubType = InnerStructProperty->Struct->GetFName();
+			}
+		}
 	}
 }
 
@@ -281,20 +378,23 @@ void URemoteControlPropertyIdRegistry::PerformChainReaction(const FRemoteControl
 					// Note : Specialization for Materials.
 					else if (FArrayProperty* ArrayProperty = CastField<FArrayProperty>(Property))
 					{
-						if (uint8* PropertyValuePtr = ArrayProperty->ContainerPtrToValuePtr<uint8>(BoundObject))
+						if (ArrayProperty->GetOwner<UClass>())
 						{
-							if (FObjectProperty* InnerObjectProperty = CastField<FObjectProperty>(ArrayProperty->Inner))
+							if (uint8* PropertyValuePtr = ArrayProperty->ContainerPtrToValuePtr<uint8>(BoundObject))
 							{
-								FScriptArrayHelper ArrayHelper(ArrayProperty, PropertyValuePtr);
-								if (TSharedPtr<FRemoteControlField> TargetRCField = StaticCastSharedPtr<FRemoteControlField>(TargetRCProperty); !TargetRCField->FieldPathInfo.Segments.IsEmpty())
+								if (FObjectProperty* InnerObjectProperty = CastField<FObjectProperty>(ArrayProperty->Inner))
 								{
-									const int32 MaterialIndex = TargetRCField->FieldPathInfo.Segments[0].ArrayIndex;
-									uint8* ObjPtrContainer = ArrayHelper.GetRawPtr(MaterialIndex);
-									UObject* CurrentObject = InnerObjectProperty->GetObjectPropertyValue(ObjPtrContainer);
-									if (CurrentObject && CurrentObject->IsA(InArgs.SourceClass))
+									FScriptArrayHelper ArrayHelper(ArrayProperty, PropertyValuePtr);
+									if (TSharedPtr<FRemoteControlField> TargetRCField = StaticCastSharedPtr<FRemoteControlField>(TargetRCProperty); !TargetRCField->FieldPathInfo.Segments.IsEmpty())
 									{
-										InnerObjectProperty->SetObjectPropertyValue(ObjPtrContainer, InArgs.SourceObject);
-										bCopyComplete = true;
+										const int32 MaterialIndex = TargetRCField->FieldPathInfo.Segments[0].ArrayIndex;
+										uint8* ObjPtrContainer = ArrayHelper.GetRawPtr(MaterialIndex);
+										UObject* CurrentObject = InnerObjectProperty->GetObjectPropertyValue(ObjPtrContainer);
+										if (CurrentObject && CurrentObject->IsA(InArgs.SourceClass))
+										{
+											InnerObjectProperty->SetObjectPropertyValue(ObjPtrContainer, InArgs.SourceObject);
+											bCopyComplete = true;
+										}
 									}
 								}
 							}
@@ -385,7 +485,7 @@ void URemoteControlPropertyIdRegistry::PerformChainReaction(const FRemoteControl
 						FMemoryReader Reader(Buffer);
 						FCborStructDeserializerBackend ReaderBackend(Reader);
 						FRCObjectReference TargetObjectRef;
-						if (!GetObjectRef(TargetRCProperty, ERCAccess::WRITE_ACCESS, TargetObjectRef))
+						if (!RemoteControlPropertyIdUtilities::GetObjectRef(TargetRCProperty, ERCAccess::WRITE_ACCESS, TargetObjectRef))
 						{
 							continue;
 						}
@@ -513,95 +613,101 @@ bool URemoteControlPropertyIdRegistry::CopyNonUClassOwnerProperty(const FRemoteC
 	return TryCopyNonUClassOwnerProperty(InArgs.VirtualProperty, TargetRCProperty, InProperty, BoundObject, TargetRCProperty->FieldPathInfo, 0, nullptr, bCheckByteEnumComparison);
 }
 
-bool URemoteControlPropertyIdRegistry::TryCopyNonUClassOwnerProperty(const TObjectPtr<URCVirtualPropertySelfContainer>& InVirtualPropertySelfContainer, const TSharedPtr<FRemoteControlProperty>& TargetRCProperty, FProperty* InProperty, UObject* BoundObject, FRCFieldPathInfo Segment, const int32 IndexSegment, uint8* PropertyContainer, const bool bCheckByteEnumComparison)
+bool URemoteControlPropertyIdRegistry::TryCopyNonUClassOwnerProperty(const TObjectPtr<URCVirtualPropertySelfContainer>& InVirtualPropertySelfContainer, const TSharedPtr<FRemoteControlProperty>& InTargetRCProperty, FProperty* InProperty, UObject* InBoundObject, FRCFieldPathInfo InSegment, const int32 InIndexSegment, uint8* InPropertyContainer, const bool bCheckByteEnumComparison)
 {
-	if (IndexSegment >= Segment.GetSegmentCount())
+	if (InIndexSegment >= InSegment.GetSegmentCount())
 	{
 		return false;
 	}
 
-	if (const FArrayProperty* ArrayProperty = CastField<FArrayProperty>(Segment.Segments[IndexSegment].ResolvedData.Field))
+	if (const FArrayProperty* ArrayProperty = CastField<FArrayProperty>(InSegment.Segments[InIndexSegment].ResolvedData.Field))
 	{
-		if (const uint8* PropertyValuePtr = ArrayProperty->ContainerPtrToValuePtr<uint8>(BoundObject))
+		if (ArrayProperty->GetOwner<UClass>())
 		{
-			FScriptArrayHelper ArrayHelper(ArrayProperty, PropertyValuePtr);
-			if (const TSharedPtr<FRemoteControlField> TargetRCField = StaticCastSharedPtr<FRemoteControlField>(TargetRCProperty);
-				!TargetRCField->FieldPathInfo.Segments.IsEmpty())
+			if (const uint8* PropertyValuePtr = ArrayProperty->ContainerPtrToValuePtr<uint8>(InBoundObject))
 			{
-				const int32 Index = Segment.Segments[IndexSegment].ArrayIndex;
-				uint8* PtrContainer = ArrayHelper.GetElementPtr(Index);
-				if (IndexSegment < Segment.GetSegmentCount() - 2)
+				uint8* PtrContainer = nullptr;
+				if (RemoteControlPropertyIdUtilities::CopyNonUClassOwnerArray(InVirtualPropertySelfContainer, ArrayProperty, InTargetRCProperty, InProperty, InSegment, InIndexSegment, PropertyValuePtr, bCheckByteEnumComparison))
 				{
-					return TryCopyNonUClassOwnerProperty(InVirtualPropertySelfContainer, TargetRCProperty, InProperty, BoundObject, Segment, IndexSegment + 1, PtrContainer, bCheckByteEnumComparison);
-				}
-				if (const FArrayProperty* ArrayRealProperty = CastField<FArrayProperty>(InProperty))
-				{
-					return InVirtualPropertySelfContainer->CopyCompleteValue(ArrayRealProperty->Inner, PtrContainer, bCheckByteEnumComparison);
+					return true;
 				}
-				return false;
+				return TryCopyNonUClassOwnerProperty(InVirtualPropertySelfContainer, InTargetRCProperty, InProperty, InBoundObject, InSegment, InIndexSegment + 1, PtrContainer, bCheckByteEnumComparison);
+			}
+		}
+		else if (InPropertyContainer)
+		{
+			uint8* PtrContainer = nullptr;
+			if (RemoteControlPropertyIdUtilities::CopyNonUClassOwnerArray(InVirtualPropertySelfContainer, ArrayProperty, InTargetRCProperty, InProperty, InSegment, InIndexSegment, InPropertyContainer, bCheckByteEnumComparison))
+			{
+				return true;
 			}
+			return TryCopyNonUClassOwnerProperty(InVirtualPropertySelfContainer, InTargetRCProperty, InProperty, InBoundObject, InSegment, InIndexSegment + 1, PtrContainer, bCheckByteEnumComparison);
 		}
 	}
-	else if (const FSetProperty* SetProperty = CastField<FSetProperty>(Segment.Segments[IndexSegment].ResolvedData.Field))
+	else if (const FSetProperty* SetProperty = CastField<FSetProperty>(InSegment.Segments[InIndexSegment].ResolvedData.Field))
 	{
-		if (const uint8* PropertyValuePtr = SetProperty->ContainerPtrToValuePtr<uint8>(BoundObject))
+		if (SetProperty->GetOwner<UClass>())
 		{
-			FScriptSetHelper SetHelper(SetProperty, PropertyValuePtr);
-			if (const TSharedPtr<FRemoteControlField> TargetRCField = StaticCastSharedPtr<FRemoteControlField>(TargetRCProperty);
-				!TargetRCField->FieldPathInfo.Segments.IsEmpty())
+			if (const uint8* PropertyValuePtr = SetProperty->ContainerPtrToValuePtr<uint8>(InBoundObject))
 			{
-				const int32 Index = Segment.Segments[IndexSegment].ArrayIndex;
-				uint8* PtrContainer = SetHelper.GetElementPtr(Index);
-				if (IndexSegment < Segment.GetSegmentCount() - 2)
-				{
-					return TryCopyNonUClassOwnerProperty(InVirtualPropertySelfContainer, TargetRCProperty, InProperty, BoundObject, Segment, IndexSegment + 1, PtrContainer, bCheckByteEnumComparison);
-				}
-				if (const FSetProperty* SetRealProperty = CastField<FSetProperty>(InProperty))
+				uint8* PtrContainer = nullptr;
+				if (RemoteControlPropertyIdUtilities::CopyNonUClassOwnerSet(InVirtualPropertySelfContainer, SetProperty, InTargetRCProperty, InProperty, InSegment, InIndexSegment, PropertyValuePtr, bCheckByteEnumComparison))
 				{
-					return InVirtualPropertySelfContainer->CopyCompleteValue(SetRealProperty->ElementProp, PtrContainer, bCheckByteEnumComparison);
+					return true;
 				}
-				return false;
+				return TryCopyNonUClassOwnerProperty(InVirtualPropertySelfContainer, InTargetRCProperty, InProperty, InBoundObject, InSegment, InIndexSegment + 1, PtrContainer, bCheckByteEnumComparison);
+			}
+		}
+		else if (InPropertyContainer)
+		{
+			uint8* PtrContainer = nullptr;
+			if (RemoteControlPropertyIdUtilities::CopyNonUClassOwnerSet(InVirtualPropertySelfContainer, SetProperty, InTargetRCProperty, InProperty, InSegment, InIndexSegment, InPropertyContainer, bCheckByteEnumComparison))
+			{
+				return true;
 			}
+			return TryCopyNonUClassOwnerProperty(InVirtualPropertySelfContainer, InTargetRCProperty, InProperty, InBoundObject, InSegment, InIndexSegment + 1, PtrContainer, bCheckByteEnumComparison);
 		}
 	}
-	else if (const FMapProperty* MapProperty = CastField<FMapProperty>(Segment.Segments[IndexSegment].ResolvedData.Field))
+	else if (const FMapProperty* MapProperty = CastField<FMapProperty>(InSegment.Segments[InIndexSegment].ResolvedData.Field))
 	{
-		if (const uint8* PropertyValuePtr = MapProperty->ContainerPtrToValuePtr<uint8>(BoundObject))
+		if (MapProperty->GetOwner<UClass>())
 		{
-			FScriptMapHelper MapHelper(MapProperty, PropertyValuePtr);
-			if (const TSharedPtr<FRemoteControlField> TargetRCField = StaticCastSharedPtr<FRemoteControlField>(TargetRCProperty);
-				!TargetRCField->FieldPathInfo.Segments.IsEmpty())
+			if (const uint8* PropertyValuePtr = MapProperty->ContainerPtrToValuePtr<uint8>(InBoundObject))
 			{
-				const int32 Index = Segment.Segments[IndexSegment].ArrayIndex;
-				uint8* PtrContainer = MapHelper.GetValuePtr(Index);
-				if (IndexSegment < Segment.GetSegmentCount() - 2)
-				{
-					return TryCopyNonUClassOwnerProperty(InVirtualPropertySelfContainer, TargetRCProperty, InProperty, BoundObject, Segment, IndexSegment + 1, PtrContainer, bCheckByteEnumComparison);
-				}
-				if (const FMapProperty* MapRealProperty = CastField<FMapProperty>(InProperty))
+				uint8* PtrContainer = nullptr;
+				if (RemoteControlPropertyIdUtilities::CopyNonUClassOwnerMap(InVirtualPropertySelfContainer, MapProperty, InTargetRCProperty, InProperty, InSegment, InIndexSegment, PropertyValuePtr, bCheckByteEnumComparison))
 				{
-					return InVirtualPropertySelfContainer->CopyCompleteValue(MapRealProperty->ValueProp, PtrContainer, bCheckByteEnumComparison);
+					return true;
 				}
-				return false;
+				return TryCopyNonUClassOwnerProperty(InVirtualPropertySelfContainer, InTargetRCProperty, InProperty, InBoundObject, InSegment, InIndexSegment + 1, PtrContainer, bCheckByteEnumComparison);
+			}
+		}
+		else if (InPropertyContainer)
+		{
+			uint8* PtrContainer = nullptr;
+			if (RemoteControlPropertyIdUtilities::CopyNonUClassOwnerMap(InVirtualPropertySelfContainer, MapProperty, InTargetRCProperty, InProperty, InSegment, InIndexSegment, InPropertyContainer, bCheckByteEnumComparison))
+			{
+				return true;
 			}
+			return TryCopyNonUClassOwnerProperty(InVirtualPropertySelfContainer, InTargetRCProperty, InProperty, InBoundObject, InSegment, InIndexSegment + 1, PtrContainer, bCheckByteEnumComparison);
 		}
 	}
-	else if (const FStructProperty* StructProperty = CastField<FStructProperty>(Segment.Segments[IndexSegment].ResolvedData.Field))
+	else if (const FStructProperty* StructProperty = CastField<FStructProperty>(InSegment.Segments[InIndexSegment].ResolvedData.Field))
 	{
 		uint8* InnerStructPtrContainer;
-		if (PropertyContainer != nullptr)
+		if (InPropertyContainer != nullptr)
 		{
-			InnerStructPtrContainer = StructProperty->ContainerPtrToValuePtr<uint8>(PropertyContainer);
+			InnerStructPtrContainer = StructProperty->ContainerPtrToValuePtr<uint8>(InPropertyContainer);
 		}
 		else
 		{
-			InnerStructPtrContainer = StructProperty->ContainerPtrToValuePtr<uint8>(BoundObject);
+			InnerStructPtrContainer = StructProperty->ContainerPtrToValuePtr<uint8>(InBoundObject);
 		}
 		if (uint8* PtrContainer = InProperty->ContainerPtrToValuePtr<uint8>(InnerStructPtrContainer))
 		{
-			if (IndexSegment < Segment.GetSegmentCount() - 2)
+			if (InIndexSegment < InSegment.GetSegmentCount() - 2)
 			{
-				return TryCopyNonUClassOwnerProperty(InVirtualPropertySelfContainer, TargetRCProperty, InProperty, BoundObject, Segment, IndexSegment + 1, InnerStructPtrContainer, bCheckByteEnumComparison);
+				return TryCopyNonUClassOwnerProperty(InVirtualPropertySelfContainer, InTargetRCProperty, InProperty, InBoundObject, InSegment, InIndexSegment + 1, InnerStructPtrContainer, bCheckByteEnumComparison);
 			}
 
 			// FLinearColor are treated as FColor to avoid creating 2 different widget for them and avoid confusion
@@ -620,6 +726,11 @@ bool URemoteControlPropertyIdRegistry::TryCopyNonUClassOwnerProperty(const TObje
 					return true;
 				}
 			}
+
+			if (InProperty->IsA<FArrayProperty>() || InProperty->IsA<FSetProperty>() || InProperty->IsA<FMapProperty>())
+			{
+				return TryCopyNonUClassOwnerProperty(InVirtualPropertySelfContainer, InTargetRCProperty, InProperty, InBoundObject, InSegment, InIndexSegment + 1, InnerStructPtrContainer, bCheckByteEnumComparison);
+			}
 			return InVirtualPropertySelfContainer->CopyCompleteValue(InProperty, PtrContainer, bCheckByteEnumComparison);
 		}
 	}
