[RemoteControl]
Adding multi selection.

diff --git a/Engine/Plugins/VirtualProduction/RemoteControl/Source/RemoteControl/Private/RemoteControlPreset.cpp b/Engine/Plugins/VirtualProduction/RemoteControl/Source/RemoteControl/Private/RemoteControlPreset.cpp
--- a/Engine/Plugins/VirtualProduction/RemoteControl/Source/RemoteControl/Private/RemoteControlPreset.cpp
+++ b/Engine/Plugins/VirtualProduction/RemoteControl/Source/RemoteControl/Private/RemoteControlPreset.cpp
@@ -250,22 +250,6 @@ FRemoteControlPresetGroup* FRemoteControlPresetLayout::FindGroupFromField(FGuid
 	return nullptr;
 }
 
-bool FRemoteControlPresetLayout::MoveField(FGuid FieldId, FGuid TargetGroupId)
-{
-	FFieldSwapArgs FieldSwapArgs;
-	FieldSwapArgs.DraggedFieldId = FieldId;
-	FGuid OriginGroupId = FindGroupFromField(FieldId)->Id;
-
-	if (OriginGroupId.IsValid() && TargetGroupId.IsValid())
-	{
-		FieldSwapArgs.OriginGroupId = MoveTemp(OriginGroupId);
-		FieldSwapArgs.TargetGroupId = MoveTemp(TargetGroupId);
-		SwapFields(FieldSwapArgs);
-		return true;
-	}
-	return false;
-}
-
 void FRemoteControlPresetLayout::SwapGroups(FGuid OriginGroupId, FGuid TargetGroupId)
 {
 	FRemoteControlPresetGroup* OriginGroup = GetGroup(OriginGroupId);
@@ -308,73 +292,93 @@ void FRemoteControlPresetLayout::SwapGroups(FGuid OriginGroupId, FGuid TargetGro
 	}
 }
 
-void FRemoteControlPresetLayout::SwapFields(const FFieldSwapArgs& FieldSwapArgs)
+void FRemoteControlPresetLayout::SwapFields(const FFieldSwapArgs& InFieldSwapArgs)
 {
-	FRemoteControlPresetGroup* DragOriginGroup = GetGroup(FieldSwapArgs.OriginGroupId);
-	FRemoteControlPresetGroup* DragTargetGroup = GetGroup(FieldSwapArgs.TargetGroupId);
+	FRemoteControlPresetGroup* DragOriginGroup = GetGroup(InFieldSwapArgs.OriginGroupId);
+	FRemoteControlPresetGroup* DragTargetGroup = GetGroup(InFieldSwapArgs.TargetGroupId);
 
-	if (!DragOriginGroup || !DragTargetGroup)
+	if (!DragOriginGroup || !DragTargetGroup || InFieldSwapArgs.DraggedFieldsIds.Num() < 1)
 	{
 		return;
 	}
 
-	int32 DragOriginFieldIndex = DragOriginGroup->AccessFields().IndexOfByKey(FieldSwapArgs.DraggedFieldId);
-	int32 DragTargetFieldIndex = DragTargetGroup->AccessFields().IndexOfByKey(FieldSwapArgs.TargetFieldId);
+	int32 DragOriginFieldIndex = DragOriginGroup->AccessFields().IndexOfByKey(InFieldSwapArgs.DraggedFieldsIds[0]);
+	int32 DragTargetInitialFieldIndex = DragTargetGroup->AccessFields().IndexOfByKey(InFieldSwapArgs.TargetFieldId);
 
 	if (DragOriginFieldIndex == INDEX_NONE)
 	{
 		return;
 	}
 
-	if (FieldSwapArgs.OriginGroupId == FieldSwapArgs.TargetGroupId && DragTargetFieldIndex != INDEX_NONE)
+	if (InFieldSwapArgs.OriginGroupId == InFieldSwapArgs.TargetGroupId && DragTargetInitialFieldIndex != INDEX_NONE)
 	{
-		if (DragTargetFieldIndex > DragOriginFieldIndex)
-		{
-			DragTargetFieldIndex += 1;
-		}
-		else
-		{
-			DragOriginFieldIndex += 1;
-		}
-
 		// Here we don't want to trigger add/delete delegates since the fields just get moved around.
 		TArray<FGuid>& Fields = DragTargetGroup->AccessFields();
 
-		FRCTransactionListenerHelper<const FGuid&, const TArray<FGuid>&>(ERCTransaction::Undo, Owner->GetPresetId(), OnFieldOrderChanged(), FieldSwapArgs.TargetGroupId, Fields);
+		FRCTransactionListenerHelper<const FGuid&, const TArray<FGuid>&>(ERCTransaction::Undo, Owner->GetPresetId(), OnFieldOrderChanged(), InFieldSwapArgs.TargetGroupId, Fields);
 
-		Fields.Insert(FieldSwapArgs.DraggedFieldId, DragTargetFieldIndex);
-		Fields.Swap(DragTargetFieldIndex, DragOriginFieldIndex);
-		Fields.RemoveAt(DragOriginFieldIndex);
+		int32 TargetFieldOffset = 0;
+		int32 DragOriginSelectedFieldIndex = Fields.IndexOfByKey(InFieldSwapArgs.DraggedFieldsIds[0]);
+		int32 DragTargetFieldIndex = Fields.IndexOfByKey(InFieldSwapArgs.TargetFieldId);
+		const bool bInitialTargetIsGreaterThanOrigin = DragTargetFieldIndex >= DragOriginSelectedFieldIndex;
+		for (FGuid SelectedField : InFieldSwapArgs.DraggedFieldsIds)
+		{
+			DragOriginSelectedFieldIndex = Fields.IndexOfByKey(SelectedField);
+			DragTargetFieldIndex = Fields.IndexOfByKey(InFieldSwapArgs.TargetFieldId);
+			if (DragTargetFieldIndex == DragOriginSelectedFieldIndex)
+			{
+				continue;
+			}
+			if (bInitialTargetIsGreaterThanOrigin)
+			{
+				TargetFieldOffset++;
+				DragTargetFieldIndex += TargetFieldOffset;
+			}
+			if (DragTargetFieldIndex < DragOriginSelectedFieldIndex)
+			{
+				DragOriginSelectedFieldIndex += 1;
+			}
+			Fields.Insert(SelectedField, DragTargetFieldIndex);
+			Fields.Swap(DragTargetFieldIndex, DragOriginSelectedFieldIndex);
+			Fields.RemoveAt(DragOriginSelectedFieldIndex);
+		}
 	
 		Owner->CacheLayoutData();
-		OnFieldOrderChangedDelegate.Broadcast(FieldSwapArgs.TargetGroupId, Fields);
+		OnFieldOrderChangedDelegate.Broadcast(InFieldSwapArgs.TargetGroupId, Fields);
 		Owner->OnPresetLayoutModified().Broadcast(Owner.Get());
 
-		FRCTransactionListenerHelper<const FGuid&, const TArray<FGuid>&>(ERCTransaction::Redo, Owner->GetPresetId(), OnFieldOrderChanged(), FieldSwapArgs.TargetGroupId, Fields);
+		FRCTransactionListenerHelper<const FGuid&, const TArray<FGuid>&>(ERCTransaction::Redo, Owner->GetPresetId(), OnFieldOrderChanged(), InFieldSwapArgs.TargetGroupId, Fields);
 
 		FRCTransactionListenerHelper<URemoteControlPreset*>(ERCTransaction::Undo, Owner->GetPresetId(), Owner->OnPresetLayoutModified(), Owner.Get());
 		FRCTransactionListenerHelper<URemoteControlPreset*>(ERCTransaction::Redo, Owner->GetPresetId(), Owner->OnPresetLayoutModified(), Owner.Get());
 	}
 	else
 	{
-		DragOriginGroup->AccessFields().RemoveAt(DragOriginFieldIndex);
+		// We do it reverse so that the order is maintained
+		for (int32 Index = InFieldSwapArgs.DraggedFieldsIds.Num() - 1; Index >= 0; --Index)
+		{
+			TArray<FGuid>& Fields = DragOriginGroup->AccessFields();
+			DragOriginFieldIndex = Fields.IndexOfByKey(InFieldSwapArgs.DraggedFieldsIds[Index]);
 
-		DragTargetFieldIndex = DragTargetFieldIndex == INDEX_NONE ? 0 : DragTargetFieldIndex;
-		DragTargetGroup->AccessFields().Insert(FieldSwapArgs.DraggedFieldId, DragTargetFieldIndex);
+			DragOriginGroup->AccessFields().RemoveAt(DragOriginFieldIndex);
 
-		FRCTransactionListenerHelper<const FGuid&, const FGuid&, int32>(ERCTransaction::Undo, Owner->GetPresetId(), OnFieldDeleted(), FieldSwapArgs.TargetGroupId, FieldSwapArgs.DraggedFieldId, DragTargetFieldIndex);
-		FRCTransactionListenerHelper<const FGuid&, const FGuid&, int32>(ERCTransaction::Undo, Owner->GetPresetId(), OnFieldAdded(), FieldSwapArgs.OriginGroupId, FieldSwapArgs.DraggedFieldId, DragOriginFieldIndex);
+			DragTargetInitialFieldIndex = DragTargetInitialFieldIndex == INDEX_NONE ? 0 : DragTargetInitialFieldIndex;
+			DragTargetGroup->AccessFields().Insert(InFieldSwapArgs.DraggedFieldsIds[Index], DragTargetInitialFieldIndex);
 
-		FRCTransactionListenerHelper<const FGuid&, const FGuid&, int32>(ERCTransaction::Redo, Owner->GetPresetId(), OnFieldDeleted(), FieldSwapArgs.OriginGroupId, FieldSwapArgs.DraggedFieldId, DragOriginFieldIndex);
-		FRCTransactionListenerHelper<const FGuid&, const FGuid&, int32>(ERCTransaction::Redo, Owner->GetPresetId(), OnFieldAdded(), FieldSwapArgs.TargetGroupId, FieldSwapArgs.DraggedFieldId, DragTargetFieldIndex);
+			FRCTransactionListenerHelper<const FGuid&, const FGuid&, int32>(ERCTransaction::Undo, Owner->GetPresetId(), OnFieldDeleted(), InFieldSwapArgs.TargetGroupId, InFieldSwapArgs.DraggedFieldsIds[Index], DragTargetInitialFieldIndex);
+			FRCTransactionListenerHelper<const FGuid&, const FGuid&, int32>(ERCTransaction::Undo, Owner->GetPresetId(), OnFieldAdded(), InFieldSwapArgs.OriginGroupId, InFieldSwapArgs.DraggedFieldsIds[Index], DragOriginFieldIndex);
 
-		Owner->CacheLayoutData();
-		OnFieldDeletedDelegate.Broadcast(FieldSwapArgs.OriginGroupId, FieldSwapArgs.DraggedFieldId, DragOriginFieldIndex);
-		OnFieldAddedDelegate.Broadcast(FieldSwapArgs.TargetGroupId, FieldSwapArgs.DraggedFieldId, DragTargetFieldIndex);
-		Owner->OnPresetLayoutModified().Broadcast(Owner.Get());
+			FRCTransactionListenerHelper<const FGuid&, const FGuid&, int32>(ERCTransaction::Redo, Owner->GetPresetId(), OnFieldDeleted(), InFieldSwapArgs.OriginGroupId, InFieldSwapArgs.DraggedFieldsIds[Index], DragOriginFieldIndex);
+			FRCTransactionListenerHelper<const FGuid&, const FGuid&, int32>(ERCTransaction::Redo, Owner->GetPresetId(), OnFieldAdded(), InFieldSwapArgs.TargetGroupId, InFieldSwapArgs.DraggedFieldsIds[Index], DragTargetInitialFieldIndex);
 
-		FRCTransactionListenerHelper<URemoteControlPreset*>(ERCTransaction::Undo, Owner->GetPresetId(), Owner->OnPresetLayoutModified(), Owner.Get());
-		FRCTransactionListenerHelper<URemoteControlPreset*>(ERCTransaction::Redo, Owner->GetPresetId(), Owner->OnPresetLayoutModified(), Owner.Get());
+			Owner->CacheLayoutData();
+			OnFieldDeletedDelegate.Broadcast(InFieldSwapArgs.OriginGroupId, InFieldSwapArgs.DraggedFieldsIds[Index], DragOriginFieldIndex);
+			OnFieldAddedDelegate.Broadcast(InFieldSwapArgs.TargetGroupId, InFieldSwapArgs.DraggedFieldsIds[Index], DragTargetInitialFieldIndex);
+			Owner->OnPresetLayoutModified().Broadcast(Owner.Get());
+
+			FRCTransactionListenerHelper<URemoteControlPreset*>(ERCTransaction::Undo, Owner->GetPresetId(), Owner->OnPresetLayoutModified(), Owner.Get());
+			FRCTransactionListenerHelper<URemoteControlPreset*>(ERCTransaction::Redo, Owner->GetPresetId(), Owner->OnPresetLayoutModified(), Owner.Get());
+		}
 	}
 }
 
@@ -383,38 +387,51 @@ void FRemoteControlPresetLayout::SwapFieldsDefaultGroup(const FFieldSwapArgs& In
 	FRemoteControlPresetGroup* DragOriginGroup = GetGroup(InFieldRealGroup);
 	FRemoteControlPresetGroup* DragTargetGroup = GetGroup(InFieldSwapArgs.TargetGroupId);
 
-	if (!DragOriginGroup || !DragTargetGroup)
+	if (!DragOriginGroup || !DragTargetGroup || InFieldSwapArgs.DraggedFieldsIds.Num() < 1)
 	{
 		return;
 	}
 
-	int32 DragOriginFieldIndex = InEntities.IndexOfByKey(InFieldSwapArgs.DraggedFieldId);
-	int32 DragTargetFieldIndex = InEntities.IndexOfByKey(InFieldSwapArgs.TargetFieldId);
+	int32 DragOriginFieldIndex = InEntities.IndexOfByKey(InFieldSwapArgs.DraggedFieldsIds[0]);
+	int32 DragTargetInitialFieldIndex = InEntities.IndexOfByKey(InFieldSwapArgs.TargetFieldId);
 
 	if (DragOriginFieldIndex == INDEX_NONE)
 	{
 		return;
 	}
 
-	if (DragTargetFieldIndex != INDEX_NONE)
+	if (DragTargetInitialFieldIndex != INDEX_NONE)
 	{
-		if (DragTargetFieldIndex > DragOriginFieldIndex)
-		{
-			DragTargetFieldIndex += 1;
-		}
-		else
-		{
-			DragOriginFieldIndex += 1;
-		}
-
 		// Here we don't want to trigger add/delete delegates since the fields just get moved around.
 		TArray<FGuid>& Fields = InEntities;
 
 		FRCTransactionListenerHelper<const FGuid&, const TArray<FGuid>&>(ERCTransaction::Undo, Owner->GetPresetId(), OnFieldOrderChanged(), InFieldSwapArgs.TargetGroupId, Fields);
 
-		Fields.Insert(InFieldSwapArgs.DraggedFieldId, DragTargetFieldIndex);
-		Fields.Swap(DragTargetFieldIndex, DragOriginFieldIndex);
-		Fields.RemoveAt(DragOriginFieldIndex);
+		int32 TargetFieldOffset = 0;
+		int32 DragOriginSelectedFieldIndex = InEntities.IndexOfByKey(InFieldSwapArgs.DraggedFieldsIds[0]);
+		int32 DragTargetFieldIndex = InEntities.IndexOfByKey(InFieldSwapArgs.TargetFieldId);
+		const bool bInitialTargetIsGreaterThanOrigin = DragTargetFieldIndex >= DragOriginSelectedFieldIndex;
+		for (FGuid SelectedField : InFieldSwapArgs.DraggedFieldsIds)
+		{
+			DragOriginSelectedFieldIndex = InEntities.IndexOfByKey(SelectedField);
+			DragTargetFieldIndex = InEntities.IndexOfByKey(InFieldSwapArgs.TargetFieldId);
+			if (DragTargetFieldIndex == DragOriginSelectedFieldIndex)
+			{
+				continue;
+			}
+			if (bInitialTargetIsGreaterThanOrigin)
+			{
+				TargetFieldOffset++;
+				DragTargetFieldIndex += TargetFieldOffset;
+			}
+			if (DragTargetFieldIndex < DragOriginSelectedFieldIndex)
+			{
+				DragOriginSelectedFieldIndex += 1;
+			}
+			Fields.Insert(SelectedField, DragTargetFieldIndex);
+			Fields.Swap(DragTargetFieldIndex, DragOriginSelectedFieldIndex);
+			Fields.RemoveAt(DragOriginSelectedFieldIndex);
+		}
 	
 		Owner->CacheLayoutData();
 		OnFieldOrderChangedDelegate.Broadcast(GetDefaultGroup().Id, Fields);
@@ -427,24 +444,35 @@ void FRemoteControlPresetLayout::SwapFieldsDefaultGroup(const FFieldSwapArgs& In
 	}
 	else if (!IsDefaultGroup(DragTargetGroup->Id))
 	{
-		DragOriginGroup->AccessFields().RemoveSwap(InFieldSwapArgs.DraggedFieldId, true);
+		// We do it reverse so that the order is maintained
+		for (int32 Index = InFieldSwapArgs.DraggedFieldsIds.Num() - 1; Index >= 0; --Index)
+		{
+			if (DragTargetGroup->AccessFields().Contains(InFieldSwapArgs.DraggedFieldsIds[Index]))
+			{
+				// Skip if the Target group already contains the Field.
+				continue;
+			}
+			DragOriginFieldIndex = InEntities.IndexOfByKey(InFieldSwapArgs.DraggedFieldsIds[Index]);
 
-		DragTargetFieldIndex = DragTargetFieldIndex == INDEX_NONE ? 0 : DragTargetFieldIndex;
-		DragTargetGroup->AccessFields().Insert(InFieldSwapArgs.DraggedFieldId, DragTargetFieldIndex);
+			DragOriginGroup->AccessFields().RemoveSwap(InFieldSwapArgs.DraggedFieldsIds[Index], true);
 
-		FRCTransactionListenerHelper<const FGuid&, const FGuid&, int32>(ERCTransaction::Undo, Owner->GetPresetId(), OnFieldDeleted(), InFieldSwapArgs.TargetGroupId, InFieldSwapArgs.DraggedFieldId, DragTargetFieldIndex);
-		FRCTransactionListenerHelper<const FGuid&, const FGuid&, int32>(ERCTransaction::Undo, Owner->GetPresetId(), OnFieldAdded(), InFieldSwapArgs.OriginGroupId, InFieldSwapArgs.DraggedFieldId, DragOriginFieldIndex);
+			DragTargetInitialFieldIndex = DragTargetInitialFieldIndex == INDEX_NONE ? 0 : DragTargetInitialFieldIndex;
+			DragTargetGroup->AccessFields().Insert(InFieldSwapArgs.DraggedFieldsIds[Index], DragTargetInitialFieldIndex);
 
-		FRCTransactionListenerHelper<const FGuid&, const FGuid&, int32>(ERCTransaction::Redo, Owner->GetPresetId(), OnFieldDeleted(), InFieldSwapArgs.OriginGroupId, InFieldSwapArgs.DraggedFieldId, DragOriginFieldIndex);
-		FRCTransactionListenerHelper<const FGuid&, const FGuid&, int32>(ERCTransaction::Redo, Owner->GetPresetId(), OnFieldAdded(), InFieldSwapArgs.TargetGroupId, InFieldSwapArgs.DraggedFieldId, DragTargetFieldIndex);
+			FRCTransactionListenerHelper<const FGuid&, const FGuid&, int32>(ERCTransaction::Undo, Owner->GetPresetId(), OnFieldDeleted(), InFieldSwapArgs.TargetGroupId, InFieldSwapArgs.DraggedFieldsIds[Index], DragTargetInitialFieldIndex);
+			FRCTransactionListenerHelper<const FGuid&, const FGuid&, int32>(ERCTransaction::Undo, Owner->GetPresetId(), OnFieldAdded(), InFieldSwapArgs.OriginGroupId, InFieldSwapArgs.DraggedFieldsIds[Index], DragOriginFieldIndex);
 
-		Owner->CacheLayoutData();
-		OnFieldDeletedDelegate.Broadcast(InFieldSwapArgs.OriginGroupId, InFieldSwapArgs.DraggedFieldId, DragOriginFieldIndex);
-		OnFieldAddedDelegate.Broadcast(InFieldSwapArgs.TargetGroupId, InFieldSwapArgs.DraggedFieldId, DragTargetFieldIndex);
-		Owner->OnPresetLayoutModified().Broadcast(Owner.Get());
+			FRCTransactionListenerHelper<const FGuid&, const FGuid&, int32>(ERCTransaction::Redo, Owner->GetPresetId(), OnFieldDeleted(), InFieldSwapArgs.OriginGroupId, InFieldSwapArgs.DraggedFieldsIds[Index], DragOriginFieldIndex);
+			FRCTransactionListenerHelper<const FGuid&, const FGuid&, int32>(ERCTransaction::Redo, Owner->GetPresetId(), OnFieldAdded(), InFieldSwapArgs.TargetGroupId, InFieldSwapArgs.DraggedFieldsIds[Index], DragTargetInitialFieldIndex);
 
-		FRCTransactionListenerHelper<URemoteControlPreset*>(ERCTransaction::Undo, Owner->GetPresetId(), Owner->OnPresetLayoutModified(), Owner.Get());
-		FRCTransactionListenerHelper<URemoteControlPreset*>(ERCTransaction::Redo, Owner->GetPresetId(), Owner->OnPresetLayoutModified(), Owner.Get());
+			Owner->CacheLayoutData();
+			OnFieldDeletedDelegate.Broadcast(InFieldSwapArgs.OriginGroupId, InFieldSwapArgs.DraggedFieldsIds[Index], DragOriginFieldIndex);
+			OnFieldAddedDelegate.Broadcast(InFieldSwapArgs.TargetGroupId, InFieldSwapArgs.DraggedFieldsIds[Index], DragTargetInitialFieldIndex);
+			Owner->OnPresetLayoutModified().Broadcast(Owner.Get());
+
+			FRCTransactionListenerHelper<URemoteControlPreset*>(ERCTransaction::Undo, Owner->GetPresetId(), Owner->OnPresetLayoutModified(), Owner.Get());
+			FRCTransactionListenerHelper<URemoteControlPreset*>(ERCTransaction::Redo, Owner->GetPresetId(), Owner->OnPresetLayoutModified(), Owner.Get());
+		}
 	}
 }
 
@@ -1765,7 +1793,6 @@ const bool URemoteControlPreset::HasEntities() const
 FName URemoteControlPreset::RenameExposedEntity(const FGuid& ExposedEntityId, FName NewLabel)
 {
 	FName AssignedLabel = Registry->RenameExposedEntity(ExposedEntityId, NewLabel);
-	PerFrameUpdatedEntities.Add(ExposedEntityId);
 	return AssignedLabel;
 }
 
diff --git a/Engine/Plugins/VirtualProduction/RemoteControl/Source/RemoteControl/Private/RemoteControlPropertyIdRegistry.cpp b/Engine/Plugins/VirtualProduction/RemoteControl/Source/RemoteControl/Private/RemoteControlPropertyIdRegistry.cpp
--- a/Engine/Plugins/VirtualProduction/RemoteControl/Source/RemoteControl/Private/RemoteControlPropertyIdRegistry.cpp
+++ b/Engine/Plugins/VirtualProduction/RemoteControl/Source/RemoteControl/Private/RemoteControlPropertyIdRegistry.cpp
@@ -529,11 +529,7 @@ void URemoteControlPropertyIdRegistry::UpdateIdentifiedField(const TSharedRef<FR
 		if (FRCPropertyIdWrapper* Wrapper = IdentifiedFields.FindByHash(Hash, InFieldToIdentify->GetId()))
 		{
 			Wrapper->SetPropertyId(InFieldToIdentify->PropertyId);
-
-			if (URemoteControlPreset* SourcePreset = GetSourcePreset())
-			{
-				SourcePreset->GetPropertyIdRegistry()->OnPropertyIdUpdated().Broadcast();
-			}
+			OnPropertyIdUpdated().Broadcast();
 		}
 	}
 }
diff --git a/Engine/Plugins/VirtualProduction/RemoteControl/Source/RemoteControl/Public/RemoteControlPreset.h b/Engine/Plugins/VirtualProduction/RemoteControl/Source/RemoteControl/Public/RemoteControlPreset.h
--- a/Engine/Plugins/VirtualProduction/RemoteControl/Source/RemoteControl/Public/RemoteControlPreset.h
+++ b/Engine/Plugins/VirtualProduction/RemoteControl/Source/RemoteControl/Public/RemoteControlPreset.h
@@ -198,7 +198,7 @@ struct REMOTECONTROL_API FRemoteControlPresetLayout
 	{
 		FGuid OriginGroupId;
 		FGuid TargetGroupId;
-		FGuid DraggedFieldId;
+		TArray<FGuid> DraggedFieldsIds;
 		FGuid TargetFieldId;
 	};
 
@@ -247,19 +247,11 @@ struct REMOTECONTROL_API FRemoteControlPresetLayout
 	 */
 	FRemoteControlPresetGroup* FindGroupFromField(FGuid FieldId);
 
-	/**
-	 * Move field to a group.
-	 * @param FieldId the field to move.
-	 * @param TargetGroupId the group to move the field in.
-	 * @return whether the operation was successful.
-	 */
-	bool MoveField(FGuid FieldId, FGuid TargetGroupId);
-
 	/** Swap two groups. */
 	void SwapGroups(FGuid OriginGroupId, FGuid TargetGroupId);
 
 	/** Swap fields across groups or in the same one. */
-	void SwapFields(const FFieldSwapArgs& FieldSwapArgs);
+	void SwapFields(const FFieldSwapArgs& InFieldSwapArgs);
 
 	/** Swap fields across groups or in the same one for the default group. */
 	void SwapFieldsDefaultGroup(const FFieldSwapArgs& FieldSwapArgs, const FGuid InFieldRealGroup, TArray<FGuid> InEntities);
diff --git a/Engine/Plugins/VirtualProduction/RemoteControl/Source/RemoteControl/Public/RemoteControlPropertyIdRegistry.h b/Engine/Plugins/VirtualProduction/RemoteControl/Source/RemoteControl/Public/RemoteControlPropertyIdRegistry.h
--- a/Engine/Plugins/VirtualProduction/RemoteControl/Source/RemoteControl/Public/RemoteControlPropertyIdRegistry.h
+++ b/Engine/Plugins/VirtualProduction/RemoteControl/Source/RemoteControl/Public/RemoteControlPropertyIdRegistry.h
@@ -83,18 +83,23 @@ private:
 
 private:
 	/** Entity identifier of the remote control property. */
+	UPROPERTY()
 	FGuid EntityId;
 
 	/** Holds the field identifier of the remote control property. */
+	UPROPERTY()
 	FName PropertyId;
 
 	/** Holds the type of the wrapped property. */
+	UPROPERTY()
 	FName SuperType;
 
 	/** Holds the type of the wrapped object. */
+	UPROPERTY()
 	FName SubType;
 
 	/** Holds the type to be created. */
+	UPROPERTY()
 	UClass* ClassToCreate;
 };
 
@@ -203,7 +208,7 @@ private:
 
 private:
 	/** Holds the identified fields. */
-	UPROPERTY()
+	UPROPERTY(Transient)
 	TSet<FRCPropertyIdWrapper> IdentifiedFields;
 
 	/** Delegate triggered when a property has its field Id changed. */
diff --git a/Engine/Plugins/VirtualProduction/RemoteControl/Source/RemoteControlLogic/Private/Action/RCPropertyIdAction.cpp b/Engine/Plugins/VirtualProduction/RemoteControl/Source/RemoteControlLogic/Private/Action/RCPropertyIdAction.cpp
--- a/Engine/Plugins/VirtualProduction/RemoteControl/Source/RemoteControlLogic/Private/Action/RCPropertyIdAction.cpp
+++ b/Engine/Plugins/VirtualProduction/RemoteControl/Source/RemoteControlLogic/Private/Action/RCPropertyIdAction.cpp
@@ -207,7 +207,8 @@ void URCPropertyIdAction::OnEntityUnexposed(URemoteControlPreset* InPreset, cons
 			//If the propertyId of the property unexposed is the same of this PropertyIdAction refresh, otherwise don't update it
 			if (UnexposedEntity.Pin()->PropertyId == PropertyId)
 			{
-				UpdatePropertyId();
+				InPreset->GetPropertyIdRegistry()->RemoveIdentifiedField(UnexposedEntity.Pin()->GetId());
+				InPreset->GetPropertyIdRegistry()->OnPropertyIdUpdated().Broadcast();
 			}
 			return;
 		}
diff --git a/Engine/Plugins/VirtualProduction/RemoteControl/Source/RemoteControlUI/Private/RemoteControlUIModule.cpp b/Engine/Plugins/VirtualProduction/RemoteControl/Source/RemoteControlUI/Private/RemoteControlUIModule.cpp
--- a/Engine/Plugins/VirtualProduction/RemoteControl/Source/RemoteControlUI/Private/RemoteControlUIModule.cpp
+++ b/Engine/Plugins/VirtualProduction/RemoteControl/Source/RemoteControlUI/Private/RemoteControlUIModule.cpp
@@ -1220,6 +1220,11 @@ void FRemoteControlUIModule::RefreshPanels()
 	if (TSharedPtr<SRemoteControlPanel> Panel = GetPanelForObject(nullptr))
 	{
 		Panel->Refresh();
+		// Propagate that the PropertyId updated during the Undo/Redo to update the action as well.
+		if (Panel->GetPreset() && Panel->GetPreset()->GetPropertyIdRegistry())
+		{
+			Panel->GetPreset()->GetPropertyIdRegistry()->OnPropertyIdUpdated().Broadcast();
+		}
 	}
 }
 
diff --git a/Engine/Plugins/VirtualProduction/RemoteControl/Source/RemoteControlUI/Private/UI/Action/SRCActionPanelList.h b/Engine/Plugins/VirtualProduction/RemoteControl/Source/RemoteControlUI/Private/UI/Action/SRCActionPanelList.h
--- a/Engine/Plugins/VirtualProduction/RemoteControl/Source/RemoteControlUI/Private/UI/Action/SRCActionPanelList.h
+++ b/Engine/Plugins/VirtualProduction/RemoteControl/Source/RemoteControlUI/Private/UI/Action/SRCActionPanelList.h
@@ -336,12 +336,17 @@ private:
 				if (TSharedPtr<FExposedEntityDragDrop> DragDropOp = StaticCastSharedPtr<FExposedEntityDragDrop>(DragDropOperation))
 				{
 					// Fetch the Exposed Entity
-					const FGuid ExposedEntityId = DragDropOp->GetId();
-
-					if (TSharedPtr<SRCActionPanel> ActionPanel = GetActionPanel())
-
-					// Add Action
-					AddAction(ExposedEntityId);
+					for (const FGuid& ExposedEntityId : DragDropOp->GetSelectedIds())
+					{
+						if (TSharedPtr<SRCActionPanel> ActionPanel = GetActionPanel())
+						{
+							// Add Action
+							if (ActionPanel->CanHaveActionForField(ExposedEntityId))
+							{
+								AddAction(ExposedEntityId);
+							}
+						}
+					}
 				}
 			}
 			else if (DragDropOperation->IsOfType<FFieldGroupDragDropOp>())
@@ -382,11 +387,16 @@ private:
 				if (TSharedPtr<FExposedEntityDragDrop> DragDropOp = StaticCastSharedPtr<FExposedEntityDragDrop>(DragDropOperation))
 				{
 					// Fetch the Exposed Entity
-					const FGuid ExposedEntityId = DragDropOp->GetId();
-
-					if (TSharedPtr<SRCActionPanel> ActionPanel = GetActionPanel())
+					for (const FGuid& ExposedEntityId : DragDropOp->GetSelectedIds())
 					{
-						return ActionPanel->CanHaveActionForField(ExposedEntityId);
+						if (TSharedPtr<SRCActionPanel> ActionPanel = GetActionPanel())
+						{
+							// Add Action
+							if (ActionPanel->CanHaveActionForField(ExposedEntityId))
+							{
+								return true;
+							}
+						}
 					}
 				}
 			}
diff --git a/Engine/Plugins/VirtualProduction/RemoteControl/Source/RemoteControlUI/Private/UI/Controller/SRCControllerPanelList.cpp b/Engine/Plugins/VirtualProduction/RemoteControl/Source/RemoteControlUI/Private/UI/Controller/SRCControllerPanelList.cpp
--- a/Engine/Plugins/VirtualProduction/RemoteControl/Source/RemoteControlUI/Private/UI/Controller/SRCControllerPanelList.cpp
+++ b/Engine/Plugins/VirtualProduction/RemoteControl/Source/RemoteControlUI/Private/UI/Controller/SRCControllerPanelList.cpp
@@ -235,15 +235,18 @@ namespace UE::RCControllerPanelList
 				if (TSharedPtr<FExposedEntityDragDrop> DragDropOp = StaticCastSharedPtr<FExposedEntityDragDrop>(DragDropOperation))
 				{
 					// Fetch the Exposed Entity
-					const FGuid ExposedEntityId = DragDropOp->GetId();
+					const TArray<FGuid>& ExposedEntitiesIds = DragDropOp->GetSelectedIds();
 
-					if (URemoteControlPreset* Preset = ControllerPanelList->GetPreset())
+					if (ExposedEntitiesIds.Num() == 1)
 					{
-						if (TSharedPtr<const FRemoteControlProperty> RemoteControlProperty = Preset->GetExposedEntity<FRemoteControlProperty>(ExposedEntityId).Pin())
+						if (URemoteControlPreset* Preset = ControllerPanelList->GetPreset())
 						{
-							if (URCController* Controller = Cast<URCController>(ControllerItem->GetVirtualProperty()))
+							if (TSharedPtr<const FRemoteControlProperty> RemoteControlProperty = Preset->GetExposedEntity<FRemoteControlProperty>(ExposedEntitiesIds[0]).Pin())
 							{
-								ControllerPanelList->CreateBindBehaviourAndAssignTo(Controller, RemoteControlProperty.ToSharedRef(), true);
+								if (URCController* Controller = Cast<URCController>(ControllerItem->GetVirtualProperty()))
+								{
+									ControllerPanelList->CreateBindBehaviourAndAssignTo(Controller, RemoteControlProperty.ToSharedRef(), true);
+								}
 							}
 						}
 					}
@@ -284,21 +287,24 @@ namespace UE::RCControllerPanelList
 						if (TSharedPtr<FExposedEntityDragDrop> DragDropOp = StaticCastSharedPtr<FExposedEntityDragDrop>(DragDropOperation))
 						{
 							// Fetch the Exposed Entity
-							const FGuid ExposedEntityId = DragDropOp->GetId();
+							const TArray<FGuid>& ExposedEntitiesIds = DragDropOp->GetSelectedIds();
 
-							if (URemoteControlPreset* Preset = ControllerPanelList->GetPreset())
+							if (ExposedEntitiesIds.Num() == 1)
 							{
-								if (TSharedPtr<const FRemoteControlField> RemoteControlField = Preset->GetExposedEntity<FRemoteControlField>(ExposedEntityId).Pin())
+								if (URemoteControlPreset* Preset = ControllerPanelList->GetPreset())
 								{
-									if (URCController* Controller = Cast<URCController>(ControllerItem->GetVirtualProperty()))
+									if (TSharedPtr<const FRemoteControlField> RemoteControlField = Preset->GetExposedEntity<FRemoteControlField>(ExposedEntitiesIds[0]).Pin())
 									{
-										const bool bAllowNumericInputAsStrings = true;
+										if (URCController* Controller = Cast<URCController>(ControllerItem->GetVirtualProperty()))
+										{
+											const bool bAllowNumericInputAsStrings = true;
 
-										const bool bAllowDrop = URCBehaviourBind::CanHaveActionForField(Controller, RemoteControlField.ToSharedRef(), bAllowNumericInputAsStrings);
+											const bool bAllowDrop = URCBehaviourBind::CanHaveActionForField(Controller, RemoteControlField.ToSharedRef(), bAllowNumericInputAsStrings);
 
-										ControllerPanelList->bIsAnyControllerItemEligibleForDragDrop |= (bAllowDrop && bIsDragActive);
+											ControllerPanelList->bIsAnyControllerItemEligibleForDragDrop |= (bAllowDrop && bIsDragActive);
 
-										return bAllowDrop;
+											return bAllowDrop;
+										}
 									}
 								}
 							}
@@ -854,10 +860,10 @@ bool SRCControllerPanelList::OnAllowDrop(TSharedPtr<FDragDropOperation> DragDrop
 	if (TSharedPtr<FExposedEntityDragDrop> DragDropOp = GetExposedEntityDragDrop(DragDropOperation))
 	{
 		// Fetch the Exposed Entity
-		const FGuid ExposedEntityId = DragDropOp->GetId();
+		const TArray<FGuid>& ExposedEntitiesIds = DragDropOp->GetSelectedIds();
 
-		// Check if Entity is supported by controllers
-		return IsEntitySupported(ExposedEntityId);
+		// Check if Entity is supported by controllers and currently only 1 dragged entity dragged is supported
+		return ExposedEntitiesIds.Num() == 1 && IsEntitySupported(ExposedEntitiesIds[0]);
 	}
 
 	return false;
@@ -868,7 +874,7 @@ FReply SRCControllerPanelList::OnControllerListViewDragDrop(TSharedPtr<FDragDrop
 	if (TSharedPtr<FExposedEntityDragDrop> DragDropOp = GetExposedEntityDragDrop(DragDropOperation))
 	{
 		// Fetch the Exposed Entity
-		const FGuid ExposedEntityId = DragDropOp->GetId();
+		const FGuid ExposedEntityId = DragDropOp->GetNodeId();
 
 		if (URemoteControlPreset* Preset = GetPreset())
 		{
diff --git a/Engine/Plugins/VirtualProduction/RemoteControl/Source/RemoteControlUI/Private/UI/SRCPanelExposedEntitiesList.cpp b/Engine/Plugins/VirtualProduction/RemoteControl/Source/RemoteControlUI/Private/UI/SRCPanelExposedEntitiesList.cpp
--- a/Engine/Plugins/VirtualProduction/RemoteControl/Source/RemoteControlUI/Private/UI/SRCPanelExposedEntitiesList.cpp
+++ b/Engine/Plugins/VirtualProduction/RemoteControl/Source/RemoteControlUI/Private/UI/SRCPanelExposedEntitiesList.cpp
@@ -164,7 +164,86 @@ private:
 	{
 		if (Entity.IsValid())
 		{
-			Entity->EnterRenameMode();
+			if (const TSharedPtr<SRCPanelExposedEntity> ExposedEntityWidget = StaticCastSharedPtr<SRCPanelExposedEntity>(Entity))
+			{
+				if (const TSharedPtr<FRemoteControlEntity> ExposedEntity = ExposedEntityWidget->GetEntity())
+				{
+					if (const TSharedPtr<FRemoteControlField> RCField = StaticCastSharedPtr<FRemoteControlField>(ExposedEntity))
+					{
+						TArray<UObject*> BoundObjects = RCField->GetBoundObjects();
+						TSet<UObject*> Objects = TSet<UObject*>{ BoundObjects };
+
+						TArray<UObject*> OwnerActors;
+						for (UObject* Object : Objects)
+						{
+							if (Object->IsA<AActor>())
+							{
+								OwnerActors.Add(Object);
+							}
+							else
+							{
+								OwnerActors.Add(Object->GetTypedOuter<AActor>());
+							}
+						}
+
+						if (RCField->GetStruct() == FRemoteControlProperty::StaticStruct())
+						{
+							const TSharedPtr<FRemoteControlProperty> RCProp = StaticCastSharedPtr<FRemoteControlProperty>(RCField);
+
+							// Resolve it to get the property path.
+							RCProp->FieldPathInfo.Resolve(RCProp->GetBoundObject());
+							if (RCProp->FieldPathInfo.IsResolved())
+							{
+								TSharedRef<FPropertyPath> PropertyPath = RCProp->FieldPathInfo.ToPropertyPath();
+
+								for (auto It = BoundObjects.CreateIterator(); It; ++It)
+								{
+									if (AActor* OwnerActor = (*It)->GetTypedOuter<AActor>())
+									{
+										const UObject* BindingObject = *It;
+
+										// When we encounter a non-component object, 
+										if (!BindingObject->IsA<UActorComponent>())
+										{
+											BoundObjects.Add(OwnerActor);
+											It.RemoveCurrent();
+										}
+
+										// --- Special NDisplay handling because of their customization ---
+										// Since display cluster config data is not created as a default subobject, therefore we have no way of retrieving the CurrentConfigData property from the config object itself.
+										static FName DisplayClusterConfigDataClassName = "DisplayClusterConfigurationData";
+										if (BindingObject->GetClass()->GetFName() == DisplayClusterConfigDataClassName)
+										{
+											if (FProperty* Property = OwnerActor->GetClass()->FindPropertyByName("CurrentConfigData"))
+											{
+												// Append "CurrentConfigData" to the beginning of the path.
+												TSharedRef<FPropertyPath> NewPropertyPath = FPropertyPath::Create(Property);
+												for (int32 Index = 0; Index < PropertyPath->GetNumProperties(); Index++)
+												{
+													NewPropertyPath->AddProperty(PropertyPath->GetPropertyInfo(Index));
+												}
+
+												PropertyPath = NewPropertyPath;
+											}
+										}
+									}
+								}
+
+								FRemoteControlUIModule::Get().SelectObjects(BoundObjects);
+								
+								FTimerHandle Handle;
+								const FTimerDelegate Delegate = FTimerDelegate::CreateLambda(([PropertyPath]()
+									{
+										FRemoteControlUIModule::Get().HighlightPropertyInDetailsPanel(*PropertyPath);
+									}));
+								
+								// Needed because modifying the selection set is asynchronous.
+								GEditor->GetTimerManager()->SetTimer(Handle, Delegate, 0.1, false, -1);
+							}
+						}
+					}
+				}
+			}
 		}
 
 		return FSuperRowType::OnMouseButtonDoubleClick(InMyGeometry, InMouseEvent);
@@ -268,7 +347,7 @@ void SRCPanelExposedEntitiesList::Construct(const FArguments& InArgs, URemoteCon
 		.ItemHeight(24.f)
 		.OnGenerateRow(this, &SRCPanelExposedEntitiesList::OnGenerateRow)
 		.OnSelectionChanged(this, &SRCPanelExposedEntitiesList::OnSelectionChanged)
-		.SelectionMode(ESelectionMode::Single)
+		.SelectionMode(ESelectionMode::Multi)
 		.TreeItemsSource(&FieldEntities)
 		.OnContextMenuOpening(this, &SRCPanelExposedEntitiesList::OnContextMenuOpening, SRCPanelTreeNode::Field)
 		.ClearSelectionOnClick(true)
@@ -499,6 +578,34 @@ TSharedPtr<SRCPanelTreeNode> SRCPanelExposedEntitiesList::GetSelectedEntity() co
 	return nullptr;
 }
 
+TArray<TSharedPtr<SRCPanelTreeNode>> SRCPanelExposedEntitiesList::GetSelectedEntities() const
+{
+	TArray<TSharedPtr<SRCPanelTreeNode>> SelectedNodes;
+	if (FieldsListView.IsValid())
+	{
+		FieldsListView->GetSelectedItems(SelectedNodes);
+	}
+	return SelectedNodes;
+}
+
+int32 SRCPanelExposedEntitiesList::GetSelectedEntitiesNum() const
+{
+	if (FieldsListView.IsValid())
+	{
+		return FieldsListView->GetNumItemsSelected();
+	}
+	return -1;
+}
+
+bool SRCPanelExposedEntitiesList::IsEntitySelected(const TSharedPtr<SRCPanelTreeNode>& InNode) const
+{
+	if (FieldsListView.IsValid())
+	{
+		return FieldsListView->IsItemSelected(InNode);
+	}
+	return false;
+}
+
 void SRCPanelExposedEntitiesList::SetSelection(const TSharedPtr<SRCPanelTreeNode>& Node, const bool bForceMouseClick)
 {
 	if (Node)
@@ -507,7 +614,7 @@ void SRCPanelExposedEntitiesList::SetSelection(const TSharedPtr<SRCPanelTreeNode
 
 		if (TSharedPtr<SRCPanelTreeNode>* FoundTreeNode = FieldWidgetMap.Find(Node->GetRCId()))
 		{
-			FieldsListView->SetSelection(*FoundTreeNode, SelectInfo);
+			FieldsListView->SetItemSelection(*FoundTreeNode, true, SelectInfo);
 			return;
 		}
 
@@ -632,6 +739,54 @@ void SRCPanelExposedEntitiesList::ExposedEntitiesNodesRefresh()
 	}
 }
 
+void SRCPanelExposedEntitiesList::OnPropertyIdRenamed(const FName InNewId, TSharedPtr<SRCPanelTreeNode> InNode)
+{
+	if (IsEntitySelected(InNode))
+	{
+		TArray<TSharedPtr<SRCPanelTreeNode>> SelectedEntities = GetSelectedEntities();
+		for (const TSharedPtr<SRCPanelTreeNode>& Entity : SelectedEntities)
+		{
+			TWeakPtr<FRemoteControlField> ExposedField = Preset->GetExposedEntity<FRemoteControlField>(Entity->GetRCId());
+			if (ExposedField.IsValid())
+			{
+				ExposedField.Pin()->PropertyId = InNewId;
+				Entity->SetPropertyId(InNewId);
+				Preset->UpdateIdentifiedField(ExposedField.Pin().ToSharedRef());
+			}
+		}
+	}
+}
+
+void SRCPanelExposedEntitiesList::OnNameRenamed(const FName InNewName)
+{
+	TArray<TSharedPtr<SRCPanelTreeNode>> SelectedEntities = GetSelectedEntities();
+	for (const TSharedPtr<SRCPanelTreeNode>& Entity : SelectedEntities)
+	{
+		TWeakPtr<FRemoteControlField> ExposedField = Preset->GetExposedEntity<FRemoteControlField>(Entity->GetRCId());
+		if (ExposedField.IsValid())
+		{
+			ExposedField.Pin()->Rename(InNewName);
+			Entity->SetName(ExposedField.Pin()->GetLabel());
+		}
+	}
+}
+
+FReply SRCPanelExposedEntitiesList::OnNodeDragDetected(const FGeometry& InGeometry, const FPointerEvent& InPointerEvent, TSharedPtr<SRCPanelTreeNode> InNode)
+{
+	if (InNode && InNode->GetRCType() == SRCPanelTreeNode::Field)
+	{
+		TArray<FGuid> SelectedIds;
+		Algo::TransformIf(GetSelectedEntities(), SelectedIds
+			, [] (const TSharedPtr<SRCPanelTreeNode>& TreeNode) { return TreeNode->GetRCType() == SRCPanelTreeNode::Field; }
+			, [] (const TSharedPtr<SRCPanelTreeNode>& TreeNode){ return TreeNode->GetRCId(); });
+
+		const TSharedRef<FExposedEntityDragDrop> DragDropOp = MakeShared<FExposedEntityDragDrop>(InNode->GetDragAndDropWidget(SelectedIds.Num()), InNode->GetRCId(), SelectedIds);
+		DragDropOp->Construct();
+		return FReply::Handled().BeginDragDrop(DragDropOp);
+	}
+	return FReply::Unhandled();
+}
+
 void SRCPanelExposedEntitiesList::OnObjectPropertyChange(UObject* InObject, FPropertyChangedEvent& InChangeEvent)
 {
 	static constexpr EPropertyChangeType::Type TypesNeedingRefresh = EPropertyChangeType::ArrayAdd | EPropertyChangeType::ArrayClear | EPropertyChangeType::ArrayRemove | EPropertyChangeType::ValueSet | EPropertyChangeType::Duplicate;
@@ -775,7 +930,7 @@ void SRCPanelExposedEntitiesList::GenerateListWidgets()
 	}
 
 	// We order the ALL group here otherwise the order would be the same order of the ExposedEntities
-	if (OrderMap.Num() && FieldEntities.Num())
+	if (OrderMap.Num() && FieldWidgetMap.Num())
 	{
 		FieldWidgetMap.KeySort(
 			[&OrderMap]
@@ -886,17 +1041,11 @@ TSharedRef<ITableRow> SRCPanelExposedEntitiesList::OnGenerateRow(TSharedPtr<SRCP
 	{
 		constexpr float LeftPadding = 3.f;
 		const FMargin Margin = Node->GetRCType() == SRCPanelTreeNode::FieldChild ? FMargin(LeftPadding + 10.f, 1.f, 1.f, 1.f) : FMargin(LeftPadding, 1.f, 1.f, 1.f);
+		Node->OnPropertyIdRenamed().BindSP(this, &SRCPanelExposedEntitiesList::OnPropertyIdRenamed, Node);
+		Node->OnNameRenamed().BindSP(this, &SRCPanelExposedEntitiesList::OnNameRenamed);
+
 		return SNew(SEntityRow, OwnerTable)
-			.OnDragDetected_Lambda([this, Node] (const FGeometry&, const FPointerEvent&)
-			{
-				if (Node && Node->GetRCType() == SRCPanelTreeNode::Field)
-				{
-					const TSharedRef<FExposedEntityDragDrop> DragDropOp = MakeShared<FExposedEntityDragDrop>(Node->GetDragAndDropWidget(), Node->GetRCId());
-					DragDropOp->Construct();
-					return FReply::Handled().BeginDragDrop(DragDropOp);
-				}
-				return FReply::Unhandled();
-			})
+			.OnDragDetected(FOnDragDetected::CreateSP(this, &SRCPanelExposedEntitiesList::OnNodeDragDetected, Node))
 			.OnDragEnter_Lambda([Node](const FDragDropEvent& Event) { if (Node && Node->GetRCType() == SRCPanelTreeNode::Field) StaticCastSharedPtr<SRCPanelExposedField>(Node)->SetIsHovered(true); })
 			.OnDragLeave_Lambda([Node](const FDragDropEvent& Event) { if (Node && Node->GetRCType() == SRCPanelTreeNode::Field) StaticCastSharedPtr<SRCPanelExposedField>(Node)->SetIsHovered(false); })
 			.OnAcceptDrop_Lambda(OnAcceptDropLambda)
@@ -942,84 +1091,6 @@ void SRCPanelExposedEntitiesList::OnSelectionChanged(TSharedPtr<SRCPanelTreeNode
 
 		FieldsListView->ClearSelection();
 	}
-	else
-	{
-		// todo call handler on node itself
-		if (TSharedPtr<FRemoteControlField> RCField = Preset->GetExposedEntity<FRemoteControlField>(Node->GetRCId()).Pin())
-		{
-			TArray<UObject*> BoundObjects = RCField->GetBoundObjects();
-			TSet<UObject*> Objects = TSet<UObject*>{ BoundObjects };
-
-			TArray<UObject*> OwnerActors;
-			for (UObject* Object : Objects)
-			{
-				if (Object->IsA<AActor>())
-				{
-					OwnerActors.Add(Object);
-				}
-				else
-				{
-					OwnerActors.Add(Object->GetTypedOuter<AActor>());
-				}
-			}
-
-			if (RCField->GetStruct() == FRemoteControlProperty::StaticStruct())
-			{
-				TSharedPtr<FRemoteControlProperty> RCProp = StaticCastSharedPtr<FRemoteControlProperty>(RCField);
-
-				// Resolve it to get the property path.
-				RCProp->FieldPathInfo.Resolve(RCProp->GetBoundObject());
-				if (RCProp->FieldPathInfo.IsResolved())
-				{
-					TSharedRef<FPropertyPath> PropertyPath = RCProp->FieldPathInfo.ToPropertyPath();
-
-					for (auto It = BoundObjects.CreateIterator(); It; ++It)
-					{
-						if (AActor* OwnerActor = (*It)->GetTypedOuter<AActor>())
-						{
-							UObject* BindingObject = *It;
-
-							// When we encounter a non-component object, 
-							if (!BindingObject->IsA<UActorComponent>())
-							{
-								BoundObjects.Add(OwnerActor);
-								It.RemoveCurrent();
-							}
-
-							// --- Special NDisplay handling because of their customization ---
-							// Since display cluster config data is not created as a defaeult subobject, therefore we have no way of retrieving the CurrentConfigData property from the config object itself.
-							static FName DisplayClusterConfigDataClassName = "DisplayClusterConfigurationData";
-							if (BindingObject->GetClass()->GetFName() == DisplayClusterConfigDataClassName)
-							{
-								if (FProperty* Property = OwnerActor->GetClass()->FindPropertyByName("CurrentConfigData"))
-								{
-									// Append "CurrentConfigData" to the beginning of the path.
-									TSharedRef<FPropertyPath> NewPropertyPath = FPropertyPath::Create(Property);
-									for (int32 Index = 0; Index < PropertyPath->GetNumProperties(); Index++)
-									{
-										NewPropertyPath->AddProperty(PropertyPath->GetPropertyInfo(Index));
-									}
-
-									PropertyPath = NewPropertyPath;
-								}
-							}
-						}
-					}
-
-					FRemoteControlUIModule::Get().SelectObjects(BoundObjects);
-				
-					FTimerHandle Handle;
-					FTimerDelegate Delegate = FTimerDelegate::CreateLambda(([PropertyPath]()
-						{
-							FRemoteControlUIModule::Get().HighlightPropertyInDetailsPanel(*PropertyPath);
-						}));
-
-					// Needed because modifying the selection set is asynchronous.
-					GEditor->GetTimerManager()->SetTimer(Handle, Delegate, 0.1, 0, -1);
-				}
-			}
-		}
-	}
 
 	OnSelectionChangeDelegate.Broadcast(Node);
 }
@@ -1037,7 +1108,7 @@ FReply SRCPanelExposedEntitiesList::OnDropOnGroup(const TSharedPtr<FDragDropOper
 				FRemoteControlPresetLayout::FFieldSwapArgs Args;
 				Args.OriginGroupId = GetSelectedGroup()->GetRCId();
 				Args.TargetGroupId = DragTargetGroup->GetRCId();
-				Args.DraggedFieldId = DragDropOp->GetId();
+				Args.DraggedFieldsIds = DragDropOp->GetSelectedIds();
 
 				if (TargetEntity)
 				{
@@ -1053,13 +1124,13 @@ FReply SRCPanelExposedEntitiesList::OnDropOnGroup(const TSharedPtr<FDragDropOper
 					OrderEntities.Add(Entity->GetRCId());
 				}
 
-				Preset->Layout.SwapFieldsDefaultGroup(Args, GetGroupId(DragDropOp->GetId()), OrderEntities);
+				Preset->Layout.SwapFieldsDefaultGroup(Args, GetGroupId(DragDropOp->GetNodeId()), OrderEntities);
 				constexpr bool bForceMouseClick = true;
 				SetSelection(FindGroupById(Args.OriginGroupId), bForceMouseClick);
 				return FReply::Handled();
 			}
 
-			FGuid DragOriginGroupId = GetGroupId(DragDropOp->GetId());
+			FGuid DragOriginGroupId = GetGroupId(DragDropOp->GetNodeId());
 			if (!DragOriginGroupId.IsValid())
 			{
 				return FReply::Unhandled();
@@ -1068,7 +1139,7 @@ FReply SRCPanelExposedEntitiesList::OnDropOnGroup(const TSharedPtr<FDragDropOper
 			FRemoteControlPresetLayout::FFieldSwapArgs Args;
 			Args.OriginGroupId = DragOriginGroupId;
 			Args.TargetGroupId = DragTargetGroup->GetRCId();
-			Args.DraggedFieldId = DragDropOp->GetId();
+			Args.DraggedFieldsIds = DragDropOp->GetSelectedIds();
 
 			if (TargetEntity)
 			{
@@ -1472,15 +1543,6 @@ void SRCPanelExposedEntitiesList::OnFieldOrderChanged(const FGuid& GroupId, cons
 
 void SRCPanelExposedEntitiesList::OnEntitiesUpdated(URemoteControlPreset*, const TSet<FGuid>& UpdatedEntities)
 {
-	for (const FGuid& EntityId : UpdatedEntities)
-	{
-		TSharedPtr<SRCPanelTreeNode>* Node = FieldWidgetMap.Find(EntityId);
-		if (Node && *Node)
-		{
-			(*Node)->Refresh();
-		}
-	}
-	
 	GEditor->GetTimerManager()->SetTimerForNextTick(FTimerDelegate::CreateLambda([WeakListPtr = TWeakPtr<SRCPanelExposedEntitiesList>(StaticCastSharedRef<SRCPanelExposedEntitiesList>(AsShared()))]()
 	{
 		if (TSharedPtr<SRCPanelExposedEntitiesList> ListPtr = WeakListPtr.Pin())
@@ -1808,16 +1870,14 @@ void SRCPanelExposedEntitiesList::ProcessRefresh()
 {
 	TRACE_CPUPROFILER_EVENT_SCOPE(SRCPanelExposedEntitiesList::Refresh);
 
+	TArray<TSharedPtr<SRCPanelTreeNode>> SelectedItems;
+	FieldsListView->GetSelectedItems(SelectedItems);
 	GenerateListWidgets();
 
 	RefreshGroups();
 
 	if (Preset.IsValid())
 	{
-		//Refresh during PostUndo/Redo to keep the action updated.
-		Preset->GetPropertyIdRegistry()->Initialize();
-		Preset->GetPropertyIdRegistry()->OnPropertyIdUpdated().Broadcast();
-
 		constexpr bool bForceMouseClick = true;
 
 		if (const FRemoteControlPresetGroup* SelectedGroup = Preset->Layout.GetGroup(CurrentlySelectedGroup))
@@ -1831,6 +1891,11 @@ void SRCPanelExposedEntitiesList::ProcessRefresh()
 			GenerateListWidgets(DefaultGroup);
 			SetSelection(FindGroupById(DefaultGroup.Id), bForceMouseClick);
 		}
+
+		for (TSharedPtr<SRCPanelTreeNode> SelectedItem : SelectedItems)
+		{
+			SetSelection(SelectedItem, bForceMouseClick);
+		}
 	}
 }
 
diff --git a/Engine/Plugins/VirtualProduction/RemoteControl/Source/RemoteControlUI/Private/UI/SRCPanelExposedEntitiesList.h b/Engine/Plugins/VirtualProduction/RemoteControl/Source/RemoteControlUI/Private/UI/SRCPanelExposedEntitiesList.h
--- a/Engine/Plugins/VirtualProduction/RemoteControl/Source/RemoteControlUI/Private/UI/SRCPanelExposedEntitiesList.h
+++ b/Engine/Plugins/VirtualProduction/RemoteControl/Source/RemoteControlUI/Private/UI/SRCPanelExposedEntitiesList.h
@@ -83,7 +83,16 @@ public:
 	
 	/** Get the currently selected exposed entity. */
 	TSharedPtr<SRCPanelTreeNode> GetSelectedEntity() const;
-	
+
+	/** Get the currently selected exposed entities. */
+	TArray<TSharedPtr<SRCPanelTreeNode>> GetSelectedEntities() const;
+
+	/** Get the currently selected exposed entities. */
+	int32 GetSelectedEntitiesNum() const;
+
+	/** Return true if the given node is currently selected. */
+	bool IsEntitySelected(const TSharedPtr<SRCPanelTreeNode>& InNode) const;
+
 	/** Set the currently selected group or exposed entity. */
 	void SetSelection(const TSharedPtr<SRCPanelTreeNode>& Node, const bool bForceMouseClick = false);
 
@@ -207,7 +216,16 @@ private:
 	 * If necessary, will perform a refresh of the exposed entities nodes widgets.
 	 */
 	void ExposedEntitiesNodesRefresh();
-	
+
+	/** Executed when a property Id is changed, will set all selected node(s) property id to the new one */
+	void OnPropertyIdRenamed(const FName InNewId, TSharedPtr<SRCPanelTreeNode> InNode);
+
+	/** Executed when a property Name is changed, will set all selected node(s) property Name to the new one */
+	void OnNameRenamed(const FName InNewName);
+
+	/** Executed when a drag is detected, will create the Drag and Drop widget of the node(s) */
+	FReply OnNodeDragDetected(const FGeometry& InGeometry, const FPointerEvent& InPointerEvent, TSharedPtr<SRCPanelTreeNode> InNode);
+
 private:
 	/** Holds the Groups list view. */
 	TSharedPtr<SListView<TSharedPtr<SRCPanelTreeNode>>> GroupsListView;
diff --git a/Engine/Plugins/VirtualProduction/RemoteControl/Source/RemoteControlUI/Private/UI/SRCPanelExposedEntity.cpp b/Engine/Plugins/VirtualProduction/RemoteControl/Source/RemoteControlUI/Private/UI/SRCPanelExposedEntity.cpp
--- a/Engine/Plugins/VirtualProduction/RemoteControl/Source/RemoteControlUI/Private/UI/SRCPanelExposedEntity.cpp
+++ b/Engine/Plugins/VirtualProduction/RemoteControl/Source/RemoteControlUI/Private/UI/SRCPanelExposedEntity.cpp
@@ -142,6 +142,11 @@ void SRCPanelExposedEntity::EnterRenameMode()
 	}
 }
 
+FName SRCPanelExposedEntity::GetPropertyId()
+{
+	return PropertyIdLabel;
+}
+
 void SRCPanelExposedEntity::Refresh()
 {
 	if (EntityId.IsValid() && Preset.IsValid())
@@ -168,7 +173,9 @@ void SRCPanelExposedEntity::Initialize(const FGuid& InEntityId, URemoteControlPr
 
 			if (RCEntity->GetStruct() == FRemoteControlProperty::StaticStruct())
 			{
-				CachedFieldPath = StaticCastSharedPtr<FRemoteControlProperty>(RCEntity)->FieldPathInfo.ToString();
+				const TSharedPtr<FRemoteControlProperty> RCProperty = StaticCastSharedPtr<FRemoteControlProperty>(RCEntity);
+				CachedFieldPath = RCProperty->FieldPathInfo.ToString();
+				PropertyIdLabel = RCProperty->PropertyId;
 			}
 
 			FName OwnerFName;
@@ -415,23 +422,11 @@ void SRCPanelExposedEntity::OnLabelCommitted(const FText& InLabel, ETextCommit::
 		FScopedTransaction Transaction(LOCTEXT("ModifyEntityLabel", "Modify exposed entity's label."));
 		RCPreset->Modify();
 		CachedLabel = RCPreset->RenameExposedEntity(EntityId, *InLabel.ToString());
-		NameTextBox->SetText(FText::FromName(CachedLabel));
-	}
-}
-
-FText SRCPanelExposedEntity::OnPropertyIdText() const
-{
-	if (const TSharedPtr<FRemoteControlEntity> RCEntity = GetEntity())
-	{
-		if (const TSharedPtr<FRemoteControlField> RCField = StaticCastSharedPtr<FRemoteControlField>(RCEntity))
-		{
-			return FText::FromName(RCField->PropertyId);
-		}
+		OnNameRenamed().ExecuteIfBound(CachedLabel);
 	}
-	return FText::FromName(NAME_None);
 }
 
-void SRCPanelExposedEntity::OnPropertyIdTextCommitted(const FText& InText, ETextCommit::Type InCommitInfo) const
+void SRCPanelExposedEntity::OnPropertyIdTextCommitted(const FText& InText, ETextCommit::Type InCommitInfo)
 {
 	if (URemoteControlPreset* RCPreset = Preset.Get())
 	{
@@ -440,7 +435,9 @@ void SRCPanelExposedEntity::OnPropertyIdTextCommitted(const FText& InText, EText
 			if (const TSharedPtr<FRemoteControlField> RCField = StaticCastSharedPtr<FRemoteControlField>(RCEntity))
 			{
 				RCField->PropertyId = FName(InText.ToString());
+				PropertyIdLabel = RCField->PropertyId;
 				RCPreset->UpdateIdentifiedField(RCField.ToSharedRef());
+				OnPropertyIdRenamed().ExecuteIfBound(PropertyIdLabel);
 			}
 		}
 	}
@@ -509,7 +506,7 @@ TSharedRef<SWidget> SRCPanelExposedEntity::CreateEntityWidget(TSharedPtr<SWidget
 			.SelectAllTextWhenFocused(true)
 			.RevertTextOnEscape(true)
 			.ClearKeyboardFocusOnCommit(true)
-			.Text(this, &SRCPanelExposedEntity::OnPropertyIdText)
+			.Text_Lambda([this] () { return FText::FromName(PropertyIdLabel); })
 			.OnTextCommitted(this, &SRCPanelExposedEntity::OnPropertyIdTextCommitted)
 		];
 
@@ -571,7 +568,7 @@ TSharedRef<SWidget> SRCPanelExposedEntity::CreateEntityWidget(TSharedPtr<SWidget
 		.AutoWidth()
 		[
 			SAssignNew(NameTextBox, SInlineEditableTextBlock)
-			.Text(FText::FromName(CachedLabel))
+			.Text_Lambda([this] () { return FText::FromName(CachedLabel); })
 			.ToolTipText(FText::FromString(CachedFieldPath))
 			.OnTextCommitted(this, &SRCPanelExposedEntity::OnLabelCommitted)
 			.OnVerifyTextChanged(this, &SRCPanelExposedEntity::OnVerifyItemLabelChanged)
diff --git a/Engine/Plugins/VirtualProduction/RemoteControl/Source/RemoteControlUI/Private/UI/SRCPanelExposedEntity.h b/Engine/Plugins/VirtualProduction/RemoteControl/Source/RemoteControlUI/Private/UI/SRCPanelExposedEntity.h
--- a/Engine/Plugins/VirtualProduction/RemoteControl/Source/RemoteControlUI/Private/UI/SRCPanelExposedEntity.h
+++ b/Engine/Plugins/VirtualProduction/RemoteControl/Source/RemoteControlUI/Private/UI/SRCPanelExposedEntity.h
@@ -27,6 +27,12 @@ struct SRCPanelExposedEntity : public SRCPanelTreeNode
 	virtual FGuid GetRCId() const override final { return EntityId; }
 	/** Make the group name's text box editable. */
 	virtual void EnterRenameMode() override;
+	/** Get the PropertyId of this Node. */
+	virtual FName GetPropertyId() override;
+	/** Set the PropertyId of this Node. */
+	virtual void SetPropertyId(FName InNewPropertyId) override { PropertyIdLabel = InNewPropertyId; };
+	/** Set the Name of this Node. */
+	virtual void SetName(FName InNewName) override { CachedLabel = InNewName; }
 	/** Updates the highlight text to active search term. */
 	virtual void SetHighlightText(const FText& InHightlightText = FText::GetEmpty()) override { HighlightText = InHightlightText; }
 
@@ -95,10 +101,8 @@ private:
 	bool OnVerifyItemLabelChanged(const FText& InLabel, FText& OutErrorMessage);
 	/** Handles committing a entity label. */
 	void OnLabelCommitted(const FText& InLabel, ETextCommit::Type InCommitInfo);
-	/** Handles text change of a property id label. */
-	FText OnPropertyIdText() const;
 	/** Handles committing a property id label. */
-	void OnPropertyIdTextCommitted(const FText& InText, ETextCommit::Type InCommitInfo) const;
+	void OnPropertyIdTextCommitted(const FText& InText, ETextCommit::Type InCommitInfo);
 	/** Returns whether or not the actor is selectable for a binding replacement. */
 	bool IsActorSelectable(const AActor* Parent) const;
 	/** Handle clicking on the unexpose button. */
@@ -123,8 +127,9 @@ public:
 
 	using WidgetType = SWidget;
 
-	FExposedEntityDragDrop(TSharedPtr<SWidget> InWidget, const FGuid& InId)
-		: Id(InId)
+	FExposedEntityDragDrop(const TSharedPtr<SWidget>& InWidget, const FGuid& InNodeId, const TArray<FGuid>& InSelectedIds)
+		: NodeId(InNodeId)
+		, SelectedIds(InSelectedIds)
 	{
 		DecoratorWidget = SNew(SBorder)
 			.Padding(1.0f)
@@ -136,9 +141,15 @@ public:
 	}
 
 	/** Get the ID of the represented entity or group. */
-	FGuid GetId() const
+	FGuid GetNodeId() const
 	{
-		return Id;
+		return NodeId;
+	}
+
+	/** Get the IDs that were selected at the time the drag started. */
+	const TArray<FGuid>& GetSelectedIds() const
+	{
+		return SelectedIds;
 	}
 
 	virtual void OnDrop(bool bDropWasHandled, const FPointerEvent& MouseEvent) override
@@ -152,6 +163,10 @@ public:
 	}
 
 private:
-	FGuid Id;
+	/** ID of the represented entity or group. */
+	FGuid NodeId;
+	/** IDs that were selected at the time the drag started. */
+	TArray<FGuid> SelectedIds;
+	/** Decorator Drag and Drop widget. */
 	TSharedPtr<SWidget> DecoratorWidget;
 };
diff --git a/Engine/Plugins/VirtualProduction/RemoteControl/Source/RemoteControlUI/Private/UI/SRCPanelFieldGroup.cpp b/Engine/Plugins/VirtualProduction/RemoteControl/Source/RemoteControlUI/Private/UI/SRCPanelFieldGroup.cpp
--- a/Engine/Plugins/VirtualProduction/RemoteControl/Source/RemoteControlUI/Private/UI/SRCPanelFieldGroup.cpp
+++ b/Engine/Plugins/VirtualProduction/RemoteControl/Source/RemoteControlUI/Private/UI/SRCPanelFieldGroup.cpp
@@ -209,10 +209,13 @@ bool SRCPanelGroup::OnAllowDropFromOtherGroup(TSharedPtr<FDragDropOperation> Dra
 		{
 			if (OnGetGroupId.IsBound())
 			{
-				const FGuid OriginGroupId = OnGetGroupId.Execute(DragDropOp->GetId());
-				if (OriginGroupId != Id)
+				for (FGuid SelectedId : DragDropOp->GetSelectedIds())
 				{
-					return true;
+					const FGuid OriginGroupId = OnGetGroupId.Execute(SelectedId);
+					if (OriginGroupId != Id)
+					{
+						return true;
+					}
 				}
 			}
 		}
diff --git a/Engine/Plugins/VirtualProduction/RemoteControl/Source/RemoteControlUI/Private/UI/SRCPanelTreeNode.cpp b/Engine/Plugins/VirtualProduction/RemoteControl/Source/RemoteControlUI/Private/UI/SRCPanelTreeNode.cpp
--- a/Engine/Plugins/VirtualProduction/RemoteControl/Source/RemoteControlUI/Private/UI/SRCPanelTreeNode.cpp
+++ b/Engine/Plugins/VirtualProduction/RemoteControl/Source/RemoteControlUI/Private/UI/SRCPanelTreeNode.cpp
@@ -78,7 +78,7 @@ TSharedRef<SWidget> SRCPanelTreeNode::GetWidget(const FName ForColumnName, const
 	return SNullWidget::NullWidget;
 }
 
-TSharedRef<SWidget> SRCPanelTreeNode::GetDragAndDropWidget()
+TSharedRef<SWidget> SRCPanelTreeNode::GetDragAndDropWidget(int32 InSelectedEntitiesNum)
 {
 	if (NodeNameWidget && NodeValueWidget)
 	{
@@ -92,11 +92,12 @@ TSharedRef<SWidget> SRCPanelTreeNode::GetDragAndDropWidget()
 				NodeNameWidget.ToSharedRef()
 			];
 
-		const TSharedRef<SWidget> RightColumn = SNew(SHorizontalBox)
+		const TSharedRef<SHorizontalBox> RightColumn = SNew(SHorizontalBox)
 			.Clipping(EWidgetClipping::OnDemand)
 			// Node Value
 			+ SHorizontalBox::Slot()
 			.HAlign(HAlign_Fill)
+			.AutoWidth()
 			[
 				SNew(SBox)
 				.HAlign(HAlign_Left)
@@ -106,7 +107,35 @@ TSharedRef<SWidget> SRCPanelTreeNode::GetDragAndDropWidget()
 				]
 			];
 
-		return MakeSplitRow(LeftColumn, RightColumn);
+		if (InSelectedEntitiesNum > 1)
+		{
+			RightColumn->AddSlot()
+				.HAlign(HAlign_Fill)
+				.AutoWidth()
+				.Padding(2, 2)
+				[
+					SNew(SBox)
+					.HAlign(HAlign_Center)
+					.VAlign(VAlign_Center)
+					[
+						SNew(STextBlock)
+						.Text(FText::Format(FText::FromString("and {0} other item(s)"), InSelectedEntitiesNum - 1))
+					]
+				];
+		}
+		return SNew(SSplitter)
+			.Style(FAppStyle::Get(), "DetailsView.Splitter")
+			.PhysicalSplitterHandleSize(1.0f)
+			.HitDetectionSplitterHandleSize(5.0f)
+			+ SSplitter::Slot()
+			[
+				LeftColumn
+			]
+			+ SSplitter::Slot()
+			.SizeRule(SSplitter::SizeToContent)
+			[
+				RightColumn
+			];
 	}
 
 	return SNullWidget::NullWidget;
diff --git a/Engine/Plugins/VirtualProduction/RemoteControl/Source/RemoteControlUI/Private/UI/SRCPanelTreeNode.h b/Engine/Plugins/VirtualProduction/RemoteControl/Source/RemoteControlUI/Private/UI/SRCPanelTreeNode.h
--- a/Engine/Plugins/VirtualProduction/RemoteControl/Source/RemoteControlUI/Private/UI/SRCPanelTreeNode.h
+++ b/Engine/Plugins/VirtualProduction/RemoteControl/Source/RemoteControlUI/Private/UI/SRCPanelTreeNode.h
@@ -40,6 +40,9 @@ namespace ERCColumn
 	};
 }
 
+DECLARE_DELEGATE_OneParam(FOnLabelModified, FName /* InNewName */)
+DECLARE_DELEGATE_OneParam(FOnPropertyIdRenamed, FName /* InNewPropertyId */)
+
 /** A node in the panel tree view. */
 struct SRCPanelTreeNode : public SCompoundWidget, public IHasProtocolExtensibility
 {
@@ -68,7 +71,7 @@ struct SRCPanelTreeNode : public SCompoundWidget, public IHasProtocolExtensibili
 	/** Get this node's ID if any. */
 	virtual FGuid GetRCId() const { return FGuid(); }
 	/** Get get this node's type. */
-	virtual ENodeType GetRCType() const { return ENodeType::Invalid; };
+	virtual ENodeType GetRCType() const { return ENodeType::Invalid; }
 	/** Returns true if this tree node has childen. */
 	virtual bool HasChildren() const { return false; }
 	/** Refresh the node. */
@@ -79,12 +82,22 @@ struct SRCPanelTreeNode : public SCompoundWidget, public IHasProtocolExtensibili
 	virtual void SetIsHovered(bool bIsBeingHovered) {}
 	/** Make the node name's text box editable. */
 	virtual void EnterRenameMode() {};
+	/** Get the PropertyId of this Node. */
+	virtual FName GetPropertyId() { return FName(TEXT("")); }
+	/** Set the PropertyId of this Node. */
+	virtual void SetPropertyId(FName InNewPropertyId) {};
+	/** Set the Name of this Node. */
+	virtual void SetName(FName InNewName) {};
 	/** Updates the highlight text to active search term. */
 	virtual void SetHighlightText(const FText& InHightlightText = FText::GetEmpty()) {};
 	/** Retrieves the referenced widget corresponding to the given column name. */
 	virtual TSharedRef<SWidget> GetWidget(const FName ForColumnName, const FName InActiveProtocol);
 	/** Retrieves the DragAndDropWidget if possible otherwise returns a NullWidget */
-	TSharedRef<SWidget> GetDragAndDropWidget();
+	TSharedRef<SWidget> GetDragAndDropWidget(int32 InSelectedEntitiesNum = 1);
+	/** Executed when the Name of the node is changed */
+	FOnLabelModified& OnNameRenamed() { return OnNameRenamedDelegate; };
+	/** Executed when the PropertyId of the node is changed */
+	FOnPropertyIdRenamed& OnPropertyIdRenamed() { return OnPropertyIdRenamedDelegate; };
 
 protected:
 	struct FMakeNodeWidgetArgs
@@ -145,4 +158,10 @@ private:
 private:
 	/** The splitter offset to align the group splitter with the other row's splitters. */
 	static constexpr float SplitterOffset = 0.008f;
+
+	/** Delegate called when the Node Name is renamed. */
+	FOnLabelModified OnNameRenamedDelegate;
+
+	/** Delegate called when the Node Property Id change. */
+	FOnPropertyIdRenamed OnPropertyIdRenamedDelegate;
 };
diff --git a/Engine/Plugins/VirtualProduction/RemoteControl/Source/RemoteControlUI/Private/UI/SRemoteControlPanel.cpp b/Engine/Plugins/VirtualProduction/RemoteControl/Source/RemoteControlUI/Private/UI/SRemoteControlPanel.cpp
--- a/Engine/Plugins/VirtualProduction/RemoteControl/Source/RemoteControlUI/Private/UI/SRemoteControlPanel.cpp
+++ b/Engine/Plugins/VirtualProduction/RemoteControl/Source/RemoteControlUI/Private/UI/SRemoteControlPanel.cpp
@@ -1741,7 +1741,7 @@ TSharedRef<SWidget> SRemoteControlPanel::CreateEntityDetailsView()
 
 	UpdateEntityDetailsView(EntityList->GetSelectedEntity());
 
-	const bool bCanShowDetailsView = SelectedEntity.IsValid() && (SelectedEntity->GetRCType() != SRCPanelTreeNode::Group) && (SelectedEntity->GetRCType() != SRCPanelTreeNode::FieldChild);
+	const bool bCanShowDetailsView = LastSelectedEntity.IsValid() && (LastSelectedEntity->GetRCType() != SRCPanelTreeNode::Group) && (LastSelectedEntity->GetRCType() != SRCPanelTreeNode::FieldChild);
 
 	if (bCanShowDetailsView)
 	{
@@ -1762,19 +1762,19 @@ void SRemoteControlPanel::UpdateEntityDetailsView(const TSharedPtr<SRCPanelTreeN
 {
 	TSharedPtr<FStructOnScope> SelectedEntityPtr;
 
-	SelectedEntity = SelectedNode;
+	LastSelectedEntity = SelectedNode;
 
-	if (SelectedEntity)
+	if (LastSelectedEntity)
 	{
-		if (SelectedEntity->GetRCType() != SRCPanelTreeNode::Group &&
-			SelectedEntity->GetRCType() != SRCPanelTreeNode::FieldChild) // Field Child does not contain entity ID, that is why it should not be processed
+		if (LastSelectedEntity->GetRCType() != SRCPanelTreeNode::Group &&
+			LastSelectedEntity->GetRCType() != SRCPanelTreeNode::FieldChild) // Field Child does not contain entity ID, that is why it should not be processed
 		{
-			const TSharedPtr<FRemoteControlEntity> Entity = Preset->GetExposedEntity<FRemoteControlEntity>(SelectedEntity->GetRCId()).Pin();
-			SelectedEntityPtr = RemoteControlPanelUtils::GetEntityOnScope(Entity, Preset->GetExposedEntityType(SelectedEntity->GetRCId()));
+			const TSharedPtr<FRemoteControlEntity> Entity = Preset->GetExposedEntity<FRemoteControlEntity>(LastSelectedEntity->GetRCId()).Pin();
+			SelectedEntityPtr = RemoteControlPanelUtils::GetEntityOnScope(Entity, Preset->GetExposedEntityType(LastSelectedEntity->GetRCId()));
 		}
 	}
 
-	const bool bCanShowDetailsView = SelectedEntity.IsValid() && (SelectedEntity->GetRCType() != SRCPanelTreeNode::Group) && (SelectedEntity->GetRCType() != SRCPanelTreeNode::FieldChild);
+	const bool bCanShowDetailsView = LastSelectedEntity.IsValid() && (LastSelectedEntity->GetRCType() != SRCPanelTreeNode::Group) && (LastSelectedEntity->GetRCType() != SRCPanelTreeNode::FieldChild);
 
 	if (bCanShowDetailsView)
 	{
@@ -1791,9 +1791,9 @@ void SRemoteControlPanel::UpdateEntityDetailsView(const TSharedPtr<SRCPanelTreeN
 	}
 
 	static const FName ProtocolWidgetsModuleName = "RemoteControlProtocolWidgets";
-	if(SelectedEntity.IsValid() && FModuleManager::Get().IsModuleLoaded(ProtocolWidgetsModuleName) && ensure(Preset.IsValid()))
+	if(LastSelectedEntity.IsValid() && FModuleManager::Get().IsModuleLoaded(ProtocolWidgetsModuleName) && ensure(Preset.IsValid()))
 	{
-		if (const TSharedPtr<FRemoteControlEntity> RCEntity = Preset->GetExposedEntity(SelectedEntity->GetRCId()).Pin())
+		if (const TSharedPtr<FRemoteControlEntity> RCEntity = Preset->GetExposedEntity(LastSelectedEntity->GetRCId()).Pin())
 		{
 			if(RCEntity->IsBound())
 			{
@@ -2426,22 +2426,29 @@ void SRemoteControlPanel::DeleteEntity_Execute()
 
 	// ~ Delete Entity ~
 
-	if (SelectedEntity->GetRCType() == SRCPanelTreeNode::FieldChild) // Field Child does not contain entity ID, that is why it should not be processed
+	if (LastSelectedEntity->GetRCType() == SRCPanelTreeNode::FieldChild) // Field Child does not contain entity ID, that is why it should not be processed
 	{
 		return;
 	}
 
-	if (SelectedEntity->GetRCType() == SRCPanelTreeNode::Group)
+	if (LastSelectedEntity->GetRCType() == SRCPanelTreeNode::Group)
 	{
 		FScopedTransaction Transaction(LOCTEXT("DeleteGroup", "Delete Group"));
 		Preset->Modify();
-		Preset->Layout.DeleteGroup(SelectedEntity->GetRCId());
+		Preset->Layout.DeleteGroup(LastSelectedEntity->GetRCId());
 	}
 	else
 	{
 		FScopedTransaction Transaction(LOCTEXT("UnexposeFunction", "Unexpose remote control entity"));
 		Preset->Modify();
-		Preset->Unexpose(SelectedEntity->GetRCId());
+		TArray<TSharedPtr<SRCPanelTreeNode>> SelectedEntities = EntityList->GetSelectedEntities();
+		for (int32 Index = 0; Index < SelectedEntities.Num(); ++Index)
+		{
+			if (SelectedEntities[Index]->GetRCType() != SRCPanelTreeNode::FieldChild)
+			{
+				Preset->Unexpose(SelectedEntities[Index]->GetRCId());
+			}
+		}
 	}
 
 	EntityList->Refresh();
@@ -2459,10 +2466,10 @@ bool SRemoteControlPanel::CanDeleteEntity() const
 		return ActiveLogicPanel->GetSelectedLogicItem() != nullptr; // User has focus on a logic panel
 	}
 
-	if (SelectedEntity.IsValid() && Preset.IsValid())
+	if (LastSelectedEntity.IsValid() && Preset.IsValid())
 	{
 		// Do not allow default group to be deleted.
-		return !Preset->Layout.IsDefaultGroup(SelectedEntity->GetRCId());
+		return !Preset->Layout.IsDefaultGroup(LastSelectedEntity->GetRCId());
 	}
 
 	return false;
@@ -2476,12 +2483,12 @@ void SRemoteControlPanel::RenameEntity_Execute() const
 		return;
 	}
 
-	if (SelectedEntity->GetRCType() == SRCPanelTreeNode::FieldChild) // Field Child does not contain entity ID, that is why it should not be processed
+	if (LastSelectedEntity->GetRCType() == SRCPanelTreeNode::FieldChild) // Field Child does not contain entity ID, that is why it should not be processed
 	{
 		return;
 	}
 
-	SelectedEntity->EnterRenameMode();
+	LastSelectedEntity->EnterRenameMode();
 }
 
 bool SRemoteControlPanel::CanRenameEntity() const
@@ -2496,10 +2503,10 @@ bool SRemoteControlPanel::CanRenameEntity() const
 		return true;
 	}
 
-	if (SelectedEntity.IsValid() && Preset.IsValid())
+	if (LastSelectedEntity.IsValid() && Preset.IsValid())
 	{
 		// Do not allow default group to be renamed.
-		return !Preset->Layout.IsDefaultGroup(SelectedEntity->GetRCId());
+		return !Preset->Layout.IsDefaultGroup(LastSelectedEntity->GetRCId());
 	}
 
 	return false;
diff --git a/Engine/Plugins/VirtualProduction/RemoteControl/Source/RemoteControlUI/Private/UI/SRemoteControlPanel.h b/Engine/Plugins/VirtualProduction/RemoteControl/Source/RemoteControlUI/Private/UI/SRemoteControlPanel.h
--- a/Engine/Plugins/VirtualProduction/RemoteControl/Source/RemoteControlUI/Private/UI/SRemoteControlPanel.h
+++ b/Engine/Plugins/VirtualProduction/RemoteControl/Source/RemoteControlUI/Private/UI/SRemoteControlPanel.h
@@ -398,8 +398,8 @@ private:
 	TSharedPtr<SBorder> AuxiliaryToolbarWidgetContent;
 	/** Additional widgets to be added to the toolbar */
 	TArray<TSharedRef<SWidget>> ToolbarWidgets;
-	/** Holds a shared pointer reference to the active entity that is selected. */
-	TSharedPtr<SRCPanelTreeNode> SelectedEntity;
+	/** Holds a shared pointer reference to the last entity that was selected. */
+	TSharedPtr<SRCPanelTreeNode> LastSelectedEntity;
 	/** Panel Drawer widget holds all docked panels. */
 	TSharedPtr<SRCPanelDrawer> PanelDrawer;
 	/** Map of Opened Drawers. */
