[RemoteControl]
Tooltip for exposing a property is incorrect.

diff --git a/Engine/Plugins/VirtualProduction/RemoteControl/Source/RemoteControlUI/Private/RemoteControlUIModule.cpp b/Engine/Plugins/VirtualProduction/RemoteControl/Source/RemoteControlUI/Private/RemoteControlUIModule.cpp
--- a/Engine/Plugins/VirtualProduction/RemoteControl/Source/RemoteControlUI/Private/RemoteControlUIModule.cpp
+++ b/Engine/Plugins/VirtualProduction/RemoteControl/Source/RemoteControlUI/Private/RemoteControlUIModule.cpp
@@ -1036,23 +1036,25 @@ void FRemoteControlUIModule::RegisterWidgetFactories()
 
 FText FRemoteControlUIModule::GetExposePropertyButtonTooltip(const FRCExposesPropertyArgs InPropertyArgs) const
 {
-	if (URemoteControlPreset* Preset = GetActivePreset())
+	if (const TSharedPtr<SRemoteControlPanel> Panel = GetPanelForProperty(InPropertyArgs))
 	{
-		const FText PresetName = FText::FromString(Preset->GetName());
-		if (GetPropertyExposeStatus(InPropertyArgs) == EPropertyExposeStatus::Exposed)
+		if (const URemoteControlPreset* Preset = Panel->GetPreset())
 		{
-			return FText::Format(LOCTEXT("ExposePropertyToolTip", "Unexpose this property from RemoteControl Preset '{0}'."), PresetName);
-		}
-		else
-		{
-			return FText::Format(LOCTEXT("UnexposePropertyToolTip", "Expose this property in RemoteControl Preset '{0}'."), PresetName);
+			const FText PresetName = FText::FromString(Preset->GetName());
+			if (GetPropertyExposeStatus(InPropertyArgs) == EPropertyExposeStatus::Exposed)
+			{
+				return FText::Format(LOCTEXT("ExposePropertyToolTip", "Unexpose this property from RemoteControl Preset '{0}'."), PresetName);
+			}
+			else
+			{
+				return FText::Format(LOCTEXT("UnexposePropertyToolTip", "Expose this property in RemoteControl Preset '{0}'."), PresetName);
+			}
 		}
 	}
 
 	return LOCTEXT("InvalidExposePropertyTooltip", "Invalid Preset");
 }
 
-
 FText FRemoteControlUIModule::GetExposePropertyButtonText(const FRCExposesPropertyArgs InPropertyArgs) const
 {
 	if (GetPropertyExposeStatus(InPropertyArgs) == EPropertyExposeStatus::Exposed)
