[RemoteControl]
Fix PtrContainer never valid in PropertyId and fix for ByteProperty that are Enum.

diff --git a/Engine/Plugins/VirtualProduction/RemoteControl/Source/RemoteControl/Private/RemoteControlPropertyIdRegistry.cpp b/Engine/Plugins/VirtualProduction/RemoteControl/Source/RemoteControl/Private/RemoteControlPropertyIdRegistry.cpp
--- a/Engine/Plugins/VirtualProduction/RemoteControl/Source/RemoteControl/Private/RemoteControlPropertyIdRegistry.cpp
+++ b/Engine/Plugins/VirtualProduction/RemoteControl/Source/RemoteControl/Private/RemoteControlPropertyIdRegistry.cpp
@@ -44,21 +44,27 @@ namespace RemoteControlPropertyIdUtilities
 		FRCFieldPathInfo InSegment,
 		const int32 InIndexSegment,
 		const uint8* InPropertyContainer,
-		const bool bCheckByteEnumComparison)
+		const bool bCheckByteEnumComparison,
+		uint8** OutPtrContainer)
 	{
 		FScriptArrayHelper ArrayHelper(InArrayProperty, InPropertyContainer);
 		if (const TSharedPtr<FRemoteControlField> TargetRCField = StaticCastSharedPtr<FRemoteControlField>(InTargetRCProperty);
 			!TargetRCField->FieldPathInfo.Segments.IsEmpty())
 		{
+			if (!OutPtrContainer)
+            {
+            	return false;
+            }
+
 			const int32 Index = InSegment.Segments[InIndexSegment].ArrayIndex;
-			uint8* PtrContainer = ArrayHelper.GetElementPtr(Index);
+			*OutPtrContainer = ArrayHelper.GetElementPtr(Index);
 			if (InIndexSegment < InSegment.GetSegmentCount() - 2)
 			{
 				return false;
 			}
 			if (const FArrayProperty* ArrayRealProperty = CastField<FArrayProperty>(InProperty))
 			{
-				return InVirtualPropertySelfContainer->CopyCompleteValue(ArrayRealProperty->Inner, PtrContainer, bCheckByteEnumComparison);
+				return InVirtualPropertySelfContainer->CopyCompleteValue(ArrayRealProperty->Inner, *OutPtrContainer, bCheckByteEnumComparison);
 			}
 		}
 		return false;
@@ -71,21 +77,27 @@ namespace RemoteControlPropertyIdUtilities
 		FRCFieldPathInfo InSegment,
 		const int32 InIndexSegment,
 		const uint8* InPropertyContainer,
-		const bool bCheckByteEnumComparison)
+		const bool bCheckByteEnumComparison,
+		uint8** OutPtrContainer)
 	{
 		FScriptMapHelper MapHelper(InMapProperty, InPropertyContainer);
 		if (const TSharedPtr<FRemoteControlField> TargetRCField = StaticCastSharedPtr<FRemoteControlField>(InTargetRCProperty);
 			!TargetRCField->FieldPathInfo.Segments.IsEmpty())
 		{
+			if (!OutPtrContainer)
+			{
+				return false;
+			}
+
 			const int32 Index = InSegment.Segments[InIndexSegment].ArrayIndex;
-			uint8* PtrContainer = MapHelper.GetValuePtr(Index);
+			*OutPtrContainer = MapHelper.GetValuePtr(Index);
 			if (InIndexSegment < InSegment.GetSegmentCount() - 2)
 			{
 				return false;
 			}
 			if (const FMapProperty* MapRealProperty = CastField<FMapProperty>(InProperty))
 			{
-				return InVirtualPropertySelfContainer->CopyCompleteValue(MapRealProperty->ValueProp, PtrContainer, bCheckByteEnumComparison);
+				return InVirtualPropertySelfContainer->CopyCompleteValue(MapRealProperty->ValueProp, *OutPtrContainer, bCheckByteEnumComparison);
 			}
 		}
 		return false;
@@ -98,21 +110,27 @@ namespace RemoteControlPropertyIdUtilities
 		FRCFieldPathInfo InSegment,
 		const int32 InIndexSegment,
 		const uint8* InPropertyContainer,
-		const bool bCheckByteEnumComparison)
+		const bool bCheckByteEnumComparison,
+		uint8** OutPtrContainer)
 	{
 		FScriptSetHelper SetHelper(InSetProperty, InPropertyContainer);
 		if (const TSharedPtr<FRemoteControlField> TargetRCField = StaticCastSharedPtr<FRemoteControlField>(InTargetRCProperty);
 			!TargetRCField->FieldPathInfo.Segments.IsEmpty())
 		{
+			if (!OutPtrContainer)
+			{
+				return false;
+			}
+
 			const int32 Index = InSegment.Segments[InIndexSegment].ArrayIndex;
-			uint8* PtrContainer = SetHelper.GetElementPtr(Index);
+			*OutPtrContainer = SetHelper.GetElementPtr(Index);
 			if (InIndexSegment < InSegment.GetSegmentCount() - 2)
 			{
 				return false;
 			}
 			if (const FSetProperty* SetRealProperty = CastField<FSetProperty>(InProperty))
 			{
-				return InVirtualPropertySelfContainer->CopyCompleteValue(SetRealProperty->ElementProp, PtrContainer, bCheckByteEnumComparison);
+				return InVirtualPropertySelfContainer->CopyCompleteValue(SetRealProperty->ElementProp, *OutPtrContainer, bCheckByteEnumComparison);
 			}
 		}
 		return false;
@@ -212,6 +230,7 @@ void FRCPropertyIdWrapper::UpdateTypes(const TSharedRef<FRemoteControlProperty>&
 	{
 		if (ByteProperty->Enum)
 		{
+			SuperType = NAME_EnumProperty;
 			SubType = ByteProperty->Enum->GetFName();
 		}
 	}
@@ -623,7 +642,7 @@ bool URemoteControlPropertyIdRegistry::TryCopyNonUClassOwnerProperty(const TObje
 			if (const uint8* PropertyValuePtr = ArrayProperty->ContainerPtrToValuePtr<uint8>(InBoundObject))
 			{
 				uint8* PtrContainer = nullptr;
-				if (RemoteControlPropertyIdUtilities::CopyNonUClassOwnerArray(InVirtualPropertySelfContainer, ArrayProperty, InTargetRCProperty, InProperty, InSegment, InIndexSegment, PropertyValuePtr, bCheckByteEnumComparison))
+				if (RemoteControlPropertyIdUtilities::CopyNonUClassOwnerArray(InVirtualPropertySelfContainer, ArrayProperty, InTargetRCProperty, InProperty, InSegment, InIndexSegment, PropertyValuePtr, bCheckByteEnumComparison, &PtrContainer))
 				{
 					return true;
 				}
@@ -633,7 +652,7 @@ bool URemoteControlPropertyIdRegistry::TryCopyNonUClassOwnerProperty(const TObje
 		else if (InPropertyContainer)
 		{
 			uint8* PtrContainer = nullptr;
-			if (RemoteControlPropertyIdUtilities::CopyNonUClassOwnerArray(InVirtualPropertySelfContainer, ArrayProperty, InTargetRCProperty, InProperty, InSegment, InIndexSegment, InPropertyContainer, bCheckByteEnumComparison))
+			if (RemoteControlPropertyIdUtilities::CopyNonUClassOwnerArray(InVirtualPropertySelfContainer, ArrayProperty, InTargetRCProperty, InProperty, InSegment, InIndexSegment, InPropertyContainer, bCheckByteEnumComparison, &PtrContainer))
 			{
 				return true;
 			}
@@ -647,7 +666,7 @@ bool URemoteControlPropertyIdRegistry::TryCopyNonUClassOwnerProperty(const TObje
 			if (const uint8* PropertyValuePtr = SetProperty->ContainerPtrToValuePtr<uint8>(InBoundObject))
 			{
 				uint8* PtrContainer = nullptr;
-				if (RemoteControlPropertyIdUtilities::CopyNonUClassOwnerSet(InVirtualPropertySelfContainer, SetProperty, InTargetRCProperty, InProperty, InSegment, InIndexSegment, PropertyValuePtr, bCheckByteEnumComparison))
+				if (RemoteControlPropertyIdUtilities::CopyNonUClassOwnerSet(InVirtualPropertySelfContainer, SetProperty, InTargetRCProperty, InProperty, InSegment, InIndexSegment, PropertyValuePtr, bCheckByteEnumComparison, &PtrContainer))
 				{
 					return true;
 				}
@@ -657,7 +676,7 @@ bool URemoteControlPropertyIdRegistry::TryCopyNonUClassOwnerProperty(const TObje
 		else if (InPropertyContainer)
 		{
 			uint8* PtrContainer = nullptr;
-			if (RemoteControlPropertyIdUtilities::CopyNonUClassOwnerSet(InVirtualPropertySelfContainer, SetProperty, InTargetRCProperty, InProperty, InSegment, InIndexSegment, InPropertyContainer, bCheckByteEnumComparison))
+			if (RemoteControlPropertyIdUtilities::CopyNonUClassOwnerSet(InVirtualPropertySelfContainer, SetProperty, InTargetRCProperty, InProperty, InSegment, InIndexSegment, InPropertyContainer, bCheckByteEnumComparison, &PtrContainer))
 			{
 				return true;
 			}
@@ -671,7 +690,7 @@ bool URemoteControlPropertyIdRegistry::TryCopyNonUClassOwnerProperty(const TObje
 			if (const uint8* PropertyValuePtr = MapProperty->ContainerPtrToValuePtr<uint8>(InBoundObject))
 			{
 				uint8* PtrContainer = nullptr;
-				if (RemoteControlPropertyIdUtilities::CopyNonUClassOwnerMap(InVirtualPropertySelfContainer, MapProperty, InTargetRCProperty, InProperty, InSegment, InIndexSegment, PropertyValuePtr, bCheckByteEnumComparison))
+				if (RemoteControlPropertyIdUtilities::CopyNonUClassOwnerMap(InVirtualPropertySelfContainer, MapProperty, InTargetRCProperty, InProperty, InSegment, InIndexSegment, PropertyValuePtr, bCheckByteEnumComparison, &PtrContainer))
 				{
 					return true;
 				}
@@ -681,7 +700,7 @@ bool URemoteControlPropertyIdRegistry::TryCopyNonUClassOwnerProperty(const TObje
 		else if (InPropertyContainer)
 		{
 			uint8* PtrContainer = nullptr;
-			if (RemoteControlPropertyIdUtilities::CopyNonUClassOwnerMap(InVirtualPropertySelfContainer, MapProperty, InTargetRCProperty, InProperty, InSegment, InIndexSegment, InPropertyContainer, bCheckByteEnumComparison))
+			if (RemoteControlPropertyIdUtilities::CopyNonUClassOwnerMap(InVirtualPropertySelfContainer, MapProperty, InTargetRCProperty, InProperty, InSegment, InIndexSegment, InPropertyContainer, bCheckByteEnumComparison, &PtrContainer))
 			{
 				return true;
 			}
