[RemoteControl]
Fix Undo/Redo crash with the WebApp open while calling the OnPropertyUnexposed.
Fix Rename issue in the WebApp and OnFieldsRenamed never broadcasted in the Engine.
Fix Rename Transaction not created correctly.
Fix RCTransactionListener not keeping the Delegate by reference causing the system to not call bound function to the Delegate if they were bound after the creation of the TransactionListener.

diff --git a/Engine/Plugins/VirtualProduction/RemoteControl/Source/RemoteControl/Private/RemoteControlEntity.cpp b/Engine/Plugins/VirtualProduction/RemoteControl/Source/RemoteControl/Private/RemoteControlEntity.cpp
--- a/Engine/Plugins/VirtualProduction/RemoteControl/Source/RemoteControl/Private/RemoteControlEntity.cpp
+++ b/Engine/Plugins/VirtualProduction/RemoteControl/Source/RemoteControl/Private/RemoteControlEntity.cpp
@@ -137,7 +137,6 @@ FName FRemoteControlEntity::Rename(FName NewLabel)
 	{
 		Preset->Modify();
 		FName NewName = Preset->RenameExposedEntity(Id, NewLabel);
-		OnEntityModifiedDelegate.ExecuteIfBound(Id);
 		return NewName;
 	}
 
diff --git a/Engine/Plugins/VirtualProduction/RemoteControl/Source/RemoteControl/Private/RemoteControlPreset.cpp b/Engine/Plugins/VirtualProduction/RemoteControl/Source/RemoteControl/Private/RemoteControlPreset.cpp
--- a/Engine/Plugins/VirtualProduction/RemoteControl/Source/RemoteControl/Private/RemoteControlPreset.cpp
+++ b/Engine/Plugins/VirtualProduction/RemoteControl/Source/RemoteControl/Private/RemoteControlPreset.cpp
@@ -1792,7 +1792,19 @@ const bool URemoteControlPreset::HasEntities() const
 
 FName URemoteControlPreset::RenameExposedEntity(const FGuid& ExposedEntityId, FName NewLabel)
 {
-	FName AssignedLabel = Registry->RenameExposedEntity(ExposedEntityId, NewLabel);
+	Registry->Modify();
+
+	FName OldLabel = FName(TEXT("Invalid Label"));
+	if (const TSharedPtr<FRemoteControlEntity>& Entity = Registry->GetExposedEntity(ExposedEntityId))
+	{
+		OldLabel = Entity->GetLabel();
+	}
+
+	const FName AssignedLabel = Registry->RenameExposedEntity(ExposedEntityId, NewLabel);
+
+	FRCTransactionListenerHelper<URemoteControlPreset*, FName, FName>(ERCTransaction::Undo, GetPresetId(), OnFieldRenamed(), this, AssignedLabel, OldLabel);
+	FRCTransactionListenerHelper<URemoteControlPreset*, FName, FName>(ERCTransaction::Redo, GetPresetId(), OnFieldRenamed(), this, OldLabel, AssignedLabel);
+
 	return AssignedLabel;
 }
 
diff --git a/Engine/Plugins/VirtualProduction/RemoteControl/Source/RemoteControl/Private/RemoteControlTransactionListenerHelper.h b/Engine/Plugins/VirtualProduction/RemoteControl/Source/RemoteControl/Private/RemoteControlTransactionListenerHelper.h
--- a/Engine/Plugins/VirtualProduction/RemoteControl/Source/RemoteControl/Private/RemoteControlTransactionListenerHelper.h
+++ b/Engine/Plugins/VirtualProduction/RemoteControl/Source/RemoteControl/Private/RemoteControlTransactionListenerHelper.h
@@ -22,7 +22,7 @@ public:
 	 * @param InDelegateToBroadcast Multicast delegate to broadcast
 	 * @param InOrderedArgs Ordered arguments that should be passed to the MulticastDelegate
 	 */
-	FRCTransactionListenerHelper(ERCTransaction::Type InType, const FGuid& PresetGuid, TMulticastDelegate<void(CallbackTypeArgs...)> InDelegateToBroadcast, CallbackTypeArgs... InOrderedArgs)
+	FRCTransactionListenerHelper(ERCTransaction::Type InType, const FGuid& PresetGuid, TMulticastDelegate<void(CallbackTypeArgs...)>& InDelegateToBroadcast, CallbackTypeArgs... InOrderedArgs)
 	{
 #if WITH_EDITOR
 		if (!IRemoteControlModule::Get().IsPresetTransient(PresetGuid))
diff --git a/Engine/Plugins/VirtualProduction/RemoteControl/Source/RemoteControl/Private/RemoteControlTransactionListenerManager.h b/Engine/Plugins/VirtualProduction/RemoteControl/Source/RemoteControl/Private/RemoteControlTransactionListenerManager.h
--- a/Engine/Plugins/VirtualProduction/RemoteControl/Source/RemoteControl/Private/RemoteControlTransactionListenerManager.h
+++ b/Engine/Plugins/VirtualProduction/RemoteControl/Source/RemoteControl/Private/RemoteControlTransactionListenerManager.h
@@ -46,11 +46,11 @@ class FRCTransactionListenerManager
 	 	* @param InDelegateToBroadcast Multicast delegate to broadcast
 	 	* @param InArgs Ordered arguments that should be passed to the MulticastDelegate
 		*/
-		FRCTransactionListener(ERCTransaction::Type InType, TMulticastDelegate<void(ListenerCallbackArgs...)> InDelegateToBroadcast, ListenerCallbackArgs&&... InArgs) : Type(InType)
+		FRCTransactionListener(ERCTransaction::Type InType, TMulticastDelegate<void(ListenerCallbackArgs...)>& InDelegateToBroadcast, ListenerCallbackArgs&&... InArgs) : Type(InType)
 		{
-			OnUndoRedoDelegate.AddLambda([InDelegateToBroadcast, InArgs...]()
+			OnUndoRedoDelegate.AddLambda([&InDelegateToBroadcast, InArgs...]()
 			{
-				InDelegateToBroadcast.Broadcast(InArgs...);	
+				InDelegateToBroadcast.Broadcast(InArgs...);
 			});
 		}
 
@@ -81,7 +81,7 @@ public:
 	/**
  	 * Create a listener with the parameter passed and Register it into the Manager to be called during Undo/Redo callbacks
  	 */
-	static void CreateListenerAndRegister(ERCTransaction::Type InType, TMulticastDelegate<void(CallbackArgs...)> InDelegateToBroadcast, CallbackArgs&&... InOrderedArgs)
+	static void CreateListenerAndRegister(ERCTransaction::Type InType, TMulticastDelegate<void(CallbackArgs...)>& InDelegateToBroadcast, CallbackArgs&&... InOrderedArgs)
 	{
 		FRCTransactionListener<CallbackArgs...> Listener = FRCTransactionListener<CallbackArgs...>(InType, InDelegateToBroadcast, Forward<CallbackArgs>(InOrderedArgs)...);
 		Get().Register(Listener);
diff --git a/Engine/Plugins/VirtualProduction/RemoteControl/Source/RemoteControlUI/Private/UI/SRCPanelExposedEntitiesList.cpp b/Engine/Plugins/VirtualProduction/RemoteControl/Source/RemoteControlUI/Private/UI/SRCPanelExposedEntitiesList.cpp
--- a/Engine/Plugins/VirtualProduction/RemoteControl/Source/RemoteControlUI/Private/UI/SRCPanelExposedEntitiesList.cpp
+++ b/Engine/Plugins/VirtualProduction/RemoteControl/Source/RemoteControlUI/Private/UI/SRCPanelExposedEntitiesList.cpp
@@ -757,16 +757,25 @@ void SRCPanelExposedEntitiesList::OnPropertyIdRenamed(const FName InNewId, TShar
 	}
 }
 
-void SRCPanelExposedEntitiesList::OnNameRenamed(const FName InNewName)
+void SRCPanelExposedEntitiesList::OnLabelModified(const FName InOldName, const FName InNewName)
 {
 	TArray<TSharedPtr<SRCPanelTreeNode>> SelectedEntities = GetSelectedEntities();
-	for (const TSharedPtr<SRCPanelTreeNode>& Entity : SelectedEntities)
+	for (const TSharedPtr<SRCPanelTreeNode>& EntityNode : SelectedEntities)
 	{
-		TWeakPtr<FRemoteControlField> ExposedField = Preset->GetExposedEntity<FRemoteControlField>(Entity->GetRCId());
+		TWeakPtr<FRemoteControlField> ExposedField = Preset->GetExposedEntity<FRemoteControlField>(EntityNode->GetRCId());
 		if (ExposedField.IsValid())
 		{
+			const bool bIsLabelAlreadyNew = ExposedField.Pin()->GetLabel() == InNewName;
+			FName OldLabel = InOldName;
+			if (!bIsLabelAlreadyNew)
+			{
+				OldLabel = ExposedField.Pin()->GetLabel();
+			}
+
 			ExposedField.Pin()->Rename(InNewName);
-			Entity->SetName(ExposedField.Pin()->GetLabel());
+			const FName NewName = ExposedField.Pin()->GetLabel();
+			EntityNode->SetName(NewName);
+			Preset->OnFieldRenamed().Broadcast(Preset.Get(), OldLabel, NewName);
 		}
 	}
 }
@@ -1042,7 +1051,7 @@ TSharedRef<ITableRow> SRCPanelExposedEntitiesList::OnGenerateRow(TSharedPtr<SRCP
 		constexpr float LeftPadding = 3.f;
 		const FMargin Margin = Node->GetRCType() == SRCPanelTreeNode::FieldChild ? FMargin(LeftPadding + 10.f, 1.f, 1.f, 1.f) : FMargin(LeftPadding, 1.f, 1.f, 1.f);
 		Node->OnPropertyIdRenamed().BindSP(this, &SRCPanelExposedEntitiesList::OnPropertyIdRenamed, Node);
-		Node->OnNameRenamed().BindSP(this, &SRCPanelExposedEntitiesList::OnNameRenamed);
+		Node->OnLabelModified().BindSP(this, &SRCPanelExposedEntitiesList::OnLabelModified);
 
 		return SNew(SEntityRow, OwnerTable)
 			.OnDragDetected(FOnDragDetected::CreateSP(this, &SRCPanelExposedEntitiesList::OnNodeDragDetected, Node))
diff --git a/Engine/Plugins/VirtualProduction/RemoteControl/Source/RemoteControlUI/Private/UI/SRCPanelExposedEntitiesList.h b/Engine/Plugins/VirtualProduction/RemoteControl/Source/RemoteControlUI/Private/UI/SRCPanelExposedEntitiesList.h
--- a/Engine/Plugins/VirtualProduction/RemoteControl/Source/RemoteControlUI/Private/UI/SRCPanelExposedEntitiesList.h
+++ b/Engine/Plugins/VirtualProduction/RemoteControl/Source/RemoteControlUI/Private/UI/SRCPanelExposedEntitiesList.h
@@ -221,7 +221,7 @@ private:
 	void OnPropertyIdRenamed(const FName InNewId, TSharedPtr<SRCPanelTreeNode> InNode);
 
 	/** Executed when a property Name is changed, will set all selected node(s) property Name to the new one */
-	void OnNameRenamed(const FName InNewName);
+	void OnLabelModified(const FName InOldName, const FName InNewName);
 
 	/** Executed when a drag is detected, will create the Drag and Drop widget of the node(s) */
 	FReply OnNodeDragDetected(const FGeometry& InGeometry, const FPointerEvent& InPointerEvent, TSharedPtr<SRCPanelTreeNode> InNode);
diff --git a/Engine/Plugins/VirtualProduction/RemoteControl/Source/RemoteControlUI/Private/UI/SRCPanelExposedEntity.cpp b/Engine/Plugins/VirtualProduction/RemoteControl/Source/RemoteControlUI/Private/UI/SRCPanelExposedEntity.cpp
--- a/Engine/Plugins/VirtualProduction/RemoteControl/Source/RemoteControlUI/Private/UI/SRCPanelExposedEntity.cpp
+++ b/Engine/Plugins/VirtualProduction/RemoteControl/Source/RemoteControlUI/Private/UI/SRCPanelExposedEntity.cpp
@@ -421,8 +421,9 @@ void SRCPanelExposedEntity::OnLabelCommitted(const FText& InLabel, ETextCommit::
 	{
 		FScopedTransaction Transaction(LOCTEXT("ModifyEntityLabel", "Modify exposed entity's label."));
 		RCPreset->Modify();
+		FName OldName = CachedLabel;
 		CachedLabel = RCPreset->RenameExposedEntity(EntityId, *InLabel.ToString());
-		OnNameRenamed().ExecuteIfBound(CachedLabel);
+		OnLabelModified().ExecuteIfBound(OldName, CachedLabel);
 	}
 }
 
diff --git a/Engine/Plugins/VirtualProduction/RemoteControl/Source/RemoteControlUI/Private/UI/SRCPanelTreeNode.h b/Engine/Plugins/VirtualProduction/RemoteControl/Source/RemoteControlUI/Private/UI/SRCPanelTreeNode.h
--- a/Engine/Plugins/VirtualProduction/RemoteControl/Source/RemoteControlUI/Private/UI/SRCPanelTreeNode.h
+++ b/Engine/Plugins/VirtualProduction/RemoteControl/Source/RemoteControlUI/Private/UI/SRCPanelTreeNode.h
@@ -40,7 +40,7 @@ namespace ERCColumn
 	};
 }
 
-DECLARE_DELEGATE_OneParam(FOnLabelModified, FName /* InNewName */)
+DECLARE_DELEGATE_TwoParams(FOnLabelModified, FName /* InOldName */, FName /* InNewName */)
 DECLARE_DELEGATE_OneParam(FOnPropertyIdRenamed, FName /* InNewPropertyId */)
 
 /** A node in the panel tree view. */
@@ -95,7 +95,7 @@ struct SRCPanelTreeNode : public SCompoundWidget, public IHasProtocolExtensibili
 	/** Retrieves the DragAndDropWidget if possible otherwise returns a NullWidget */
 	TSharedRef<SWidget> GetDragAndDropWidget(int32 InSelectedEntitiesNum = 1);
 	/** Executed when the Name of the node is changed */
-	FOnLabelModified& OnNameRenamed() { return OnNameRenamedDelegate; };
+	FOnLabelModified& OnLabelModified() { return OnLabelModifiedDelegate; };
 	/** Executed when the PropertyId of the node is changed */
 	FOnPropertyIdRenamed& OnPropertyIdRenamed() { return OnPropertyIdRenamedDelegate; };
 
@@ -160,7 +160,7 @@ private:
 	static constexpr float SplitterOffset = 0.008f;
 
 	/** Delegate called when the Node Name is renamed. */
-	FOnLabelModified OnNameRenamedDelegate;
+	FOnLabelModified OnLabelModifiedDelegate;
 
 	/** Delegate called when the Node Property Id change. */
 	FOnPropertyIdRenamed OnPropertyIdRenamedDelegate;
diff --git a/Engine/Plugins/VirtualProduction/RemoteControl/Source/WebRemoteControl/Private/WebSocketMessageHandler.cpp b/Engine/Plugins/VirtualProduction/RemoteControl/Source/WebRemoteControl/Private/WebSocketMessageHandler.cpp
--- a/Engine/Plugins/VirtualProduction/RemoteControl/Source/WebRemoteControl/Private/WebSocketMessageHandler.cpp
+++ b/Engine/Plugins/VirtualProduction/RemoteControl/Source/WebRemoteControl/Private/WebSocketMessageHandler.cpp
@@ -1094,6 +1094,20 @@ void FWebSocketMessageHandler::OnPropertyExposed(URemoteControlPreset* Owner, co
 		return;
 	}
 
+	// Cache used during the Unexposed as the last resort to get the property Label correctly (main use case is Undo/Redo)
+	TTuple<TArray<FGuid>, TArray<FName>>& Entries = CacheUndoRedoAddedRemovedProperties.FindOrAdd(Owner->GetPresetId());
+	Entries.Key.AddUnique(EntityId);
+
+	if (const TSharedPtr<FRemoteControlEntity> Entity = Owner->GetExposedEntity(EntityId).Pin())
+	{
+		Entries.Value.AddUnique(Entity->GetLabel());
+	}
+	else
+	{
+		// If the label couldn't be set from the Entity then we use the EntityId as the Label since it is unique
+		Entries.Value.AddUnique(FName(EntityId.ToString()));
+	}
+
 	//Cache the property field that was removed for end of frame notification
 	PerFrameAddedProperties.FindOrAdd(Owner->GetPresetId()).AddUnique(EntityId);
 }
@@ -1152,13 +1166,38 @@ void FWebSocketMessageHandler::OnPropertyUnexposed(URemoteControlPreset* Owner,
 		return;
 	}
 
-	TSharedPtr<FRemoteControlEntity> Entity = Owner->GetExposedEntity(EntityId).Pin();
-	check(Entity);
+	const TSharedPtr<FRemoteControlEntity> Entity = Owner->GetExposedEntity(EntityId).Pin();
+	TPair<TArray<FGuid>, TArray<FName>>& Entries = PerFrameRemovedProperties.FindOrAdd(Owner->GetPresetId());
 
 	// Cache the property field that was removed for end of frame notification
-	TTuple<TArray<FGuid>, TArray<FName>>& Entries = PerFrameRemovedProperties.FindOrAdd(Owner->GetPresetId());
 	Entries.Key.AddUnique(EntityId);
-	Entries.Value.AddUnique(Entity->GetLabel());
+	bool bLabelSet = false;
+
+	if (Entity.IsValid())
+	{
+		Entries.Value.AddUnique(Entity->GetLabel());
+		bLabelSet = true;
+	}
+	else
+	{
+		// If the Entity is not valid try using the cached properties saved during the Expose
+		// This is done because during the Undo/Redo the entity is already removed from the preset and we can't get the Label correctly
+		if (TPair<TArray<FGuid>, TArray<FName>>* CachedProperties = CacheUndoRedoAddedRemovedProperties.Find(Owner->GetPresetId()))
+		{
+			const int32 Index = CachedProperties->Key.IndexOfByKey(EntityId);
+			if (Index != INDEX_NONE && CachedProperties->Value.IsValidIndex(Index))
+			{
+				Entries.Value.AddUnique(CachedProperties->Value[Index]);
+				bLabelSet = true;
+			}
+		}
+	}
+
+	// If the label couldn't be set from the Entity or the Cache then we use the EntityId as the Label since it is unique
+	if (!bLabelSet)
+	{
+		Entries.Value.AddUnique(FName(EntityId.ToString()));
+	}
 }
 
 void FWebSocketMessageHandler::OnFieldRenamed(URemoteControlPreset* Owner, FName OldFieldLabel, FName NewFieldLabel)
@@ -1173,6 +1212,16 @@ void FWebSocketMessageHandler::OnFieldRenamed(URemoteControlPreset* Owner, FName
 		return;
 	}
 
+	// Update the cached Name with the new name
+	if (TPair<TArray<FGuid>, TArray<FName>>* CachedProperties = CacheUndoRedoAddedRemovedProperties.Find(Owner->GetPresetId()))
+	{
+		const int32 Index = CachedProperties->Value.IndexOfByKey(OldFieldLabel);
+		if (Index != INDEX_NONE)
+		{
+			CachedProperties->Value[Index] = NewFieldLabel;
+		}
+	}
+
 	//Cache the field that was renamed for end of frame notification
 	PerFrameRenamedFields.FindOrAdd(Owner->GetPresetId()).AddUnique(TTuple<FName, FName>(OldFieldLabel, NewFieldLabel));
 }
diff --git a/Engine/Plugins/VirtualProduction/RemoteControl/Source/WebRemoteControl/Private/WebSocketMessageHandler.h b/Engine/Plugins/VirtualProduction/RemoteControl/Source/WebRemoteControl/Private/WebSocketMessageHandler.h
--- a/Engine/Plugins/VirtualProduction/RemoteControl/Source/WebRemoteControl/Private/WebSocketMessageHandler.h
+++ b/Engine/Plugins/VirtualProduction/RemoteControl/Source/WebRemoteControl/Private/WebSocketMessageHandler.h
@@ -384,6 +384,9 @@ private:
 	/** Properties that were exposed for a frame, per preset */
 	TMap<FGuid, TArray<FGuid>> PerFrameAddedProperties;
 
+	/** Cache used during Undo/Redo for the Remove to get correctly the Label */
+	TMap<FGuid, TTuple<TArray<FGuid>, TArray<FName>>> CacheUndoRedoAddedRemovedProperties;
+
 	/** Properties that were unexposed for a frame, per preset */
 	TMap<FGuid, TTuple<TArray<FGuid>, TArray<FName>>> PerFrameRemovedProperties;
 
diff --git a/Engine/Plugins/VirtualProduction/RemoteControlWebInterface/WebApp/Server/src/UnrealEngine.ts b/Engine/Plugins/VirtualProduction/RemoteControlWebInterface/WebApp/Server/src/UnrealEngine.ts
--- a/Engine/Plugins/VirtualProduction/RemoteControlWebInterface/WebApp/Server/src/UnrealEngine.ts
+++ b/Engine/Plugins/VirtualProduction/RemoteControlWebInterface/WebApp/Server/src/UnrealEngine.ts
@@ -293,6 +293,7 @@ export namespace UnrealEngine {
         }
 
         case UnrealApi.PresetEvent.FieldsRemoved:
+        case UnrealApi.PresetEvent.FieldsRenamed:
         case UnrealApi.PresetEvent.EntitiesModified: {
           await refreshPreset(message.PresetId, message.PresetName);
           break;
diff --git a/Engine/Plugins/VirtualProduction/RemoteControlWebInterface/WebApp/_package.json b/Engine/Plugins/VirtualProduction/RemoteControlWebInterface/WebApp/_package.json
--- a/Engine/Plugins/VirtualProduction/RemoteControlWebInterface/WebApp/_package.json
+++ b/Engine/Plugins/VirtualProduction/RemoteControlWebInterface/WebApp/_package.json
@@ -1,6 +1,6 @@
 {
   "name": "conductor",
-  "version": "0.3.91",
+  "version": "0.3.92",
   "scripts": {
     "postinstall": "concurrently \"cd ./Client && npm install && cd ..\" \"cd ./Server && npm install && cd ..\"",
     "build": "node ./Scripts/build.js",
