[RemoteControl]
Fix LogError regarding the property not being initialized correctly.

diff --git a/Engine/Plugins/VirtualProduction/RemoteControl/Source/RemoteControl/Public/RemoteControlPropertyIdRegistry.h b/Engine/Plugins/VirtualProduction/RemoteControl/Source/RemoteControl/Public/RemoteControlPropertyIdRegistry.h
--- a/Engine/Plugins/VirtualProduction/RemoteControl/Source/RemoteControl/Public/RemoteControlPropertyIdRegistry.h
+++ b/Engine/Plugins/VirtualProduction/RemoteControl/Source/RemoteControl/Public/RemoteControlPropertyIdRegistry.h
@@ -100,7 +100,7 @@ private:
 
 	/** Holds the type to be created. */
 	UPROPERTY()
-	UClass* ClassToCreate;
+	UClass* ClassToCreate = nullptr;
 };
 
 template<> struct TStructOpsTypeTraits<FRCPropertyIdWrapper> : public TStructOpsTypeTraitsBase2<FRCPropertyIdWrapper>
