[RemoteControl]
Add Getter path when Reading Property.

diff --git a/Engine/Plugins/VirtualProduction/RemoteControl/Source/RemoteControl/Private/RemoteControlModule.cpp b/Engine/Plugins/VirtualProduction/RemoteControl/Source/RemoteControl/Private/RemoteControlModule.cpp
--- a/Engine/Plugins/VirtualProduction/RemoteControl/Source/RemoteControl/Private/RemoteControlModule.cpp
+++ b/Engine/Plugins/VirtualProduction/RemoteControl/Source/RemoteControl/Private/RemoteControlModule.cpp
@@ -213,6 +213,12 @@ namespace RemoteControlUtil
 		return Access == ERCAccess::WRITE_ACCESS || Access == ERCAccess::WRITE_TRANSACTION_ACCESS || Access == ERCAccess::WRITE_MANUAL_TRANSACTION_ACCESS;
 	}
 
+	/** Returns whether the access is a read access. */
+	bool IsReadAccess(ERCAccess Access)
+	{
+		return Access == ERCAccess::READ_ACCESS;
+	}
+
 	/** Helper function to check if a value is 0 without always calling FMath::IsNearlyZero, since that results in an ambiguous call for non-float values. */
 	template <typename T>
 	bool IsNearlyZero(T Value)
@@ -1161,7 +1167,8 @@ bool FRemoteControlModule::ResolveObjectProperty(ERCAccess AccessType, UObject*
 
 					// When resolving a property for writing, resolve successfully if it should use a setter since it will end up using it. 
 					if ((RemoteControlUtil::IsWriteAccess(AccessType) && PropertyModificationShouldUseSetter(Object, ResolvedProperty))
-						|| RemoteControlUtil::IsPropertyAllowed(ResolvedProperty, AccessType, bObjectInGame))
+						|| RemoteControlUtil::IsPropertyAllowed(ResolvedProperty, AccessType, bObjectInGame)
+						|| (RemoteControlUtil::IsReadAccess(AccessType) && ResolvedProperty->HasGetter()))
 					{
 						OutObjectRef = FRCObjectReference{AccessType, Object, MoveTemp(PropertyPath)};
 					}
@@ -1236,7 +1243,9 @@ bool FRemoteControlModule::GetObjectProperties(const FRCObjectReference& ObjectA
 			bool bObjectInGame = !GIsEditor || Object->GetOutermost()->HasAnyPackageFlags(PKG_PlayInEditor);
 			Policies.PropertyFilter = [&ObjectAccess, bObjectInGame](const FProperty* CurrentProp, const FProperty* ParentProp)
 			{
-				return RemoteControlUtil::IsPropertyAllowed(CurrentProp, ObjectAccess.Access, bObjectInGame) || ParentProp != nullptr;
+				return	RemoteControlUtil::IsPropertyAllowed(CurrentProp, ObjectAccess.Access, bObjectInGame) ||
+						ParentProp != nullptr ||
+						(RemoteControlUtil::IsReadAccess(ObjectAccess.Access) && CurrentProp->HasGetter());
 			};
 		}
 
diff --git a/Engine/Plugins/VirtualProduction/RemoteControl/Source/WebRemoteControl/Private/RemoteControlReflectionUtils.cpp b/Engine/Plugins/VirtualProduction/RemoteControl/Source/WebRemoteControl/Private/RemoteControlReflectionUtils.cpp
--- a/Engine/Plugins/VirtualProduction/RemoteControl/Source/WebRemoteControl/Private/RemoteControlReflectionUtils.cpp
+++ b/Engine/Plugins/VirtualProduction/RemoteControl/Source/WebRemoteControl/Private/RemoteControlReflectionUtils.cpp
@@ -79,14 +79,34 @@ void UE::WebRCReflectionUtils::SetStringPropertyValue(FName PropertyName, const
 void UE::WebRCReflectionUtils::CopyPropertyValue(FName PropertyName, const FStructOnScope& TargetStruct, const FRCObjectReference& SourceObject)
 {
 	FProperty* TargetValueProp = TargetStruct.GetStruct()->FindPropertyByName(PropertyName);
-	FProperty* SourceValueProp = SourceObject.Property.Get();
+	const FProperty* SourceValueProp = SourceObject.Property.Get();
 
 	if (TargetValueProp && SourceValueProp && TargetValueProp->GetClass() == SourceValueProp->GetClass() && TargetValueProp->GetSize() == SourceValueProp->GetSize())
 	{
 		void* DestPtr = TargetValueProp->ContainerPtrToValuePtr<void>((void*)TargetStruct.GetStructMemory());
-		void* SrcPtr = SourceObject.Property->ContainerPtrToValuePtr<void>(SourceObject.ContainerAdress);
+		if (!DestPtr)
+		{
+			checkNoEntry();
+			return;
+		}
+
+		if (SourceValueProp->HasGetter() && SourceObject.Object.IsValid())
+		{
+			SourceObject.Property->PerformOperationWithGetter(SourceObject.Object.Get(), nullptr,
+				[TargetValueProp, DestPtr](const void* InValuePtr)
+				{
+					if (InValuePtr)
+					{
+						TargetValueProp->CopyCompleteValue(DestPtr, InValuePtr);
+						return;
+					}
+					checkNoEntry();
+				});
+			return;
+		}
 
-		if (DestPtr && SrcPtr)
+		const void* SrcPtr = SourceObject.Property->ContainerPtrToValuePtr<void>(SourceObject.ContainerAdress);
+		if (SrcPtr)
 		{
 			TargetValueProp->CopyCompleteValue(DestPtr, SrcPtr);
 			return;
