[RemoteControl]
Fix ensure being hit when trying to open the WebApp through the RCPanel button of an embedded preset.

diff --git a/Engine/Plugins/VirtualProduction/RemoteControl/Source/RemoteControlUI/Private/RemoteControlUIModule.cpp b/Engine/Plugins/VirtualProduction/RemoteControl/Source/RemoteControlUI/Private/RemoteControlUIModule.cpp
--- a/Engine/Plugins/VirtualProduction/RemoteControl/Source/RemoteControlUI/Private/RemoteControlUIModule.cpp
+++ b/Engine/Plugins/VirtualProduction/RemoteControl/Source/RemoteControlUI/Private/RemoteControlUIModule.cpp
@@ -313,31 +313,23 @@ TSharedRef<SRemoteControlPanel> FRemoteControlUIModule::CreateRemoteControlPanel
 
 	TSharedRef<SRemoteControlPanel> PanelRef = SNew(SRemoteControlPanel, Preset, ToolkitHost)
 		.OnLiveModeChange_Lambda(
-			[this](TSharedPtr<SRemoteControlPanel> Panel, bool bLiveMode)
+			[this](TSharedPtr<SRemoteControlPanel> InPanel, bool bInLiveMode)
 			{
-				URemoteControlPreset* Preset = Panel->GetPreset();
-
-				if (!IsValid(Preset) || Preset->IsEmbeddedPreset() == false)
+				// Activating the live mode on a panel sets it as the active panel
+				if (bInLiveMode)
 				{
-					// Activating the live mode on a panel sets it as the active panel
-					if (bLiveMode)
+					if (const TSharedPtr<SRemoteControlPanel> ActivePanel = WeakActivePanel.Pin())
 					{
-						if (TSharedPtr<SRemoteControlPanel> ActivePanel = WeakActivePanel.Pin())
+						if (ActivePanel != InPanel)
 						{
-							if (ActivePanel != Panel)
-							{
-								ActivePanel->SetLiveMode(true);
-							}
+							ActivePanel->SetLiveMode(true);
 						}
-						WeakActivePanel = MoveTemp(Panel);
 					}
+					WeakActivePanel = MoveTemp(InPanel);
 				}
 			});
 
-	if (Preset->IsEmbeddedPreset() == false)
-	{
-		WeakActivePanel = PanelRef;
-	}
+	WeakActivePanel = PanelRef;
 
 	RegisteredRemoteControlPanels.Add(TWeakPtr<SRemoteControlPanel>(PanelRef));
 
