[Remote Control] Fix FRemoteControlModule::UnregisterEmbeddedPreset to rely primarily on the PresetId rather than GetPresetName() given that this is unreliable during GC. If the outer has it's BeginDestroy called prior to the RCP being unregistered, GetPresetName() will return "None" and it will not be unregistered leading to a stale entry in EmbeddedPresets and CachedPresetNamesById, which are going to cause an ensure in sometime later in the ResolvePreset when registering another instance of the same preset.

#rnx

diff --git a/Engine/Plugins/VirtualProduction/RemoteControl/Source/RemoteControl/Private/RemoteControlModule.cpp b/Engine/Plugins/VirtualProduction/RemoteControl/Source/RemoteControl/Private/RemoteControlModule.cpp
--- a/Engine/Plugins/VirtualProduction/RemoteControl/Source/RemoteControl/Private/RemoteControlModule.cpp
+++ b/Engine/Plugins/VirtualProduction/RemoteControl/Source/RemoteControl/Private/RemoteControlModule.cpp
@@ -755,6 +755,7 @@ void FRemoteControlModule::UnregisterEmbeddedPreset(FName Name)
 	// Check the cached preset ids and remove it if our name matches the stored id
 	TWeakObjectPtr<URemoteControlPreset>* FoundPreset = EmbeddedPresets.Find(Name);
 
+	// Remark: the weak ptr will be invalid during GC and the object will not be unregistered.
 	if (FoundPreset && FoundPreset->IsValid())
 	{
 		UnregisterEmbeddedPreset(FoundPreset->Get());
@@ -768,26 +769,34 @@ void FRemoteControlModule::UnregisterEmbeddedPreset(URemoteControlPreset* Preset
 		return;
 	}
 
+	// Remark: the PresetName is unreliable. When unregistering during the garbage collection,
+	// the name can be "None". In that case we can recover the name from CachedPresetNamesById
+	// by using the PresetId.
 	FName PresetName = Preset->GetPresetName();
 
-	if (PresetName == NAME_None)
-	{
-		return;
-	}
-
-	FGuid PresetId = Preset->GetPresetId();
-
+	const FGuid PresetId = Preset->GetPresetId();
+	
 	if (PresetId.IsValid())
 	{
-		FName* FoundPresetName = CachedPresetNamesById.Find(PresetId);
-
-		if (FoundPresetName && *FoundPresetName == PresetName)
+		if (const FName* FoundPresetName = CachedPresetNamesById.Find(PresetId))
 		{
-			CachedPresetNamesById.Remove(PresetId);
+			// Recover the name from the id.
+			if (PresetName == NAME_None)
+			{
+				PresetName = *FoundPresetName;
+			}
+			
+			if (PresetName == *FoundPresetName)
+			{
+				CachedPresetNamesById.Remove(PresetId);
+			}
 		}
 	}
 
-	EmbeddedPresets.Remove(PresetName);
+	if (PresetName != NAME_None)
+	{
+		EmbeddedPresets.Remove(PresetName);
+	}
 }
 
 bool FRemoteControlModule::CanResetToDefaultValue(UObject* InObject, const FRCResetToDefaultArgs& InArgs) const
