[RemoteControl]
The number of sides for a circle cannot be set using remote control under the Property ID system.

diff --git a/Engine/Plugins/VirtualProduction/RemoteControl/Source/RemoteControl/Private/PropertyIdHandler/BasePropertyIdHandler.cpp b/Engine/Plugins/VirtualProduction/RemoteControl/Source/RemoteControl/Private/PropertyIdHandler/BasePropertyIdHandler.cpp
--- a/Engine/Plugins/VirtualProduction/RemoteControl/Source/RemoteControl/Private/PropertyIdHandler/BasePropertyIdHandler.cpp
+++ b/Engine/Plugins/VirtualProduction/RemoteControl/Source/RemoteControl/Private/PropertyIdHandler/BasePropertyIdHandler.cpp
@@ -33,10 +33,10 @@ EPropertyBagPropertyType FBasePropertyIdHandler::GetPropertyType(const FProperty
 				{
 					PropertyType = EPropertyBagPropertyType::Enum;
 				}
-			}
-			else
-			{
-				PropertyType = EPropertyBagPropertyType::Byte;
+				else
+				{
+					PropertyType = EPropertyBagPropertyType::Byte;
+				}
 			}
 		}
 		else if (PropertyName == NAME_Int32Property)
