[RemoteControl]
Tooltip add "Double Click to Select in Outliner".

diff --git a/Engine/Plugins/VirtualProduction/RemoteControl/Source/RemoteControlUI/Private/UI/SRCPanelExposedEntity.cpp b/Engine/Plugins/VirtualProduction/RemoteControl/Source/RemoteControlUI/Private/UI/SRCPanelExposedEntity.cpp
--- a/Engine/Plugins/VirtualProduction/RemoteControl/Source/RemoteControlUI/Private/UI/SRCPanelExposedEntity.cpp
+++ b/Engine/Plugins/VirtualProduction/RemoteControl/Source/RemoteControlUI/Private/UI/SRCPanelExposedEntity.cpp
@@ -70,6 +70,8 @@ namespace RebindingUtils
 	}
 }
 
+const FText SRCPanelExposedEntity::SelectInOutliner = LOCTEXT("RCSelectInOutliner", "\nDouble Click to Select in Outliner");
+
 TSharedPtr<FRemoteControlEntity> SRCPanelExposedEntity::GetEntity() const
 {
 	if (Preset.IsValid())
@@ -532,7 +534,7 @@ TSharedRef<SWidget> SRCPanelExposedEntity::CreateEntityWidget(TSharedPtr<SWidget
 				SNew(STextBlock)
 				.ColorAndOpacity_Lambda([this]() { return bValidBinding ? FSlateColor::UseForeground() : FSlateColor::UseSubduedForeground(); })
 				.Text(FText::FromName(CachedOwnerName))
-				.ToolTipText(FText::FromName(CachedBindingPath))
+				.ToolTipText(FText::FromString(CachedBindingPath.ToString() + SelectInOutliner.ToString()))
 			]
 		];
 
@@ -542,7 +544,7 @@ TSharedRef<SWidget> SRCPanelExposedEntity::CreateEntityWidget(TSharedPtr<SWidget
 			SNew(STextBlock)
 			.ColorAndOpacity_Lambda([this]() { return bValidBinding ? FSlateColor::UseForeground() : FSlateColor::UseSubduedForeground(); })
 			.Text(CachedSubobjectPath.IsNone() ? FText::GetEmpty() : FText::FromName(CachedSubobjectPath))
-			.ToolTipText(LOCTEXT("SubobjectPathToolTip", "The path from the owner actor to the uobject holding the exposed property."))
+			.ToolTipText(FText::Format(LOCTEXT("SubobjectPathToolTip", "The path from the owner actor to the uobject holding the exposed property.{0}"), SelectInOutliner))
 		];
 
 	Args.NameWidget = SNew(SHorizontalBox)
@@ -570,7 +572,7 @@ TSharedRef<SWidget> SRCPanelExposedEntity::CreateEntityWidget(TSharedPtr<SWidget
 		[
 			SAssignNew(NameTextBox, SInlineEditableTextBlock)
 			.Text_Lambda([this] () { return FText::FromName(CachedLabel); })
-			.ToolTipText(FText::FromString(CachedFieldPath))
+			.ToolTipText(FText::FromString(CachedFieldPath + SelectInOutliner.ToString()))
 			.OnTextCommitted(this, &SRCPanelExposedEntity::OnLabelCommitted)
 			.OnVerifyTextChanged(this, &SRCPanelExposedEntity::OnVerifyItemLabelChanged)
 			.IsReadOnly_Lambda([this]() { return bLiveMode.Get(); })
diff --git a/Engine/Plugins/VirtualProduction/RemoteControl/Source/RemoteControlUI/Private/UI/SRCPanelExposedEntity.h b/Engine/Plugins/VirtualProduction/RemoteControl/Source/RemoteControlUI/Private/UI/SRCPanelExposedEntity.h
--- a/Engine/Plugins/VirtualProduction/RemoteControl/Source/RemoteControlUI/Private/UI/SRCPanelExposedEntity.h
+++ b/Engine/Plugins/VirtualProduction/RemoteControl/Source/RemoteControlUI/Private/UI/SRCPanelExposedEntity.h
@@ -75,7 +75,8 @@ protected:
 	TAttribute<FText> HighlightText;
 	/** Panel Style reference. */
 	const FRCPanelStyle* RCPanelStyle;
-
+	/** Extra information to add to the tooltips. */
+	static const FText SelectInOutliner;
 private:
 	/** Handles changing the object this entity is bound to upon selecting an actor in the rebinding dropdown. */
 	void OnActorSelected(AActor* InActor) const;
