[RemoteControl]
Fix crash with PropertyId when ArrayIndex is invalid.
Avoid creating Object for Abstract class.

diff --git a/Engine/Plugins/VirtualProduction/RemoteControl/Source/RemoteControl/Private/PropertyIdHandler/ObjectPropertyIdHandler.cpp b/Engine/Plugins/VirtualProduction/RemoteControl/Source/RemoteControl/Private/PropertyIdHandler/ObjectPropertyIdHandler.cpp
--- a/Engine/Plugins/VirtualProduction/RemoteControl/Source/RemoteControl/Private/PropertyIdHandler/ObjectPropertyIdHandler.cpp
+++ b/Engine/Plugins/VirtualProduction/RemoteControl/Source/RemoteControl/Private/PropertyIdHandler/ObjectPropertyIdHandler.cpp
@@ -44,10 +44,17 @@ TObjectPtr<UObject> FObjectPropertyIdHandler::GetObjectPropertyDefaultValue(cons
 {
 	if (InClassToCreate)
 	{
+		if (InClassToCreate->HasAnyClassFlags(CLASS_Abstract))
+		{
+			// for now skip any AbstractClass
+			return nullptr;
+		}
+
 		if (InClassToCreate->IsChildOf(UMaterialInterface::StaticClass()))
 		{
 			return UMaterial::GetDefaultMaterial(MD_Surface);
 		}
+
 		return NewObject<UObject>(GetTransientPackage(), InClassToCreate);
 	}
 	return nullptr;
diff --git a/Engine/Plugins/VirtualProduction/RemoteControl/Source/RemoteControl/Private/RemoteControlPropertyIdRegistry.cpp b/Engine/Plugins/VirtualProduction/RemoteControl/Source/RemoteControl/Private/RemoteControlPropertyIdRegistry.cpp
--- a/Engine/Plugins/VirtualProduction/RemoteControl/Source/RemoteControl/Private/RemoteControlPropertyIdRegistry.cpp
+++ b/Engine/Plugins/VirtualProduction/RemoteControl/Source/RemoteControl/Private/RemoteControlPropertyIdRegistry.cpp
@@ -286,19 +286,22 @@ void FRCPropertyIdWrapper::UpdateTypes(const TSharedRef<FRemoteControlProperty>&
 					
 					if (const TSharedRef<FRemoteControlField> TargetRCField = StaticCastSharedRef<FRemoteControlField>(InRCProperty); !TargetRCField->FieldPathInfo.Segments.IsEmpty())
 					{
-						const int32 MaterialIndex = TargetRCField->FieldPathInfo.Segments[0].ArrayIndex;
-						const uint8* ObjPtrContainer = ArrayHelper.GetRawPtr(MaterialIndex);
-						if (const UObject* CurrentObject = InnerObjectProperty->GetObjectPropertyValue(ObjPtrContainer))
+						const int32 ArrayIndex = TargetRCField->FieldPathInfo.Segments[0].ArrayIndex;
+						if (ArrayIndex != INDEX_NONE)
 						{
-							if (CurrentObject->IsA(UMaterialInterface::StaticClass()))
+							const uint8* ObjPtrContainer = ArrayHelper.GetRawPtr(ArrayIndex);
+							if (const UObject* CurrentObject = InnerObjectProperty->GetObjectPropertyValue(ObjPtrContainer))
 							{
-								SubType = UMaterialInterface::StaticClass()->GetFName();
-								ClassToCreate = UMaterialInterface::StaticClass();
-							}
-							else
-							{
-								SubType = CurrentObject->GetClass()->GetFName();
-								ClassToCreate = CurrentObject->GetClass();
+								if (CurrentObject->IsA(UMaterialInterface::StaticClass()))
+								{
+									SubType = UMaterialInterface::StaticClass()->GetFName();
+									ClassToCreate = UMaterialInterface::StaticClass();
+								}
+								else
+								{
+									SubType = CurrentObject->GetClass()->GetFName();
+									ClassToCreate = CurrentObject->GetClass();
+								}
 							}
 						}
 					}
