Relocate tests under "Plugin." to "Plugins.".

diff --git a/Engine/Plugins/Experimental/Web/WebAPI/Source/WebAPIEditor/Private/Tests/WebAPIEditorUtilities.spec.cpp b/Engine/Plugins/Experimental/Web/WebAPI/Source/WebAPIEditor/Private/Tests/WebAPIEditorUtilities.spec.cpp
--- a/Engine/Plugins/Experimental/Web/WebAPI/Source/WebAPIEditor/Private/Tests/WebAPIEditorUtilities.spec.cpp
+++ b/Engine/Plugins/Experimental/Web/WebAPI/Source/WebAPIEditor/Private/Tests/WebAPIEditorUtilities.spec.cpp
@@ -1,4 +1,4 @@
-﻿// Copyright Epic Games, Inc. All Rights Reserved.
+// Copyright Epic Games, Inc. All Rights Reserved.
 
 #include "Misc/AutomationTest.h"
 
@@ -7,7 +7,7 @@
 #include "WebAPIEditorUtilities.h"
 
 BEGIN_DEFINE_SPEC(FWebAPIEditorUtilitiesSpec,
-	TEXT("Plugin.WebAPI.Editor.Utilities"),
+	TEXT("Plugins.WebAPI.Editor.Utilities"),
 	EAutomationTestFlags::EditorContext | EAutomationTestFlags::EngineFilter | EAutomationTestFlags::ApplicationContextMask)
 
 END_DEFINE_SPEC(FWebAPIEditorUtilitiesSpec)
diff --git a/Engine/Plugins/Experimental/Web/WebAPI/Source/WebAPIEditor/Private/Tests/WebAPIJsonUtilities.spec.cpp b/Engine/Plugins/Experimental/Web/WebAPI/Source/WebAPIEditor/Private/Tests/WebAPIJsonUtilities.spec.cpp
--- a/Engine/Plugins/Experimental/Web/WebAPI/Source/WebAPIEditor/Private/Tests/WebAPIJsonUtilities.spec.cpp
+++ b/Engine/Plugins/Experimental/Web/WebAPI/Source/WebAPIEditor/Private/Tests/WebAPIJsonUtilities.spec.cpp
@@ -8,7 +8,7 @@
 #include "Tests/WebAPIJsonTestData.h"
 
 BEGIN_DEFINE_SPEC(FWebAPIEditorJsonSpec,
-	TEXT("Plugin.WebAPI.Editor.Json"),
+	TEXT("Plugins.WebAPI.Editor.Json"),
 	EAutomationTestFlags::EditorContext | EAutomationTestFlags::EngineFilter | EAutomationTestFlags::ApplicationContextMask)
 
 	TSharedPtr<FJsonObject> JsonObject;
diff --git a/Engine/Plugins/Experimental/Web/WebAPI/Source/WebAPIOpenAPI/Private/Tests/WebAPIOpenAPI3.spec.cpp b/Engine/Plugins/Experimental/Web/WebAPI/Source/WebAPIOpenAPI/Private/Tests/WebAPIOpenAPI3.spec.cpp
--- a/Engine/Plugins/Experimental/Web/WebAPI/Source/WebAPIOpenAPI/Private/Tests/WebAPIOpenAPI3.spec.cpp
+++ b/Engine/Plugins/Experimental/Web/WebAPI/Source/WebAPIOpenAPI/Private/Tests/WebAPIOpenAPI3.spec.cpp
@@ -18,7 +18,7 @@
 #if WITH_DEV_AUTOMATION_TESTS
 
 BEGIN_DEFINE_SPEC(FWebAPIOpenAPI3Spec,
-				"Plugin.WebAPI.OpenAPI3",
+				"Plugins.WebAPI.OpenAPI3",
 				EAutomationTestFlags::EditorContext | EAutomationTestFlags::EngineFilter | EAutomationTestFlags::ApplicationContextMask)
 
 	TSharedPtr<UE::WebAPI::OpenAPI::V3::FOpenAPIObject> InputDefinition;
diff --git a/Engine/Plugins/Experimental/Web/WebAPI/Source/WebAPIOpenAPI/Private/Tests/WebAPISwagger.spec.cpp b/Engine/Plugins/Experimental/Web/WebAPI/Source/WebAPIOpenAPI/Private/Tests/WebAPISwagger.spec.cpp
--- a/Engine/Plugins/Experimental/Web/WebAPI/Source/WebAPIOpenAPI/Private/Tests/WebAPISwagger.spec.cpp
+++ b/Engine/Plugins/Experimental/Web/WebAPI/Source/WebAPIOpenAPI/Private/Tests/WebAPISwagger.spec.cpp
@@ -14,7 +14,7 @@
 #if WITH_DEV_AUTOMATION_TESTS
 
 BEGIN_DEFINE_SPEC(FWebAPISwaggerSpec,
-	TEXT("Plugin.WebAPI.Swagger"),
+	TEXT("Plugins.WebAPI.Swagger"),
 	EAutomationTestFlags::EditorContext | EAutomationTestFlags::EngineFilter | EAutomationTestFlags::ApplicationContextMask)
 
 	TSharedPtr<UE::WebAPI::OpenAPI::V2::FSwagger> InputDefinition;
diff --git a/Engine/Plugins/Media/MediaIOFramework/Source/MediaIOCore/Private/Tests/MediaIODeinterlacerTests.cpp b/Engine/Plugins/Media/MediaIOFramework/Source/MediaIOCore/Private/Tests/MediaIODeinterlacerTests.cpp
--- a/Engine/Plugins/Media/MediaIOFramework/Source/MediaIOCore/Private/Tests/MediaIODeinterlacerTests.cpp
+++ b/Engine/Plugins/Media/MediaIOFramework/Source/MediaIOCore/Private/Tests/MediaIODeinterlacerTests.cpp
@@ -187,14 +187,26 @@ namespace UE::MediaIOTests
 	}
 }
 
-
-IMPLEMENT_SIMPLE_AUTOMATION_TEST(FMediaIODeinterlacerTests, "Plugin.MediaIO.Deinterlace", EAutomationTestFlags::EditorContext | EAutomationTestFlags::EngineFilter)
-bool FMediaIODeinterlacerTests::RunTest(const FString& Parameters)
+DEFINE_SPEC(FMediaIODeinterlacerTests, "Plugins.MediaIO.Deinterlace", EAutomationTestFlags::EditorContext | EAutomationTestFlags::EngineFilter)
+void FMediaIODeinterlacerTests::Define()
 {
-	// Test expose whole container
-	UE::MediaIOTests::TestBobDeinterlace(*this);
-	UE::MediaIOTests::TestBlendDeinterlace(*this);
-	UE::MediaIOTests::TestDiscardDeinterlace(*this);
-	UE::MediaIOTests::TestDiscardDeinterlaceBottomFieldOrder(*this);
-	return true;
+	It("BobDeinterlace", [this]()
+		{
+			UE::MediaIOTests::TestBobDeinterlace(*this);
+		});
+
+	It("BlendDeinterlace", [this]()
+		{
+			UE::MediaIOTests::TestBlendDeinterlace(*this);
+		});
+
+	It("DiscardDeinterlace", [this]()
+		{
+			UE::MediaIOTests::TestDiscardDeinterlace(*this);
+		});
+
+	It("DiscardDeinterlaceBottomFieldOrder", [this]()
+		{
+			UE::MediaIOTests::TestDiscardDeinterlaceBottomFieldOrder(*this);
+		});
 }
\ No newline at end of file
diff --git a/Engine/Plugins/VirtualProduction/CameraCalibration/Source/CameraCalibrationEditor/Private/TestCameraCalibration.cpp b/Engine/Plugins/VirtualProduction/CameraCalibration/Source/CameraCalibrationEditor/Private/TestCameraCalibration.cpp
--- a/Engine/Plugins/VirtualProduction/CameraCalibration/Source/CameraCalibrationEditor/Private/TestCameraCalibration.cpp
+++ b/Engine/Plugins/VirtualProduction/CameraCalibration/Source/CameraCalibrationEditor/Private/TestCameraCalibration.cpp
@@ -18,8 +18,8 @@
 #include "TestCameraCalibrationSettings.h"
 #include "Widgets/Images/SImage.h"
 
-IMPLEMENT_SIMPLE_AUTOMATION_TEST(FTestDistortionSpherical, "Plugin.CameraCalibration.TestDistortionSpherical", EAutomationTestFlags::EditorContext | EAutomationTestFlags::EngineFilter)
-IMPLEMENT_SIMPLE_AUTOMATION_TEST(FTestNodalOffset, "Plugin.CameraCalibration.TestNodalOffset", EAutomationTestFlags::EditorContext | EAutomationTestFlags::EngineFilter)
+IMPLEMENT_SIMPLE_AUTOMATION_TEST(FTestDistortionSpherical, "Plugins.CameraCalibration.TestDistortionSpherical", EAutomationTestFlags::EditorContext | EAutomationTestFlags::EngineFilter)
+IMPLEMENT_SIMPLE_AUTOMATION_TEST(FTestNodalOffset, "Plugins.CameraCalibration.TestNodalOffset", EAutomationTestFlags::EditorContext | EAutomationTestFlags::EngineFilter)
 
 namespace UE::Private::CameraCalibration::AutomatedTests
 {
diff --git a/Engine/Plugins/VirtualProduction/CameraCalibrationCore/Source/CameraCalibrationCore/Private/Tests/TestDistortionMapping.cpp b/Engine/Plugins/VirtualProduction/CameraCalibrationCore/Source/CameraCalibrationCore/Private/Tests/TestDistortionMapping.cpp
--- a/Engine/Plugins/VirtualProduction/CameraCalibrationCore/Source/CameraCalibrationCore/Private/Tests/TestDistortionMapping.cpp
+++ b/Engine/Plugins/VirtualProduction/CameraCalibrationCore/Source/CameraCalibrationCore/Private/Tests/TestDistortionMapping.cpp
@@ -14,9 +14,6 @@
 #include "UObject/Package.h"
 #include "UObject/StrongObjectPtr.h"
 
-IMPLEMENT_SIMPLE_AUTOMATION_TEST(FTestCameraCalibrationCore, "Plugin.CameraCalibrationCore", EAutomationTestFlags::EditorContext | EAutomationTestFlags::EngineFilter)
-
-
 namespace CameraCalibrationTestUtil
 {
 	UWorld* GetFirstWorld()
@@ -426,13 +423,23 @@ namespace CameraCalibrationTestUtil
 	}
 }
 
-
-bool FTestCameraCalibrationCore::RunTest(const FString& Parameters)
+DEFINE_SPEC(FTestCameraCalibrationCore, "Plugins.CameraCalibrationCore", EAutomationTestFlags::EditorContext | EAutomationTestFlags::EngineFilter)
+void FTestCameraCalibrationCore::Define()
 {
-	CameraCalibrationTestUtil::TestDistortionParameterCurveBlending(*this);
-	CameraCalibrationTestUtil::TestLensFileAddPoints(*this);
-	CameraCalibrationTestUtil::TestLensFileRemovePoints(*this);
-	return true;
+	It("DistortionParameterCurveBlending", [this]()
+		{
+			CameraCalibrationTestUtil::TestDistortionParameterCurveBlending(*this);
+		});
+
+	It("LensFileAddPoints", [this]()
+		{
+			CameraCalibrationTestUtil::TestLensFileAddPoints(*this);
+		});
+
+	It("LensFileRemovePoints", [this]()
+		{
+			CameraCalibrationTestUtil::TestLensFileRemovePoints(*this);
+		});
 }
 
 #endif // WITH_DEV_AUTOMATION_TESTS
diff --git a/Engine/Plugins/VirtualProduction/CaptureManager/Source/CaptureSourceFramework/Private/Tests/TestCaptureSource.cpp b/Engine/Plugins/VirtualProduction/CaptureManager/Source/CaptureSourceFramework/Private/Tests/TestCaptureSource.cpp
--- a/Engine/Plugins/VirtualProduction/CaptureManager/Source/CaptureSourceFramework/Private/Tests/TestCaptureSource.cpp
+++ b/Engine/Plugins/VirtualProduction/CaptureManager/Source/CaptureSourceFramework/Private/Tests/TestCaptureSource.cpp
@@ -6,7 +6,7 @@
 
 #if WITH_DEV_AUTOMATION_TESTS
 
-BEGIN_DEFINE_SPEC(FCaptureSourceTest, "Plugin.CaptureSourceFramework.CaptureSource", EAutomationTestFlags::EditorContext | EAutomationTestFlags::ProductFilter | EAutomationTestFlags::MediumPriority)
+BEGIN_DEFINE_SPEC(FCaptureSourceTest, "Plugins.CaptureSourceFramework.CaptureSource", EAutomationTestFlags::EditorContext | EAutomationTestFlags::ProductFilter | EAutomationTestFlags::MediumPriority)
 
 static const FString TestCaptureSourceFactoryId;
 static const FString TestCaptureSourceName;
diff --git a/Engine/Plugins/VirtualProduction/CaptureManager/Source/CaptureSourceFramework/Private/Tests/TestCaptureSourceCapability.cpp b/Engine/Plugins/VirtualProduction/CaptureManager/Source/CaptureSourceFramework/Private/Tests/TestCaptureSourceCapability.cpp
--- a/Engine/Plugins/VirtualProduction/CaptureManager/Source/CaptureSourceFramework/Private/Tests/TestCaptureSourceCapability.cpp
+++ b/Engine/Plugins/VirtualProduction/CaptureManager/Source/CaptureSourceFramework/Private/Tests/TestCaptureSourceCapability.cpp
@@ -6,7 +6,7 @@
 
 #if WITH_DEV_AUTOMATION_TESTS
 
-BEGIN_DEFINE_SPEC(FCaptureSourceCapabilityTest, "Plugin.CaptureSourceFramework.CaptureSourceCapability", EAutomationTestFlags::EditorContext | EAutomationTestFlags::ProductFilter | EAutomationTestFlags::MediumPriority)
+BEGIN_DEFINE_SPEC(FCaptureSourceCapabilityTest, "Plugins.CaptureSourceFramework.CaptureSourceCapability", EAutomationTestFlags::EditorContext | EAutomationTestFlags::ProductFilter | EAutomationTestFlags::MediumPriority)
 
 class FTestCommand final : public FCommandBase
 {
diff --git a/Engine/Plugins/VirtualProduction/CaptureManager/Source/CaptureSourceFramework/Private/Tests/TestCaptureSourceFactory.cpp b/Engine/Plugins/VirtualProduction/CaptureManager/Source/CaptureSourceFramework/Private/Tests/TestCaptureSourceFactory.cpp
--- a/Engine/Plugins/VirtualProduction/CaptureManager/Source/CaptureSourceFramework/Private/Tests/TestCaptureSourceFactory.cpp
+++ b/Engine/Plugins/VirtualProduction/CaptureManager/Source/CaptureSourceFramework/Private/Tests/TestCaptureSourceFactory.cpp
@@ -6,7 +6,7 @@
 
 #if WITH_DEV_AUTOMATION_TESTS
 
-BEGIN_DEFINE_SPEC(FCaptureSourceFactoryTest, "Plugin.CaptureSourceFramework.CaptureSourceFactory", EAutomationTestFlags::EditorContext | EAutomationTestFlags::ProductFilter | EAutomationTestFlags::MediumPriority)
+BEGIN_DEFINE_SPEC(FCaptureSourceFactoryTest, "Plugins.CaptureSourceFramework.CaptureSourceFactory", EAutomationTestFlags::EditorContext | EAutomationTestFlags::ProductFilter | EAutomationTestFlags::MediumPriority)
 
 class FTestCaptureSource final : public FCaptureSource
 {
diff --git a/Engine/Plugins/VirtualProduction/CaptureManager/Source/CaptureSourceFramework/Private/Tests/TestCaptureSourceManager.cpp b/Engine/Plugins/VirtualProduction/CaptureManager/Source/CaptureSourceFramework/Private/Tests/TestCaptureSourceManager.cpp
--- a/Engine/Plugins/VirtualProduction/CaptureManager/Source/CaptureSourceFramework/Private/Tests/TestCaptureSourceManager.cpp
+++ b/Engine/Plugins/VirtualProduction/CaptureManager/Source/CaptureSourceFramework/Private/Tests/TestCaptureSourceManager.cpp
@@ -6,7 +6,7 @@
 
 #if WITH_DEV_AUTOMATION_TESTS
 
-BEGIN_DEFINE_SPEC(FCaptureSourceManagerTest, "Plugin.CaptureSourceFramework.CaptureSourceManager", EAutomationTestFlags::EditorContext | EAutomationTestFlags::ProductFilter | EAutomationTestFlags::MediumPriority)
+BEGIN_DEFINE_SPEC(FCaptureSourceManagerTest, "Plugins.CaptureSourceFramework.CaptureSourceManager", EAutomationTestFlags::EditorContext | EAutomationTestFlags::ProductFilter | EAutomationTestFlags::MediumPriority)
 
 class FTestCaptureSource final : public FCaptureSource
 {
diff --git a/Engine/Plugins/VirtualProduction/RemoteControl/Source/RemoteControl/Private/Tests/RemoteControlDeltaAPITests.spec.cpp b/Engine/Plugins/VirtualProduction/RemoteControl/Source/RemoteControl/Private/Tests/RemoteControlDeltaAPITests.spec.cpp
--- a/Engine/Plugins/VirtualProduction/RemoteControl/Source/RemoteControl/Private/Tests/RemoteControlDeltaAPITests.spec.cpp
+++ b/Engine/Plugins/VirtualProduction/RemoteControl/Source/RemoteControl/Private/Tests/RemoteControlDeltaAPITests.spec.cpp
@@ -13,7 +13,7 @@
 #include "IRemoteControlModule.h"
 #include "RemoteControlDeltaAPITestData.h"
 
-BEGIN_DEFINE_SPEC(FRemoteControlDeltaAPISpec, "Plugin.RemoteControl.DeltaAPI", EAutomationTestFlags::EditorContext | EAutomationTestFlags::EngineFilter | EAutomationTestFlags::ApplicationContextMask)
+BEGIN_DEFINE_SPEC(FRemoteControlDeltaAPISpec, "Plugins.RemoteControl.DeltaAPI", EAutomationTestFlags::EditorContext | EAutomationTestFlags::EngineFilter | EAutomationTestFlags::ApplicationContextMask)
 
 URemoteControlDeltaAPITestObject* TestObject;
 
diff --git a/Engine/Plugins/VirtualProduction/RemoteControl/Source/RemoteControl/Private/Tests/RemoteControlInterceptionTest.cpp b/Engine/Plugins/VirtualProduction/RemoteControl/Source/RemoteControl/Private/Tests/RemoteControlInterceptionTest.cpp
--- a/Engine/Plugins/VirtualProduction/RemoteControl/Source/RemoteControl/Private/Tests/RemoteControlInterceptionTest.cpp
+++ b/Engine/Plugins/VirtualProduction/RemoteControl/Source/RemoteControl/Private/Tests/RemoteControlInterceptionTest.cpp
@@ -387,7 +387,7 @@ private:
 	TUniquePtr<IRemoteControlInterceptionFeatureProcessor>   FeatureProcessor;
 };
 
-IMPLEMENT_SIMPLE_AUTOMATION_TEST(FRemoteControlPresetInterceptionTest, "Plugin.RemoteControl.Interception", EAutomationTestFlags::EditorContext | EAutomationTestFlags::EngineFilter)
+IMPLEMENT_SIMPLE_AUTOMATION_TEST(FRemoteControlPresetInterceptionTest, "Plugins.RemoteControl.Interception", EAutomationTestFlags::EditorContext | EAutomationTestFlags::EngineFilter)
 bool FRemoteControlPresetInterceptionTest::RunTest(const FString& Parameters)
 {
 	FRemoteControlInterceptionTest RemoteControlInterceptionTest(this);
diff --git a/Engine/Plugins/VirtualProduction/RemoteControl/Source/RemoteControl/Private/Tests/RemoteControlPropertyHandleTests.cpp b/Engine/Plugins/VirtualProduction/RemoteControl/Source/RemoteControl/Private/Tests/RemoteControlPropertyHandleTests.cpp
--- a/Engine/Plugins/VirtualProduction/RemoteControl/Source/RemoteControl/Private/Tests/RemoteControlPropertyHandleTests.cpp
+++ b/Engine/Plugins/VirtualProduction/RemoteControl/Source/RemoteControl/Private/Tests/RemoteControlPropertyHandleTests.cpp
@@ -1054,31 +1054,67 @@ namespace RemoteControlAPIIntegrationTest
 	}
 }
 
-IMPLEMENT_SIMPLE_AUTOMATION_TEST(FRemoteControlAPIIntegrationTest, "Plugin.RemoteControl.PropertyHandle", EAutomationTestFlags::EditorContext | EAutomationTestFlags::EngineFilter)
-bool FRemoteControlAPIIntegrationTest::RunTest(const FString& Parameters)
+DEFINE_SPEC(FRemoteControlAPIIntegrationTest, "Plugins.RemoteControl.PropertyHandle", EAutomationTestFlags::EditorContext | EAutomationTestFlags::EngineFilter)
+void FRemoteControlAPIIntegrationTest::Define()
 {
 	// Test expose whole container
-	RemoteControlAPIIntegrationTest::TestIntPropertyHandle(*this);
-	RemoteControlAPIIntegrationTest::TestCreatePropertyHandleStruct(*this);
-	RemoteControlAPIIntegrationTest::TestCreatePropertyHandleArray(*this);
-	RemoteControlAPIIntegrationTest::TestCreatePropertyHandleSet(*this);
-	RemoteControlAPIIntegrationTest::TestCreatePropertyHandleMap(*this);
-	RemoteControlAPIIntegrationTest::TestFloatProperty(*this);
-	RemoteControlAPIIntegrationTest::TestDoubleProperty(*this);
-	RemoteControlAPIIntegrationTest::TestStringProperty(*this);
-	RemoteControlAPIIntegrationTest::TestNameProperty(*this);
-	RemoteControlAPIIntegrationTest::TestBoolProperty(*this);
-	RemoteControlAPIIntegrationTest::TestByteProperty(*this);
-	RemoteControlAPIIntegrationTest::TestTextProperty(*this);
-	RemoteControlAPIIntegrationTest::TestVectorProperty(*this);
-	RemoteControlAPIIntegrationTest::TestRotatorProperty(*this);
-	RemoteControlAPIIntegrationTest::TestColorProperty(*this);
-	RemoteControlAPIIntegrationTest::TestLinearColorProperty(*this);
-	RemoteControlAPIIntegrationTest::TestComplexPath(*this);
-	RemoteControlAPIIntegrationTest::TestGetPropertyHandleByFieldPath(*this);
-	RemoteControlAPIIntegrationTest::TestStructInObjectProperty(*this);
-
-	return true;
+	It("IntPropertyHandle", [this]() {
+		RemoteControlAPIIntegrationTest::TestIntPropertyHandle(*this);
+	});
+	It("CreatePropertyHandleStruct", [this]() {
+		RemoteControlAPIIntegrationTest::TestCreatePropertyHandleStruct(*this);
+	});
+	It("CreatePropertyHandleArray", [this]() {
+		RemoteControlAPIIntegrationTest::TestCreatePropertyHandleArray(*this);
+	});
+	It("CreatePropertyHandleSet", [this]() {
+		RemoteControlAPIIntegrationTest::TestCreatePropertyHandleSet(*this);
+	});
+	It("CreatePropertyHandleMap", [this]() {
+		RemoteControlAPIIntegrationTest::TestCreatePropertyHandleMap(*this);
+	});
+	It("FloatProperty", [this]() {
+		RemoteControlAPIIntegrationTest::TestFloatProperty(*this);
+	});
+	It("DoubleProperty", [this]() {
+		RemoteControlAPIIntegrationTest::TestDoubleProperty(*this);
+	});
+	It("StringProperty", [this]() {
+		RemoteControlAPIIntegrationTest::TestStringProperty(*this);
+	});
+	It("NameProperty", [this]() {
+		RemoteControlAPIIntegrationTest::TestNameProperty(*this);
+	});
+	It("BoolProperty", [this]() {
+		RemoteControlAPIIntegrationTest::TestBoolProperty(*this);
+	});
+	It("ByteProperty", [this]() {
+		RemoteControlAPIIntegrationTest::TestByteProperty(*this);
+	});
+	It("TextProperty", [this]() {
+		RemoteControlAPIIntegrationTest::TestTextProperty(*this);
+	});
+	It("VectorProperty", [this]() {
+		RemoteControlAPIIntegrationTest::TestVectorProperty(*this);
+	});
+	It("RotatorProperty", [this]() {
+		RemoteControlAPIIntegrationTest::TestRotatorProperty(*this);
+	});
+	It("ColorProperty", [this]() {
+		RemoteControlAPIIntegrationTest::TestColorProperty(*this);
+	});
+	It("LinearColorProperty", [this]() {
+		RemoteControlAPIIntegrationTest::TestLinearColorProperty(*this);
+	});
+	It("ComplexPath", [this]() {
+		RemoteControlAPIIntegrationTest::TestComplexPath(*this);
+	});
+	It("GetPropertyHandleByFieldPath", [this]() {
+		RemoteControlAPIIntegrationTest::TestGetPropertyHandleByFieldPath(*this);
+	});
+	It("StructInObjectProperty", [this]() {
+		RemoteControlAPIIntegrationTest::TestStructInObjectProperty(*this);
+	});
 }
 
 #undef GET_TEST_PROP
diff --git a/Engine/Plugins/VirtualProduction/RemoteControl/Source/RemoteControl/Private/Tests/RemoteControlTests.cpp b/Engine/Plugins/VirtualProduction/RemoteControl/Source/RemoteControl/Private/Tests/RemoteControlTests.cpp
--- a/Engine/Plugins/VirtualProduction/RemoteControl/Source/RemoteControl/Private/Tests/RemoteControlTests.cpp
+++ b/Engine/Plugins/VirtualProduction/RemoteControl/Source/RemoteControl/Private/Tests/RemoteControlTests.cpp
@@ -81,7 +81,7 @@ namespace RemoteControlTest
 	}
 }
 
-IMPLEMENT_SIMPLE_AUTOMATION_TEST(FRemoteControlPresetIntegrationTest, "Plugin.RemoteControl.Expose", EAutomationTestFlags::EditorContext | EAutomationTestFlags::EngineFilter)
+IMPLEMENT_SIMPLE_AUTOMATION_TEST(FRemoteControlPresetIntegrationTest, "Plugins.RemoteControl.Expose", EAutomationTestFlags::EditorContext | EAutomationTestFlags::EngineFilter)
 bool FRemoteControlPresetIntegrationTest::RunTest(const FString& Parameters)
 {
 	// Test expose whole container
diff --git a/Engine/Plugins/VirtualProduction/RemoteControl/Source/RemoteControlCommon/Private/Tests/RCPropertyContainerTests.cpp b/Engine/Plugins/VirtualProduction/RemoteControl/Source/RemoteControlCommon/Private/Tests/RCPropertyContainerTests.cpp
--- a/Engine/Plugins/VirtualProduction/RemoteControl/Source/RemoteControlCommon/Private/Tests/RCPropertyContainerTests.cpp
+++ b/Engine/Plugins/VirtualProduction/RemoteControl/Source/RemoteControlCommon/Private/Tests/RCPropertyContainerTests.cpp
@@ -11,7 +11,7 @@
 #if WITH_DEV_AUTOMATION_TESTS
 
 BEGIN_DEFINE_SPEC(FPropertyContainerSpec,
-	"Plugin.RemoteControl.PropertyContainer",
+	"Plugins.RemoteControl.PropertyContainer",
 	EAutomationTestFlags::EditorContext | EAutomationTestFlags::EngineFilter | EAutomationTestFlags::ApplicationContextMask)
 END_DEFINE_SPEC(FPropertyContainerSpec)
 
diff --git a/Engine/Plugins/VirtualProduction/RemoteControl/Source/RemoteControlLogic/Private/Tests/RemoteControlLogicTests.cpp b/Engine/Plugins/VirtualProduction/RemoteControl/Source/RemoteControlLogic/Private/Tests/RemoteControlLogicTests.cpp
--- a/Engine/Plugins/VirtualProduction/RemoteControl/Source/RemoteControlLogic/Private/Tests/RemoteControlLogicTests.cpp
+++ b/Engine/Plugins/VirtualProduction/RemoteControl/Source/RemoteControlLogic/Private/Tests/RemoteControlLogicTests.cpp
@@ -23,7 +23,7 @@
 #define PROP_NAME(Class, Name) GET_MEMBER_NAME_CHECKED(Class, Name)
 #define GET_TEST_PROP(PropName) URemoteControlLogicTestData::StaticClass()->FindPropertyByName(PROP_NAME(URemoteControlLogicTestData, PropName))
  
-IMPLEMENT_SIMPLE_AUTOMATION_TEST(FRemoteControlLogicTest, "Plugin.RemoteControl.Logic.Runtime", EAutomationTestFlags::EditorContext | EAutomationTestFlags::EngineFilter)
+IMPLEMENT_SIMPLE_AUTOMATION_TEST(FRemoteControlLogicTest, "Plugins.RemoteControl.Logic.Runtime", EAutomationTestFlags::EditorContext | EAutomationTestFlags::EngineFilter)
 bool FRemoteControlLogicTest::RunTest(const FString& Parameters)
 {
 	// 1. Create Preset
diff --git a/Engine/Plugins/VirtualProduction/RemoteControl/Source/RemoteControlUI/Private/UI/Controller/Tests/RCTypeTranslator.spec.cpp b/Engine/Plugins/VirtualProduction/RemoteControl/Source/RemoteControlUI/Private/UI/Controller/Tests/RCTypeTranslator.spec.cpp
--- a/Engine/Plugins/VirtualProduction/RemoteControl/Source/RemoteControlUI/Private/UI/Controller/Tests/RCTypeTranslator.spec.cpp
+++ b/Engine/Plugins/VirtualProduction/RemoteControl/Source/RemoteControlUI/Private/UI/Controller/Tests/RCTypeTranslator.spec.cpp
@@ -1,4 +1,4 @@
-﻿// Copyright Epic Games, Inc. All Rights Reserved.
+// Copyright Epic Games, Inc. All Rights Reserved.
 
 #include "Misc/AutomationTest.h"
 #include "RemoteControlPreset.h"
@@ -7,7 +7,7 @@
 
 BEGIN_DEFINE_SPEC(
 	FRCTypeTranslatorSpec,
-	"Plugin.RemoteControlUI.RCTypeTranslator",
+	"Plugins.RemoteControlUI.RCTypeTranslator",
 	EAutomationTestFlags::ApplicationContextMask | EAutomationTestFlags::ProductFilter)
 	TStrongObjectPtr<URemoteControlPreset> Preset;
 	TStrongObjectPtr<URCVirtualPropertyBase> BoolVirtualProperty;
diff --git a/Engine/Plugins/VirtualProduction/Rivermax/RivermaxCore/Source/RivermaxCore/Private/Tests/PTPTest.cpp b/Engine/Plugins/VirtualProduction/Rivermax/RivermaxCore/Source/RivermaxCore/Private/Tests/PTPTest.cpp
--- a/Engine/Plugins/VirtualProduction/Rivermax/RivermaxCore/Source/RivermaxCore/Private/Tests/PTPTest.cpp
+++ b/Engine/Plugins/VirtualProduction/Rivermax/RivermaxCore/Source/RivermaxCore/Private/Tests/PTPTest.cpp
@@ -14,7 +14,7 @@
 
 #if WITH_DEV_AUTOMATION_TESTS
 
-IMPLEMENT_SIMPLE_AUTOMATION_TEST(FPTPTest, "Plugin.Rivermax.PTP", EAutomationTestFlags::ApplicationContextMask | EAutomationTestFlags::EngineFilter)
+IMPLEMENT_SIMPLE_AUTOMATION_TEST(FPTPTest, "Plugins.Rivermax.PTP", EAutomationTestFlags::ApplicationContextMask | EAutomationTestFlags::EngineFilter)
 
 struct FTestData
 {
