[RemoteControl]
Better error in RC, both for the Editor and the API.

diff --git a/Engine/Plugins/VirtualProduction/RemoteControl/Source/RemoteControl/Private/RemoteControlField.cpp b/Engine/Plugins/VirtualProduction/RemoteControl/Source/RemoteControl/Private/RemoteControlField.cpp
--- a/Engine/Plugins/VirtualProduction/RemoteControl/Source/RemoteControl/Private/RemoteControlField.cpp
+++ b/Engine/Plugins/VirtualProduction/RemoteControl/Source/RemoteControl/Private/RemoteControlField.cpp
@@ -311,9 +311,14 @@ void FRemoteControlProperty::EnableEditCondition()
 #endif
 }
 
-bool FRemoteControlProperty::IsEditableInPackaged() const
+bool FRemoteControlProperty::IsEditableInPackaged(FString* OutError) const
 {
-	return bIsEditableInPackaged || IRemoteControlModule::Get().PropertySupportsRawModificationWithoutEditor(GetProperty(), GetSupportedBindingClass());
+	return bIsEditableInPackaged || IRemoteControlModule::Get().PropertySupportsRawModification(GetProperty(), GetBoundObject(), false, OutError);
+}
+
+bool FRemoteControlProperty::IsEditableInEditor(FString* OutError) const
+{
+	return IRemoteControlModule::Get().PropertySupportsRawModification(GetProperty(), GetBoundObject(), true, OutError);
 }
 
 bool FRemoteControlProperty::Serialize(FArchive& Ar)
diff --git a/Engine/Plugins/VirtualProduction/RemoteControl/Source/RemoteControl/Private/RemoteControlModule.cpp b/Engine/Plugins/VirtualProduction/RemoteControl/Source/RemoteControl/Private/RemoteControlModule.cpp
--- a/Engine/Plugins/VirtualProduction/RemoteControl/Source/RemoteControl/Private/RemoteControlModule.cpp
+++ b/Engine/Plugins/VirtualProduction/RemoteControl/Source/RemoteControl/Private/RemoteControlModule.cpp
@@ -112,28 +112,178 @@ namespace RemoteControlUtil
 		return Function;
 	}
 
-	bool IsPropertyAllowed(const FProperty* InProperty, ERCAccess InAccessType, bool bObjectInGamePackage)
+	/** Returns whether the access is a write access regardless of if it generates a transaction. */
+	bool IsWriteAccess(ERCAccess Access)
+	{
+		return Access == ERCAccess::WRITE_ACCESS || Access == ERCAccess::WRITE_TRANSACTION_ACCESS || Access == ERCAccess::WRITE_MANUAL_TRANSACTION_ACCESS;
+	}
+
+	/** Returns whether the access is a read access. */
+	bool IsReadAccess(ERCAccess Access)
+	{
+		return Access == ERCAccess::READ_ACCESS;
+	}
+
+	bool PropertyModificationShouldUseSetter(const UObject* Object, FProperty* Property)
+	{
+		if (!Property || !Object)
+		{
+			return false;
+		}
+
+		return Property->HasSetter() || !!RemoteControlPropertyUtilities::FindSetterFunction(Property, Object->GetClass());
+	}
+
+	bool IsPropertyAllowed(const FProperty* InProperty, ERCAccess InAccessType, const UObject* InObject, bool bObjectInGamePackage, FString* OutError)
 	{
+		FTextBuilder Error;
+		FText PreError;
+		if (IsWriteAccess(InAccessType))
+		{
+			PreError = LOCTEXT("RCWritingAccess", "Property is not writeable because");
+		}
+		else if (IsReadAccess(InAccessType))
+		{
+			PreError = LOCTEXT("RCReadingAccess", "Property is not readable because");
+		}
+		else
+		{
+			PreError = LOCTEXT("RCNoAccess", "Property");
+		}
 #if WITH_EDITOR
 		URemoteControlSettings* RCSettings = GetMutableDefault<URemoteControlSettings>();
 		// Override this flag to false if we are running in Editor Mode with -game flag.
 		bObjectInGamePackage = false;
 #endif
 		// The property is allowed to be accessed if it exists...
-		return InProperty &&
-				// it doesn't have exposed getter/setter that should be used instead
+		const bool bPropertyIsValid = InProperty != nullptr;
+		if (!bPropertyIsValid)
+		{
+			Error.AppendLine(FText::Format(LOCTEXT("RCPropertyNotValid", "{0} is not valid."), PreError));
+			*OutError = Error.ToText().ToString();
+			return false;
+		}
+#if WITH_EDITOR
+		// it doesn't have exposed getter AND setter that should be used instead or if it should ignore them
+		const bool bHasBlueprintGetterOrSetterOrIgnored = !InProperty->HasMetaData(RemoteControlUtil::NAME_BlueprintGetter) ||
+												 		  !InProperty->HasMetaData(RemoteControlUtil::NAME_BlueprintSetter) ||
+												 		  RCSettings->bIgnoreGetterSetterCheck;
+		if (!bHasBlueprintGetterOrSetterOrIgnored)
+		{
+			Error.AppendLine(FText::Format(LOCTEXT("RCBlueprintHasGetterSetterOrNotIgnored", "{0} has blueprint getter and setter to use.\nEnable in the RemoteControl settings IgnoreGetterSetterCheck if you want to ignore it."), PreError));
+		}
+
+		// it isn't private or protected, except if AllowPrivateAccess is true and if only for Protected whenever it should be ignored
+		const bool bAllowPrivateAccessForPrivate =  !InProperty->HasAnyPropertyFlags(CPF_NativeAccessSpecifierPrivate) ||
+													InProperty->GetBoolMetaData(RemoteControlUtil::NAME_AllowPrivateAccess);
+
+		const bool bAllowPrivateAccessForProtected = !InProperty->HasAnyPropertyFlags(CPF_NativeAccessSpecifierProtected) ||
+													 RCSettings->bIgnoreProtectedCheck ||
+													 InProperty->GetBoolMetaData(RemoteControlUtil::NAME_AllowPrivateAccess);
+
+#endif
+
+		const bool bPropertyHasSetterToUse = IsWriteAccess(InAccessType)? PropertyModificationShouldUseSetter(InObject, const_cast<FProperty*>(InProperty)) : true;
+
+		const bool bPropertyHasGetterToUse = IsReadAccess(InAccessType)? InProperty->HasGetter() : true;
+
+		bool bAppendSettingInformation = false;
+		if (!bPropertyHasSetterToUse)
+		{
+#if WITH_EDITOR
+			if (!bAllowPrivateAccessForPrivate)
+			{
+				Error.AppendLine(FText::Format(LOCTEXT("RCPropertyPrivateWithoutSetter", "{0} is private without the AllowPrivateAccess specifier set to true or a setter assigned."), PreError));
+			}
+
+			if (!bAllowPrivateAccessForProtected)
+			{
+				Error.AppendLine(FText::Format(LOCTEXT("RCPropertyProtectedWithoutSetterOrIgnored", "{0} is protected without the AllowPrivateAccess specifier set to true or a setter assigned."), PreError));
+				bAppendSettingInformation = true;
+			}
+#else
+			Error.AppendLine(FText::Format(LOCTEXT("RCPropertyNoSetter", "{0} it doesn't have a setter specifier assigned to it."), PreError));
+#endif
+			
+		}
+
+		if (!bPropertyHasGetterToUse)
+		{
+#if WITH_EDITOR
+			if (!bAllowPrivateAccessForPrivate)
+			{
+				Error.AppendLine(FText::Format(LOCTEXT("RCPropertyPrivateWithoutGetter", "{0} is private without the AllowPrivateAccess specifier set to true or a getter assigned."), PreError));
+			}
+
+			if (!bAllowPrivateAccessForProtected)
+			{
+				Error.AppendLine(FText::Format(LOCTEXT("RCPropertyProtectedWithoutGetterOrIgnored", "{0} is protected without the AllowPrivateAccess specifier set to true or a getter assigned."), PreError));
+				bAppendSettingInformation = true;
+			}
+#else
+			Error.AppendLine(FText::Format(LOCTEXT("RCPropertyNoGetter", "{0} doesn't have a getter specifier assigned to it."), PreError));
+#endif
+		}
+
 #if WITH_EDITOR
-				(!InProperty->HasMetaData(RemoteControlUtil::NAME_BlueprintGetter) || !InProperty->HasMetaData(RemoteControlUtil::NAME_BlueprintSetter) || RCSettings->bIgnoreGetterSetterCheck) &&
-				// it isn't private or protected, except if AllowPrivateAccess is true and if only for Protected whenever it should be ignored
-				(!InProperty->HasAnyPropertyFlags(CPF_NativeAccessSpecifierPrivate) || InProperty->GetBoolMetaData(RemoteControlUtil::NAME_AllowPrivateAccess)) &&
-				(!InProperty->HasAnyPropertyFlags(CPF_NativeAccessSpecifierProtected) || RCSettings->bIgnoreProtectedCheck || InProperty->GetBoolMetaData(RemoteControlUtil::NAME_AllowPrivateAccess)) &&	
+		if (bAppendSettingInformation)
+		{
+			Error.AppendLine(LOCTEXT("RCSettingsIgnoreProtected","You can enable the RemoteControl setting IgnoreProtectedCheck to ignore it for protected properties."));
+		}
 #endif
-				// it isn't blueprint private
-				!InProperty->HasAnyPropertyFlags(CPF_DisableEditOnInstance) &&
-				// and it's either blueprint visible if in game or editable if in editor and it isn't read only if the access type is write
-				(bObjectInGamePackage
-					? InProperty->HasAnyPropertyFlags(CPF_BlueprintVisible) && (InAccessType == ERCAccess::READ_ACCESS || !InProperty->HasAnyPropertyFlags(CPF_BlueprintReadOnly))
-					: InAccessType == ERCAccess::READ_ACCESS || (InProperty->HasAnyPropertyFlags(CPF_Edit) && !InProperty->HasAnyPropertyFlags(CPF_EditConst)));
+		// it isn't blueprint private
+		const bool bIsBlueprintPublic = !InProperty->HasAnyPropertyFlags(CPF_DisableEditOnInstance);
+
+		if (!bIsBlueprintPublic)
+		{
+			Error.AppendLine(FText::Format(LOCTEXT("RCBlueprintPropertyPrivate", "{0} has flag DisableEditOnInstance.\nEnable the InstanceEditable flag of the property in the blueprint editor."), PreError));
+		}
+
+		// and it's either blueprint visible if in game or editable if in editor and it isn't read only if the access type is write
+		bool bIsAccessibleInGamePackage;
+		bool bIsAccessibleInNonGamePackage;
+
+		if (bObjectInGamePackage)
+		{
+			bIsAccessibleInGamePackage = InProperty->HasAnyPropertyFlags(CPF_BlueprintVisible) &&
+										 (InAccessType == ERCAccess::READ_ACCESS ||
+										 !InProperty->HasAnyPropertyFlags(CPF_BlueprintReadOnly));
+
+			if (!bIsAccessibleInGamePackage)
+			{
+				Error.AppendLine(FText::Format(LOCTEXT("RCBlueprintNotAccessibleInGamePackage", "{0} doesn't have the BlueprintVisible flag and/or you are trying to write to it but it has the BlueprintReadOnly flag."), PreError));
+			}
+			// set this to true since it will not be evaluated in GamePackage
+			bIsAccessibleInNonGamePackage = true;
+		}
+		else
+		{
+			bIsAccessibleInNonGamePackage = InAccessType == ERCAccess::READ_ACCESS ||
+											(InProperty->HasAnyPropertyFlags(CPF_Edit) &&
+											!InProperty->HasAnyPropertyFlags(CPF_EditConst));
+
+			if (!bIsAccessibleInNonGamePackage)
+			{
+				Error.AppendLine(FText::Format(LOCTEXT("RCBlueprintNotAccessibleInNonGamePackage", "{0} has the EditConst flag."), PreError));
+			}
+			// set this to true since it will not be evaluated in NON GamePackage
+			bIsAccessibleInGamePackage = true;
+		}
+
+		const bool bIsAccessible = bObjectInGamePackage? bIsAccessibleInGamePackage : bIsAccessibleInNonGamePackage;
+
+		// Assign the error to be propagated after this call in case it failed
+		*OutError = Error.ToText().ToString();
+
+		return (bPropertyIsValid &&
+#if WITH_EDITOR
+			bHasBlueprintGetterOrSetterOrIgnored &&
+			bAllowPrivateAccessForPrivate &&
+			bAllowPrivateAccessForProtected &&
+#endif
+			bIsBlueprintPublic &&
+			bIsAccessible) ||
+			(bPropertyHasGetterToUse && bPropertyHasSetterToUse);
 	};
 
 	FARFilter GetBasePresetFilter()
@@ -207,18 +357,6 @@ namespace RemoteControlUtil
 		return Id;
 	}
 
-	/** Returns whether the access is a write access regardless of if it generates a transaction. */
-	bool IsWriteAccess(ERCAccess Access)
-	{
-		return Access == ERCAccess::WRITE_ACCESS || Access == ERCAccess::WRITE_TRANSACTION_ACCESS || Access == ERCAccess::WRITE_MANUAL_TRANSACTION_ACCESS;
-	}
-
-	/** Returns whether the access is a read access. */
-	bool IsReadAccess(ERCAccess Access)
-	{
-		return Access == ERCAccess::READ_ACCESS;
-	}
-
 	/** Helper function to check if a value is 0 without always calling FMath::IsNearlyZero, since that results in an ambiguous call for non-float values. */
 	template <typename T>
 	bool IsNearlyZero(T Value)
@@ -1174,16 +1312,12 @@ bool FRemoteControlModule::ResolveObjectProperty(ERCAccess AccessType, UObject*
 				{
 					FProperty* ResolvedProperty = PropertyPath.GetResolvedData().Field;
 
-					// When resolving a property for writing, resolve successfully if it should use a setter since it will end up using it. 
-					if ((RemoteControlUtil::IsWriteAccess(AccessType) && PropertyModificationShouldUseSetter(Object, ResolvedProperty))
-						|| RemoteControlUtil::IsPropertyAllowed(ResolvedProperty, AccessType, bObjectInGame)
-						|| (RemoteControlUtil::IsReadAccess(AccessType) && ResolvedProperty->HasGetter()))
+					if (RemoteControlUtil::IsPropertyAllowed(ResolvedProperty, AccessType, Object, bObjectInGame, &ErrorText))
 					{
 						OutObjectRef = FRCObjectReference{AccessType, Object, MoveTemp(PropertyPath)};
 					}
 					else
 					{
-						ErrorText = FString::Printf(TEXT("Object property: %s is unavailable remotely on object: %s"), *PropertyPath.GetFieldName().ToString(), *Object->GetPathName());
 						bSuccess = false;
 					}
 				}
@@ -1213,11 +1347,11 @@ bool FRemoteControlModule::ResolveObjectProperty(ERCAccess AccessType, UObject*
 	if (!bSuccess && !ErrorText.IsEmpty())
 	{
 		IRemoteControlModule::BroadcastError(ErrorText);
-	}
 
-	if (OutErrorText && !ErrorText.IsEmpty())
-	{
-		*OutErrorText = MoveTemp(ErrorText);
+		if (OutErrorText != nullptr)
+		{
+			*OutErrorText = MoveTemp(ErrorText);
+		}
 	}
 
 	return bSuccess;
@@ -1252,9 +1386,15 @@ bool FRemoteControlModule::GetObjectProperties(const FRCObjectReference& ObjectA
 			bool bObjectInGame = !GIsEditor || Object->GetOutermost()->HasAnyPackageFlags(PKG_PlayInEditor);
 			Policies.PropertyFilter = [&ObjectAccess, bObjectInGame](const FProperty* CurrentProp, const FProperty* ParentProp)
 			{
-				return	RemoteControlUtil::IsPropertyAllowed(CurrentProp, ObjectAccess.Access, bObjectInGame) ||
-						ParentProp != nullptr ||
-						(RemoteControlUtil::IsReadAccess(ObjectAccess.Access) && CurrentProp->HasGetter());
+				FString ErrorText;
+				const bool bIsPropertyAllowed =	RemoteControlUtil::IsPropertyAllowed(CurrentProp, ObjectAccess.Access, ObjectAccess.Object.Get(), bObjectInGame, &ErrorText);
+				if (bIsPropertyAllowed || ParentProp != nullptr)
+				{
+					return true;
+				}
+
+				IRemoteControlModule::BroadcastError(ErrorText + TEXT("\nAlso parent property is not valid."));
+				return false;
 			};
 		}
 
@@ -1407,7 +1547,7 @@ bool FRemoteControlModule::SetObjectProperties(const FRCObjectReference& ObjectA
 	if (InPayload.Num() != 0 && ObjectAccess.Object.IsValid())
 	{
 		// Convert raw property modifications to setter function calls if necessary.
-		if (PropertyModificationShouldUseSetter(ObjectAccess.Object.Get(), ObjectAccess.Property.Get()))
+		if (RemoteControlUtil::PropertyModificationShouldUseSetter(ObjectAccess.Object.Get(), ObjectAccess.Property.Get()))
 		{
 			FRCCall Call;
 			FRCInterceptionPayload InterceptionPayload;
@@ -1441,8 +1581,11 @@ bool FRemoteControlModule::SetObjectProperties(const FRCObjectReference& ObjectA
 
 		// If a setter wasn't used, verify if the property should be allowed.
 		bool bObjectInGame = !GIsEditor;
-		if (!RemoteControlUtil::IsPropertyAllowed(ObjectAccess.Property.Get(), ObjectAccess.Access, bObjectInGame))
+		FString ErrorText;
+
+		if (!RemoteControlUtil::IsPropertyAllowed(ObjectAccess.Property.Get(), ObjectAccess.Access, ObjectAccess.Object.Get() ,bObjectInGame, &ErrorText))
 		{
+			IRemoteControlModule::BroadcastError(ErrorText);
 			return false;
 		}
 
@@ -1476,7 +1619,7 @@ bool FRemoteControlModule::SetObjectProperties(const FRCObjectReference& ObjectA
 	}
 
 	// Convert raw property modifications to setter function calls if necessary.
-	if (PropertyModificationShouldUseSetter(ObjectAccess.Object.Get(), ObjectAccess.Property.Get()))
+	if (RemoteControlUtil::PropertyModificationShouldUseSetter(ObjectAccess.Object.Get(), ObjectAccess.Property.Get()))
 	{
 		FRCCall Call;
 		RemoteControlSetterUtils::FConvertToFunctionCallArgs Args(ObjectAccess, Backend, Call);
@@ -1536,7 +1679,15 @@ bool FRemoteControlModule::SetObjectProperties(const FRCObjectReference& ObjectA
 			bool bObjectInGame = !GIsEditor || MutableObjectReference.Object->GetOutermost()->HasAnyPackageFlags(PKG_PlayInEditor);
 			Policies.PropertyFilter = [&MutableObjectReference, bObjectInGame](const FProperty* CurrentProp, const FProperty* ParentProp)
 			{
-				return RemoteControlUtil::IsPropertyAllowed(CurrentProp, MutableObjectReference.Access, bObjectInGame) || ParentProp != nullptr;
+				FString ErrorText;
+				const bool bIsPropertyAllowed =	RemoteControlUtil::IsPropertyAllowed(CurrentProp, MutableObjectReference.Access, MutableObjectReference.Object.Get(), bObjectInGame, &ErrorText);
+				if (bIsPropertyAllowed || ParentProp != nullptr)
+				{
+					return true;
+				}
+
+				IRemoteControlModule::BroadcastError(ErrorText + TEXT("\nAlso parent property is not valid."));
+				return false;
 			};
 		}
 
@@ -2099,7 +2250,46 @@ void FRemoteControlModule::UnregisterDefaultEntityMetadata(FName MetadataKey)
 bool FRemoteControlModule::PropertySupportsRawModificationWithoutEditor(FProperty* Property, UClass* OwnerClass) const
 {
 	constexpr bool bInGameOrPackage = true;
-	return Property && (RemoteControlUtil::IsPropertyAllowed(Property, ERCAccess::WRITE_ACCESS, bInGameOrPackage) || !!RemoteControlPropertyUtilities::FindSetterFunction(Property, OwnerClass));
+	return Property && (RemoteControlUtil::IsPropertyAllowed(Property, ERCAccess::WRITE_ACCESS, nullptr, bInGameOrPackage, nullptr) || !!RemoteControlPropertyUtilities::FindSetterFunction(Property, OwnerClass));
+}
+
+bool FRemoteControlModule::PropertySupportsRawModification(FProperty* InProperty, const UObject* InObject, const bool bInWithEditor, FString* OutError) const
+{
+	FString ErrorTextSet;
+	FString ErrorTextGet;
+	const bool bCheckPackaged = !bInWithEditor;
+	if (!InProperty)
+	{
+		return false;
+	}
+
+	const bool bIsPropertyAllowedToSet = RemoteControlUtil::IsPropertyAllowed(InProperty, ERCAccess::WRITE_ACCESS, InObject, bCheckPackaged, &ErrorTextSet);
+	const bool bIsPropertyAllowedToRead = RemoteControlUtil::IsPropertyAllowed(InProperty, ERCAccess::READ_ACCESS, InObject, bCheckPackaged, &ErrorTextGet);
+
+	if (bIsPropertyAllowedToSet && bIsPropertyAllowedToRead)
+	{
+		return true;
+	}
+
+	if (OutError)
+	{
+		if (!bIsPropertyAllowedToSet)
+		{
+			if (!bIsPropertyAllowedToRead)
+			{
+				*OutError = FString::Printf(TEXT("%s\n\n%s"), *ErrorTextSet, *ErrorTextGet);
+			}
+			else
+			{
+				*OutError = ErrorTextSet;
+			}
+		}
+		else if (!bIsPropertyAllowedToRead)
+		{
+			*OutError = ErrorTextGet;
+		}
+	}
+	return false;
 }
 
 void FRemoteControlModule::RegisterEntityFactory(const FName InFactoryName, const TSharedRef<IRemoteControlPropertyFactory>& InFactory)
@@ -2395,16 +2585,6 @@ bool FRemoteControlModule::DestroyTransientPreset(URemoteControlPreset* Preset)
 	return false;
 }
 
-bool FRemoteControlModule::PropertyModificationShouldUseSetter(UObject* Object, FProperty* Property)
-{
-	if (!Property || !Object)
-	{
-		return false;
-	}
-
-	return Property->HasSetter() || !!RemoteControlPropertyUtilities::FindSetterFunction(Property, Object->GetClass());
-}
-
 bool FRemoteControlModule::DeserializeDeltaModificationData(const FRCObjectReference& ObjectAccess, IStructDeserializerBackend& Backend, ERCModifyOperation Operation, TArray<uint8>& OutData)
 {
 	// Allocate data to deserialize the request data into
diff --git a/Engine/Plugins/VirtualProduction/RemoteControl/Source/RemoteControl/Private/RemoteControlModule.h b/Engine/Plugins/VirtualProduction/RemoteControl/Source/RemoteControl/Private/RemoteControlModule.h
--- a/Engine/Plugins/VirtualProduction/RemoteControl/Source/RemoteControl/Private/RemoteControlModule.h
+++ b/Engine/Plugins/VirtualProduction/RemoteControl/Source/RemoteControl/Private/RemoteControlModule.h
@@ -67,6 +67,7 @@ public:
 	virtual bool RegisterDefaultEntityMetadata(FName MetadataKey, FEntityMetadataInitializer MetadataInitializer) override;
 	virtual void UnregisterDefaultEntityMetadata(FName MetadataKey) override;
 	virtual bool PropertySupportsRawModificationWithoutEditor(FProperty* Property, UClass* OwnerClass = nullptr) const override;
+	virtual bool PropertySupportsRawModification(FProperty* InProperty, const UObject* InObject, const bool bInWithEditor, FString* OutError = nullptr) const override;
 	virtual void RegisterEntityFactory( const FName InFactoryName, const TSharedRef<IRemoteControlPropertyFactory>& InFactory) override;
 	virtual void UnregisterEntityFactory( const FName InFactoryName ) override;
 	virtual FGuid BeginManualEditorTransaction(const FText& InDescription, uint32 TypeHash) override;
@@ -95,9 +96,6 @@ private:
 	/** Destroy a transient preset using an object reference. */
 	bool DestroyTransientPreset(URemoteControlPreset* Preset);
 
-	/** Determines if a property modification should use a setter or default to deserializing directly onto an object. */
-	static bool PropertyModificationShouldUseSetter(UObject* Object, FProperty* Property);
-
 	/**
 	 * Deserialize data for a non-EQUAL modification request and apply the operation to the resulting data.
 	 * 
diff --git a/Engine/Plugins/VirtualProduction/RemoteControl/Source/RemoteControl/Public/IRemoteControlModule.h b/Engine/Plugins/VirtualProduction/RemoteControl/Source/RemoteControl/Public/IRemoteControlModule.h
--- a/Engine/Plugins/VirtualProduction/RemoteControl/Source/RemoteControl/Public/IRemoteControlModule.h
+++ b/Engine/Plugins/VirtualProduction/RemoteControl/Source/RemoteControl/Public/IRemoteControlModule.h
@@ -563,8 +563,19 @@ public:
 	/**
 	 * Returns whether the property can be modified through SetObjectProperties when running without an editor.
 	 */
+	UE_DEPRECATED(5.4, "This function is deprecated, please use PropertySupportsRawModification.")
 	virtual bool PropertySupportsRawModificationWithoutEditor(FProperty* Property, UClass* OwnerClass = nullptr) const = 0;
 
+	/**
+	 * Returns whether the property can be modified through SetObjectProperties or read through ResolveObjectProperties when running with or without an editor.
+	 * @param InProperty Property to check
+	 * @param InObject Object that owns the property
+	 * @param bInWithEditor True if it should check if the property is supported in Editor, false to check for Packaged
+	 * @param OutError Will contain the error in case it is not supported
+	 * @return True if the property is supported both for read and write, false otherwise
+	 */
+	virtual bool PropertySupportsRawModification(FProperty* InProperty, const UObject* InObject, const bool bInWithEditor, FString* OutError = nullptr) const = 0;
+
 	/**
 	 * Register factory 
 	 * @param InFactoryName the factory unique name
diff --git a/Engine/Plugins/VirtualProduction/RemoteControl/Source/RemoteControl/Public/RemoteControlField.h b/Engine/Plugins/VirtualProduction/RemoteControl/Source/RemoteControl/Public/RemoteControlField.h
--- a/Engine/Plugins/VirtualProduction/RemoteControl/Source/RemoteControl/Public/RemoteControlField.h
+++ b/Engine/Plugins/VirtualProduction/RemoteControl/Source/RemoteControl/Public/RemoteControlField.h
@@ -208,7 +208,10 @@ public:
 	void EnableEditCondition();
 	
 	/** Returns whether the property is editable in a packaged build. */
-	bool IsEditableInPackaged() const;
+	bool IsEditableInPackaged(FString* OutError = nullptr) const;
+
+	/** Returns whether the property is editable in the Editor. */
+	bool IsEditableInEditor(FString* OutError = nullptr) const;
 
 	bool Serialize(FArchive& Ar);
 	void PostSerialize(const FArchive& Ar);
diff --git a/Engine/Plugins/VirtualProduction/RemoteControl/Source/RemoteControlLogic/Private/Action/RCActionContainer.cpp b/Engine/Plugins/VirtualProduction/RemoteControl/Source/RemoteControlLogic/Private/Action/RCActionContainer.cpp
--- a/Engine/Plugins/VirtualProduction/RemoteControl/Source/RemoteControlLogic/Private/Action/RCActionContainer.cpp
+++ b/Engine/Plugins/VirtualProduction/RemoteControl/Source/RemoteControlLogic/Private/Action/RCActionContainer.cpp
@@ -158,11 +158,18 @@ URCPropertyAction* URCActionContainer::AddPropertyAction(const TSharedRef<const
 
 	if(!bFoundMatchingContainer)
 	{
+		// Check both Reading and Writing since if one of the two is not valid then the action won't work
+		FRCObjectReference ObjectRefReading;
+		const bool bResolveForReading = IRemoteControlModule::Get().ResolveObjectProperty(ERCAccess::READ_ACCESS, InRemoteControlProperty->GetBoundObjects()[0], InRemoteControlProperty->FieldPathInfo.ToString(), ObjectRefReading);
+
+		FRCObjectReference ObjectRefWriting;
+		const bool bResolveForWriting = IRemoteControlModule::Get().ResolveObjectProperty(ERCAccess::WRITE_ACCESS, InRemoteControlProperty->GetBoundObjects()[0], InRemoteControlProperty->FieldPathInfo.ToString(), ObjectRefWriting);
+
 		// Create an input field for the Action by duplicating the Remote Control Property associated with it
-		if (FRCObjectReference ObjectRef; IRemoteControlModule::Get().ResolveObjectProperty(ERCAccess::READ_ACCESS, InRemoteControlProperty->GetBoundObjects()[0], InRemoteControlProperty->FieldPathInfo.ToString(), ObjectRef))
+		if (bResolveForReading && bResolveForWriting)
 		{
 			const FName& PropertyName = InRemoteControlProperty->GetProperty()->GetFName();
-			NewPropertyAction->PropertySelfContainer->DuplicatePropertyWithCopy(PropertyName, InRemoteControlProperty->GetProperty(), (uint8*)ObjectRef.ContainerAdress);
+			NewPropertyAction->PropertySelfContainer->DuplicatePropertyWithCopy(PropertyName, InRemoteControlProperty->GetProperty(), (uint8*)ObjectRefReading.ContainerAdress);
 		}
 	}
 
diff --git a/Engine/Plugins/VirtualProduction/RemoteControl/Source/RemoteControlUI/Private/UI/SRCPanelExposedEntity.cpp b/Engine/Plugins/VirtualProduction/RemoteControl/Source/RemoteControlUI/Private/UI/SRCPanelExposedEntity.cpp
--- a/Engine/Plugins/VirtualProduction/RemoteControl/Source/RemoteControlUI/Private/UI/SRCPanelExposedEntity.cpp
+++ b/Engine/Plugins/VirtualProduction/RemoteControl/Source/RemoteControlUI/Private/UI/SRCPanelExposedEntity.cpp
@@ -558,6 +558,7 @@ TSharedRef<SWidget> SRCPanelExposedEntity::CreateEntityWidget(TSharedPtr<SWidget
 			.Visibility(!OptionalWarningMessage.IsEmpty() ? EVisibility::Visible : EVisibility::Collapsed)
             .TextStyle(FRemoteControlPanelStyle::Get(), "RemoteControlPanel.Button.TextStyle")
             .Font(FAppStyle::Get().GetFontStyle("FontAwesome.10"))
+            .ColorAndOpacity(FSlateColor(FLinearColor::Yellow))
             .ToolTipText(OptionalWarningMessage)
             .Text(FEditorFontGlyphs::Exclamation_Triangle)
 		]
diff --git a/Engine/Plugins/VirtualProduction/RemoteControl/Source/RemoteControlUI/Private/UI/SRCPanelExposedField.cpp b/Engine/Plugins/VirtualProduction/RemoteControl/Source/RemoteControlUI/Private/UI/SRCPanelExposedField.cpp
--- a/Engine/Plugins/VirtualProduction/RemoteControl/Source/RemoteControlUI/Private/UI/SRCPanelExposedField.cpp
+++ b/Engine/Plugins/VirtualProduction/RemoteControl/Source/RemoteControlUI/Private/UI/SRCPanelExposedField.cpp
@@ -540,18 +540,26 @@ TSharedRef<SWidget> SRCPanelExposedField::MakeFieldWidget(const TSharedRef<SWidg
 
 	FText WarningMessage;
 
-	if (GetDefault<URemoteControlSettings>()->bDisplayInEditorOnlyWarnings)
+	if (const TSharedPtr<FRemoteControlField> RCField = GetRemoteControlField().Pin())
 	{
 		bool bIsEditorOnly = false,
             bIsEditableInPackaged = true,
-            bIsCallableInPackaged = true;
+            bIsCallableInPackaged = true,
+			bIsEditableInEditor = true;
 
-		if (const TSharedPtr<FRemoteControlField> RCField = GetRemoteControlField().Pin())
+		FString ErrorTextInEditor;
+		if (RCField->FieldType == EExposedFieldType::Property)
+		{
+			bIsEditableInEditor = StaticCastSharedPtr<FRemoteControlProperty>(RCField)->IsEditableInEditor(&ErrorTextInEditor);
+		}
+
+		FString ErrorTextInPackage;
+		if (GetDefault<URemoteControlSettings>()->bDisplayInEditorOnlyWarnings)
 		{
 			bIsEditorOnly = RCField->IsEditorOnly();
 			if (RCField->FieldType == EExposedFieldType::Property)
 			{
-				bIsEditableInPackaged = StaticCastSharedPtr<FRemoteControlProperty>(RCField)->IsEditableInPackaged();
+				bIsEditableInPackaged = StaticCastSharedPtr<FRemoteControlProperty>(RCField)->IsEditableInPackaged(&ErrorTextInPackage);
 			}
 			else
 			{
@@ -565,9 +573,38 @@ TSharedRef<SWidget> SRCPanelExposedField::MakeFieldWidget(const TSharedRef<SWidg
 			Builder.AppendLine(LOCTEXT("EditorOnlyWarning", "This field will be unavailable in packaged projects."));
 		}
 
+		bool bEditorAndPackagedErrorEquals = false;
 		if (!bIsEditableInPackaged)
 		{
-			Builder.AppendLine(LOCTEXT("NotEditableInPackagedWarning", "This property will not be editable in packaged projects."));
+			if (!bIsEditableInEditor)
+			{
+				const FName PackagedWarning = FName(ErrorTextInPackage);
+				const FName EditorWarning = FName(ErrorTextInEditor);
+				if (PackagedWarning.IsEqual(EditorWarning))
+				{
+					Builder.AppendLine(LOCTEXT("NonEditableInPackagedAndEditorWarning", "Packaged and Editor:"));
+					bEditorAndPackagedErrorEquals = true;
+				}
+			}
+
+			if (!bEditorAndPackagedErrorEquals)
+			{
+				Builder.AppendLine(LOCTEXT("NonEditableInPackagedWarning", "Packaged:"));
+			}
+
+			Builder.AppendLine(ErrorTextInPackage);
+		}
+
+		if (!bIsEditableInEditor && !bEditorAndPackagedErrorEquals)
+		{
+			if (!bIsEditableInPackaged)
+			{
+				// put an empty line between the last error and this
+				Builder.AppendLine();
+			}
+
+			Builder.AppendLine(LOCTEXT("NonEditableInEditorWarning", "Editor:"));
+			Builder.AppendLine(ErrorTextInEditor);
 		}
 
 		if (!bIsCallableInPackaged)
diff --git a/Engine/Plugins/VirtualProduction/RemoteControl/Source/WebRemoteControl/Private/WebRemoteControl.cpp b/Engine/Plugins/VirtualProduction/RemoteControl/Source/WebRemoteControl/Private/WebRemoteControl.cpp
--- a/Engine/Plugins/VirtualProduction/RemoteControl/Source/WebRemoteControl/Private/WebRemoteControl.cpp
+++ b/Engine/Plugins/VirtualProduction/RemoteControl/Source/WebRemoteControl/Private/WebRemoteControl.cpp
@@ -1339,9 +1339,10 @@ bool FWebRemoteControlModule::HandlePresetSetPropertyRoute(const FHttpServerRequ
 		return true;
 	}
 
+	FString Error;
 	const bool bSuccess = WebRemoteControlInternalUtils::ModifyPropertyUsingPayload(
 		*RemoteControlProperty.Get(), SetPropertyRequest, SetPropertyRequest.TCHARBody, ActingClientId, *WebSocketHandler,
-		SetPropertyRequest.GenerateTransaction ? ERCAccess::WRITE_TRANSACTION_ACCESS : ERCAccess::WRITE_ACCESS);
+		SetPropertyRequest.GenerateTransaction ? ERCAccess::WRITE_TRANSACTION_ACCESS : ERCAccess::WRITE_ACCESS, &Error);
 
 	if (bSuccess)
 	{
@@ -1349,7 +1350,7 @@ bool FWebRemoteControlModule::HandlePresetSetPropertyRoute(const FHttpServerRequ
 	}
 	else
 	{
-		WebRemoteControlInternalUtils::CreateUTF8ErrorMessage(FString::Printf(TEXT("Error while trying to modify property %s."), *Args.FieldLabel), Response->Body);
+		WebRemoteControlInternalUtils::CreateUTF8ErrorMessage(FString::Printf(TEXT("Error while trying to modify property %s.\n%s"), *Args.FieldLabel, *Error), Response->Body);
 	}
 
 	OnComplete(MoveTemp(Response));
@@ -1431,8 +1432,9 @@ bool FWebRemoteControlModule::HandlePresetGetPropertyRoute(const FHttpServerRequ
 	FMemoryWriter Writer(WorkingBuffer);
 
 	FRCObjectReference ObjectRef;
+	FString Error;
 
-	bool bSuccess = IRemoteControlModule::Get().ResolveObjectProperty(ERCAccess::READ_ACCESS, RemoteControlProperty->GetBoundObjects()[0], RemoteControlProperty->FieldPathInfo.ToString(), ObjectRef);
+	bool bSuccess = IRemoteControlModule::Get().ResolveObjectProperty(ERCAccess::READ_ACCESS, RemoteControlProperty->GetBoundObjects()[0], RemoteControlProperty->FieldPathInfo.ToString(), ObjectRef, &Error);
 
 	if (bSuccess)
 	{
@@ -1445,7 +1447,7 @@ bool FWebRemoteControlModule::HandlePresetGetPropertyRoute(const FHttpServerRequ
 	}
 	else
 	{
-		WebRemoteControlInternalUtils::CreateUTF8ErrorMessage(FString::Printf(TEXT("Error while trying to read property %s."), *Args.FieldLabel), Response->Body);
+		WebRemoteControlInternalUtils::CreateUTF8ErrorMessage(FString::Printf(TEXT("Error while trying to read property %s.\n%s"), *Args.FieldLabel, *Error), Response->Body);
 	}
 
 	OnComplete(MoveTemp(Response));
@@ -1589,8 +1591,9 @@ bool FWebRemoteControlModule::HandlePresetGetExposedActorPropertyRoute(const FHt
 			FMemoryWriter Writer(WorkingBuffer);
 
 			FRCObjectReference ObjectRef;
+			FString Error;
 
-			if (IRemoteControlModule::Get().ResolveObjectProperty(ERCAccess::READ_ACCESS, Actor, FieldPath, ObjectRef))
+			if (IRemoteControlModule::Get().ResolveObjectProperty(ERCAccess::READ_ACCESS, Actor, FieldPath, ObjectRef, &Error))
 			{
 				FStructOnScope ActorPropertyOnScope = WebRemoteControlStructUtils::CreateActorPropertyOnScope(ObjectRef);
 				WebRemoteControlInternalUtils::SerializeStructOnScope(ActorPropertyOnScope, Writer);
@@ -1601,7 +1604,7 @@ bool FWebRemoteControlModule::HandlePresetGetExposedActorPropertyRoute(const FHt
 			else
 			{
 				Response->Code = EHttpServerResponseCodes::NotFound;
-				WebRemoteControlInternalUtils::CreateUTF8ErrorMessage(FString::Printf(TEXT("Error while trying to read property %s on actor %s."), *PropertyName, *ActorRCLabel), Response->Body);
+				WebRemoteControlInternalUtils::CreateUTF8ErrorMessage(FString::Printf(TEXT("Error while trying to read property %s on actor %s.\n%s"), *PropertyName, *ActorRCLabel, *Error), Response->Body);
 			}
 		}
 	}
@@ -1688,6 +1691,7 @@ bool FWebRemoteControlModule::HandlePresetSetExposedActorPropertyRoute(const FHt
 	}
 
 	bool bSuccess = true;
+	FString Error;
 	
 	if (TSharedPtr<FRemoteControlActor> RCActor = WebRemoteControl::GetRCEntity<FRemoteControlActor>(Preset, ActorRCLabel))
 	{
@@ -1695,7 +1699,7 @@ bool FWebRemoteControlModule::HandlePresetSetExposedActorPropertyRoute(const FHt
 		{
 			FRCObjectReference ObjectRef;
 			ERCAccess Access = SetPropertyRequest.GenerateTransaction ? ERCAccess::WRITE_TRANSACTION_ACCESS : ERCAccess::WRITE_ACCESS;
-			bSuccess &= IRemoteControlModule::Get().ResolveObjectProperty(Access, Actor, FieldPath, ObjectRef);
+			bSuccess &= IRemoteControlModule::Get().ResolveObjectProperty(Access, Actor, FieldPath, ObjectRef, &Error);
 
 			if (bSuccess && ObjectRef.Property.IsValid())
 			{
@@ -1741,7 +1745,7 @@ bool FWebRemoteControlModule::HandlePresetSetExposedActorPropertyRoute(const FHt
 	}
 	else
 	{
-		WebRemoteControlInternalUtils::CreateUTF8ErrorMessage(FString::Printf(TEXT("Error while trying to modify property %s on actor %s."), *PropertyName, *ActorRCLabel), Response->Body);
+		WebRemoteControlInternalUtils::CreateUTF8ErrorMessage(FString::Printf(TEXT("Error while trying to modify property %s on actor %s.\n%s"), *PropertyName, *ActorRCLabel, *Error), Response->Body);
 	}
 
 	OnComplete(MoveTemp(Response));
diff --git a/Engine/Plugins/VirtualProduction/RemoteControl/Source/WebRemoteControl/Private/WebRemoteControlInternalUtils.h b/Engine/Plugins/VirtualProduction/RemoteControl/Source/WebRemoteControl/Private/WebRemoteControlInternalUtils.h
--- a/Engine/Plugins/VirtualProduction/RemoteControl/Source/WebRemoteControl/Private/WebRemoteControlInternalUtils.h
+++ b/Engine/Plugins/VirtualProduction/RemoteControl/Source/WebRemoteControl/Private/WebRemoteControlInternalUtils.h
@@ -379,7 +379,7 @@ namespace WebRemoteControlInternalUtils
 	 * @param Access The access mode to use for this operation.
 	 */
 	template <typename RequestType>
-	bool ModifyPropertyUsingPayload(FRemoteControlProperty& Property, const RequestType& Request, const TArrayView<uint8>& Payload, const FGuid& ClientId, FWebSocketMessageHandler& WebSocketHandler, ERCAccess Access)
+	bool ModifyPropertyUsingPayload(FRemoteControlProperty& Property, const RequestType& Request, const TArrayView<uint8>& Payload, const FGuid& ClientId, FWebSocketMessageHandler& WebSocketHandler, ERCAccess Access, FString* OutError = nullptr)
 	{
 		FRCObjectReference ObjectRef;
 
@@ -407,7 +407,7 @@ namespace WebRemoteControlInternalUtils
 
 		for (UObject* Object : Property.GetBoundObjects())
 		{
-			IRemoteControlModule::Get().ResolveObjectProperty(ObjectRef.Access, Object, Property.FieldPathInfo.ToString(), ObjectRef);
+			IRemoteControlModule::Get().ResolveObjectProperty(ObjectRef.Access, Object, Property.FieldPathInfo.ToString(), ObjectRef, OutError);
 
 			// Notify the handler before the change to ensure that the notification triggered by PostEditChange is ignored by the handler 
 			// if the client does not want remote change notifications.
