- Fixing PVS warnings in rivermax / remote control

diff --git a/Engine/Plugins/VirtualProduction/RemoteControl/Source/RemoteControl/Private/RemoteControlModule.cpp b/Engine/Plugins/VirtualProduction/RemoteControl/Source/RemoteControl/Private/RemoteControlModule.cpp
--- a/Engine/Plugins/VirtualProduction/RemoteControl/Source/RemoteControl/Private/RemoteControlModule.cpp
+++ b/Engine/Plugins/VirtualProduction/RemoteControl/Source/RemoteControl/Private/RemoteControlModule.cpp
@@ -159,8 +159,12 @@ namespace RemoteControlUtil
 		const bool bPropertyIsValid = InProperty != nullptr;
 		if (!bPropertyIsValid)
 		{
-			Error.AppendLine(FText::Format(LOCTEXT("RCPropertyNotValid", "{0} is not valid."), PreError));
-			*OutError = Error.ToText().ToString();
+			if (OutError)
+			{
+				Error.AppendLine(FText::Format(LOCTEXT("RCPropertyNotValid", "{0} is not valid."), PreError));
+				*OutError = Error.ToText().ToString();
+			}
+
 			return false;
 		}
 #if WITH_EDITOR
@@ -273,7 +277,10 @@ namespace RemoteControlUtil
 		const bool bIsAccessible = bObjectInGamePackage? bIsAccessibleInGamePackage : bIsAccessibleInNonGamePackage;
 
 		// Assign the error to be propagated after this call in case it failed
-		*OutError = Error.ToText().ToString();
+		if (OutError)
+		{
+			*OutError = Error.ToText().ToString();
+		}
 
 		return (bPropertyIsValid &&
 #if WITH_EDITOR
diff --git a/Engine/Plugins/VirtualProduction/Rivermax/RivermaxCore/Source/RivermaxCore/Private/RivermaxOutputStream.cpp b/Engine/Plugins/VirtualProduction/Rivermax/RivermaxCore/Source/RivermaxCore/Private/RivermaxOutputStream.cpp
--- a/Engine/Plugins/VirtualProduction/Rivermax/RivermaxCore/Source/RivermaxCore/Private/RivermaxOutputStream.cpp
+++ b/Engine/Plugins/VirtualProduction/Rivermax/RivermaxCore/Source/RivermaxCore/Private/RivermaxOutputStream.cpp
@@ -2074,7 +2074,7 @@ namespace UE::RivermaxCore::Private
 					if (UE::RivermaxCore::Private::GbTriggerRandomTimingIssue)
 					{
 						FRandomStream RandomStream(FPlatformTime::Cycles64());
-						const bool bTriggerDesync = (RandomStream.RandRange(0.0, 1.0) > 0.7) ? true : false;
+						const bool bTriggerDesync = (RandomStream.FRandRange(0.0, 1.0) > 0.7) ? true : false;
 						if (bTriggerDesync)
 						{
 							TRACE_CPUPROFILER_EVENT_SCOPE(RmaxOut::ForceTimingIssue);
diff --git a/Engine/Plugins/VirtualProduction/Rivermax/RivermaxSync/Source/RivermaxSync/Private/MediaOutputSynchronizationPolicyRivermax.cpp b/Engine/Plugins/VirtualProduction/Rivermax/RivermaxSync/Source/RivermaxSync/Private/MediaOutputSynchronizationPolicyRivermax.cpp
--- a/Engine/Plugins/VirtualProduction/Rivermax/RivermaxSync/Source/RivermaxSync/Private/MediaOutputSynchronizationPolicyRivermax.cpp
+++ b/Engine/Plugins/VirtualProduction/Rivermax/RivermaxSync/Source/RivermaxSync/Private/MediaOutputSynchronizationPolicyRivermax.cpp
@@ -151,7 +151,7 @@ void FMediaOutputSynchronizationPolicyRivermaxHandler::Synchronize()
 	if (UE::RivermaxSync::GbTriggerRandomDesync)
 	{
 		FRandomStream RandomStream(FPlatformTime::Cycles64());
-		const bool bTriggerDesync = (RandomStream.RandRange(0.0, 1.0) > 0.7) ? true : false;
+		const bool bTriggerDesync = (RandomStream.FRandRange(0.0, 1.0) > 0.7) ? true : false;
 		if (bTriggerDesync)
 		{
 			TRACE_CPUPROFILER_EVENT_SCOPE(RmaxSync::ForceBadSync);
