[RemoteControl]
Behaviour name lower case

diff --git a/Engine/Plugins/VirtualProduction/RemoteControl/Source/RemoteControlUI/Private/UI/Behaviour/RCBehaviourModel.cpp b/Engine/Plugins/VirtualProduction/RemoteControl/Source/RemoteControlUI/Private/UI/Behaviour/RCBehaviourModel.cpp
--- a/Engine/Plugins/VirtualProduction/RemoteControl/Source/RemoteControlUI/Private/UI/Behaviour/RCBehaviourModel.cpp
+++ b/Engine/Plugins/VirtualProduction/RemoteControl/Source/RemoteControlUI/Private/UI/Behaviour/RCBehaviourModel.cpp
@@ -26,7 +26,7 @@ FRCBehaviourModel::FRCBehaviourModel(URCBehaviour* InBehaviour
 
 	if (BehaviourWeakPtr.IsValid())
 	{
-		const FText BehaviorDisplayName = BehaviourWeakPtr->GetDisplayName().ToUpper();
+		const FText BehaviorDisplayName = BehaviourWeakPtr->GetDisplayName();
 		
 		SAssignNew(BehaviourTitleText, STextBlock)
 			.Text(BehaviorDisplayName)
