[RemoteControl]
Add grouping inside RC.

diff --git a/Engine/Plugins/VirtualProduction/RemoteControl/Source/RemoteControlUI/Private/UI/SRCPanelExposedEntitiesGroup.cpp b/Engine/Plugins/VirtualProduction/RemoteControl/Source/RemoteControlUI/Private/UI/SRCPanelExposedEntitiesGroup.cpp
new file mode 100644
--- /dev/null
+++ b/Engine/Plugins/VirtualProduction/RemoteControl/Source/RemoteControlUI/Private/UI/SRCPanelExposedEntitiesGroup.cpp
@@ -0,0 +1,135 @@
+// Copyright Epic Games, Inc. All Rights Reserved.
+
+#include "SRCPanelExposedEntitiesGroup.h"
+
+#include "RemoteControlPreset.h"
+#include "SRCPanelExposedField.h"
+#include "Widgets/Layout/SBox.h"
+#include "Widgets/Input/SEditableTextBox.h"
+
+void SRCPanelExposedEntitiesGroup::Construct(const FArguments& InArgs, EFieldGroupType InFieldGroupType, URemoteControlPreset* Preset)
+{
+	FieldKey = InArgs._FieldKey;
+	GroupType = InFieldGroupType;
+	PresetWeak = Preset;
+	OnGroupPropertyIdChangedDelegate = InArgs._OnGroupPropertyIdChanged;
+
+	const FMakeNodeWidgetArgs Args = CreateNodeWidgetArgs();
+	MakeNodeWidget(Args);
+}
+
+SRCPanelTreeNode::FMakeNodeWidgetArgs SRCPanelExposedEntitiesGroup::CreateNodeWidgetArgs()
+{
+	FMakeNodeWidgetArgs Args;
+
+	if (GroupType == EFieldGroupType::PropertyId)
+	{
+		Args.PropertyIdWidget = SNew(SBox)
+		[
+			SNew(SEditableTextBox)
+			.Justification(ETextJustify::Center)
+			.MinDesiredWidth(50.f)
+			.SelectAllTextWhenFocused(true)
+			.RevertTextOnEscape(true)
+			.ClearKeyboardFocusOnCommit(true)
+			.Text_Lambda([this] () { return FText::FromName(FieldKey); })
+			.OnTextCommitted(this, &SRCPanelExposedEntitiesGroup::OnPropertyIdTextCommitted)
+		];
+
+		Args.OwnerNameWidget = SNullWidget::NullWidget;
+	}
+	else if (GroupType == EFieldGroupType::Owner)
+	{
+		Args.PropertyIdWidget = SNullWidget::NullWidget;
+		Args.OwnerNameWidget = SNew(SBox)
+			.HeightOverride(25)
+			.VAlign(VAlign_Center)
+			[
+				SNew(STextBlock)
+				.Text(FText::FromName(FieldKey))
+			];
+	}
+	else
+	{
+		Args.PropertyIdWidget = SNullWidget::NullWidget;
+		Args.OwnerNameWidget = SNullWidget::NullWidget;
+	}
+
+	Args.NameWidget = SNew(SBox).HeightOverride(25).VAlign(VAlign_Center)
+		[
+			SNew(STextBlock)
+			.Text(FText::FromString("Group"))
+		];
+	Args.SubObjectPathWidget = SNullWidget::NullWidget;
+	Args.ValueWidget = SNullWidget::NullWidget;
+	Args.ResetButton = SNullWidget::NullWidget;
+
+	return Args;
+}
+
+void SRCPanelExposedEntitiesGroup::OnPropertyIdTextCommitted(const FText& InText, ETextCommit::Type InTextCommitType)
+{
+	if (!PresetWeak.IsValid() || PresetWeak.IsStale())
+	{
+		return;
+	}
+
+	const FName NewId = FName(InText.ToString());
+
+	if (FieldKey.Compare(NewId) == 0)
+	{
+		return;
+	}
+
+	URemoteControlPreset* Preset = PresetWeak.Get();
+	FieldKey = NewId;
+	for (const TSharedPtr<SRCPanelTreeNode>& Child : ChildWidgets)
+	{
+		TWeakPtr<FRemoteControlField> ExposedField = Preset->GetExposedEntity<FRemoteControlField>(Child->GetRCId());
+		if (ExposedField.IsValid())
+		{
+			ExposedField.Pin()->PropertyId = NewId;
+			Child->SetPropertyId(NewId);
+			Preset->UpdateIdentifiedField(ExposedField.Pin().ToSharedRef());
+		}
+	}
+	
+	OnGroupPropertyIdChangedDelegate.ExecuteIfBound();
+}
+
+void SRCPanelExposedEntitiesGroup::AssignChildren(const TArray<TSharedPtr<SRCPanelTreeNode>>& InFieldEntities)
+{
+	ChildWidgets.Empty();
+	if (GroupType == EFieldGroupType::PropertyId)
+	{
+		for (const TSharedPtr<SRCPanelTreeNode>& Entity : InFieldEntities)
+		{
+			if (const TSharedPtr<SRCPanelExposedField> ExposedField = StaticCastSharedPtr<SRCPanelExposedField>(Entity))
+			{
+				if (ExposedField->GetPropertyId() == FieldKey)
+				{
+					ChildWidgets.Add(Entity);
+				}
+			}
+		}
+	}
+	else if (GroupType == EFieldGroupType::Owner)
+	{
+		for (const TSharedPtr<SRCPanelTreeNode>& Entity : InFieldEntities)
+		{
+			if (const TSharedPtr<SRCPanelExposedField> ExposedField = StaticCastSharedPtr<SRCPanelExposedField>(Entity))
+			{
+				if (ExposedField->GetOwnerName() == FieldKey)
+				{
+					ChildWidgets.Add(Entity);
+				}
+			}
+		}
+	}
+}
+
+void SRCPanelExposedEntitiesGroup::GetNodeChildren(TArray<TSharedPtr<SRCPanelTreeNode>>& OutChildren) const
+{
+	OutChildren.Append(ChildWidgets);
+	SRCPanelTreeNode::GetNodeChildren(OutChildren);
+}
diff --git a/Engine/Plugins/VirtualProduction/RemoteControl/Source/RemoteControlUI/Private/UI/SRCPanelExposedEntitiesGroup.h b/Engine/Plugins/VirtualProduction/RemoteControl/Source/RemoteControlUI/Private/UI/SRCPanelExposedEntitiesGroup.h
new file mode 100644
--- /dev/null
+++ b/Engine/Plugins/VirtualProduction/RemoteControl/Source/RemoteControlUI/Private/UI/SRCPanelExposedEntitiesGroup.h
@@ -0,0 +1,87 @@
+// Copyright Epic Games, Inc. All Rights Reserved.
+
+#pragma once
+
+#include "CoreMinimal.h"
+#include "SRCPanelExposedEntity.h"
+#include "SRCPanelTreeNode.h"
+#include "Widgets/SCompoundWidget.h"
+
+/** Grouping type available */
+enum class EFieldGroupType
+{
+	None,
+	PropertyId,
+	Owner
+};
+
+ENUM_CLASS_FLAGS(EFieldGroupType)
+
+/** Delegate called when the PropertyId of a group change */
+DECLARE_DELEGATE(FOnGroupPropertyIdChanged)
+
+/** Widget that is used as a parent for field widget */
+class SRCPanelExposedEntitiesGroup : public SRCPanelTreeNode
+{
+	SLATE_BEGIN_ARGS(SRCPanelExposedEntitiesGroup)
+	{}
+
+	SLATE_EVENT(FOnGroupPropertyIdChanged, OnGroupPropertyIdChanged)
+	SLATE_ARGUMENT(FName, FieldKey)
+
+	SLATE_END_ARGS()
+
+	//~ Begin SRCPanelTreeNode interface
+
+	/** Get get this node's type. */
+	virtual ENodeType GetRCType() const override { return ENodeType::FieldGroup; }
+
+	//~ End SRCPanelTreeNode interface
+
+	void Construct(const FArguments& InArgs, EFieldGroupType InFieldGroupType, URemoteControlPreset* Preset);
+
+	/**
+	 * Take the current fields in RC and assign them to its child widgets based on the group key
+	 * @param InFieldEntities Current fields exposed in RC
+	 */
+	void AssignChildren(const TArray<TSharedPtr<SRCPanelTreeNode>>& InFieldEntities);
+
+	/** Get this tree node's children. */
+	virtual void GetNodeChildren(TArray<TSharedPtr<SRCPanelTreeNode>>& OutChildren) const override;
+
+	/** Returns true if this tree node has children. */
+	virtual bool HasChildren() const override { return true; }
+
+	/** Get the group type of this group */
+	EFieldGroupType GetGroupType() const { return GroupType; }
+
+	/** Get the field key of this group */
+	FName GetFieldKey() const { return FieldKey; }
+
+private:
+	/** Create the node Args for this class that will be used to create the widget */
+	FMakeNodeWidgetArgs CreateNodeWidgetArgs();
+
+	/**
+	 * Called when the text label is renamed
+	 * @param InText New Text
+	 * @param InTextCommitType Type of commit, look at ETextCommit for more information
+	 */
+	void OnPropertyIdTextCommitted(const FText& InText, ETextCommit::Type InTextCommitType);
+
+private:
+	/** Child widgets of this group */
+	TArray<TSharedPtr<SRCPanelTreeNode>> ChildWidgets;
+
+	/** Field key value of this group */
+	FName FieldKey;
+
+	/** Field group type of this group */
+	EFieldGroupType GroupType = EFieldGroupType::None;
+
+	/** Weak ptr to the preset where this group reside */
+	TWeakObjectPtr<URemoteControlPreset> PresetWeak;
+
+	/** Delegate called when the group type is PropertyId and the value change */
+	FOnGroupPropertyIdChanged OnGroupPropertyIdChangedDelegate;
+};
diff --git a/Engine/Plugins/VirtualProduction/RemoteControl/Source/RemoteControlUI/Private/UI/SRCPanelExposedEntitiesList.cpp b/Engine/Plugins/VirtualProduction/RemoteControl/Source/RemoteControlUI/Private/UI/SRCPanelExposedEntitiesList.cpp
--- a/Engine/Plugins/VirtualProduction/RemoteControl/Source/RemoteControlUI/Private/UI/SRCPanelExposedEntitiesList.cpp
+++ b/Engine/Plugins/VirtualProduction/RemoteControl/Source/RemoteControlUI/Private/UI/SRCPanelExposedEntitiesList.cpp
@@ -27,6 +27,7 @@
 #include "ScopedTransaction.h"
 #include "SRCHeaderRow.h"
 #include "SRCModeSwitcher.h"
+#include "SRCPanelExposedEntitiesGroup.h"
 #include "SRCPanelFieldGroup.h"
 #include "SRCPanelExposedField.h"
 #include "SSearchToggleButton.h"
@@ -34,9 +35,11 @@
 #include "UI/Panels/SRCDockPanel.h"
 #include "UObject/Object.h"
 #include "Widgets/Input/SButton.h"
+#include "Widgets/Input/SComboBox.h"
+#include "Widgets/Input/SComboButton.h"
 #include "Widgets/Input/SSearchBox.h"
-#include "Widgets/Layout/SSpacer.h"
 #include "Widgets/Layout/SSeparator.h"
+#include "Widgets/Layout/SSpacer.h"
 #include "Widgets/Layout/SWidgetSwitcher.h"
 #include "Widgets/Views/SListView.h"
 #include "Widgets/Views/STableRow.h"
@@ -162,7 +165,7 @@ private:
 	//~ SWidget Interface
 	virtual FReply OnMouseButtonDoubleClick(const FGeometry& InMyGeometry, const FPointerEvent& InMouseEvent) override
 	{
-		if (Entity.IsValid())
+		if (Entity.IsValid() && Entity->GetRCType() != SRCPanelTreeNode::FieldGroup)
 		{
 			if (const TSharedPtr<SRCPanelExposedEntity> ExposedEntityWidget = StaticCastSharedPtr<SRCPanelExposedEntity>(Entity))
 			{
@@ -486,9 +489,36 @@ void SRCPanelExposedEntitiesList::Construct(const FArguments& InArgs, URemoteCon
 		.OnTextCommitted(this, &SRCPanelExposedEntitiesList::OnSearchTextCommitted)
 		.DelayChangeNotificationsWhileTyping(true);
 
+	// Create grouping button
+	ComboButtonGroupButton = SNew(SComboButton)
+		.ToolTipText(LOCTEXT("RCFieldsGroupingTooltip", "Select grouping type for the fields"))
+		.HAlign(HAlign_Center)
+		.VAlign(VAlign_Center)
+		.ButtonStyle(&RCPanelStyle->FlatButtonStyle)
+		.ForegroundColor(FSlateColor::UseForeground())
+		.CollapseMenuOnParentFocus(true)
+		.HasDownArrow(false)
+		.ContentPadding(FMargin(4.f, 2.f))
+		.ButtonContent()
+		[
+			SNew(SBox)
+			.WidthOverride(RCPanelStyle->IconSize.X)
+			.HeightOverride(RCPanelStyle->IconSize.Y)
+			[
+				SNew(SImage)
+				.ColorAndOpacity(FSlateColor::UseForeground())
+				.Image(FAppStyle::GetBrush("InputBindingEditor.LevelViewport"))
+			]
+		]
+		.MenuContent()
+		[
+			GetGroupMenuContentWidget()
+		];
+
 	// Expose Button
 	ExposeDockPanel->AddHeaderToolbarItem(EToolbar::Center, SearchBoxPtr.ToSharedRef());
 	ExposeDockPanel->AddHeaderToolbarItem(EToolbar::Right, SNew(SSearchToggleButton, SearchBoxPtr.ToSharedRef()));
+	ExposeDockPanel->AddHeaderToolbarItem(EToolbar::Right, ComboButtonGroupButton.ToSharedRef());
 	ExposeDockPanel->AddHeaderToolbarItem(EToolbar::Right, FilterComboButton.ToSharedRef());
 	ExposeDockPanel->AddHeaderToolbarItem(EToolbar::Right, ModeSwitcher.ToSharedRef());
 	ExposeDockPanel->AddHeaderToolbarItem(EToolbar::Left, InArgs._ExposeActorsComboButton.Get().ToSharedRef());
@@ -547,6 +577,12 @@ void SRCPanelExposedEntitiesList::Tick(const FGeometry& AllottedGeometry, const
 		bRefreshRequested = false;
 	}
 
+	if (bRefreshEntitiesGroups)
+	{
+		RefreshGroupsAndRestoreExpansions();
+		bRefreshEntitiesGroups = false;
+	}
+
 	ExposedEntitiesNodesRefresh();
 }
 
@@ -755,6 +791,12 @@ void SRCPanelExposedEntitiesList::OnPropertyIdRenamed(const FName InNewId, TShar
 			}
 		}
 	}
+
+	if (CurrentGroupType == EFieldGroupType::PropertyId)
+	{
+		CreateFieldGroup();
+		OrderGroups();
+	}
 }
 
 void SRCPanelExposedEntitiesList::OnLabelModified(const FName InOldName, const FName InNewName)
@@ -963,6 +1005,7 @@ void SRCPanelExposedEntitiesList::GenerateListWidgets(const FRemoteControlPreset
 	TRACE_CPUPROFILER_EVENT_SCOPE(SRCPanelExposedEntitiesList::GenerateListWidgetsFromGroup);
 
 	FieldEntities.Reset();
+	CachedFieldEntities.Reset();
 
 	if (FieldWidgetMap.IsEmpty())
 	{
@@ -973,11 +1016,13 @@ void SRCPanelExposedEntitiesList::GenerateListWidgets(const FRemoteControlPreset
 
 	if (Preset->Layout.IsDefaultGroup(FromGroup.Id))
 	{
-		FieldWidgetMap.GenerateValueArray(FieldEntities);
+		FieldWidgetMap.GenerateValueArray(CachedFieldEntities);
+		CreateGroupsAndSort();
 	}
 	else if (TSharedPtr<SRCPanelGroup> SRCGroup = FindGroupById(FromGroup.Id))
 	{
-		SRCGroup->GetNodeChildren(FieldEntities);
+		SRCGroup->GetNodeChildren(CachedFieldEntities);
+		CreateGroupsAndSort();
 	}
 
 	FieldsListView->RequestListRefresh();
@@ -1108,7 +1153,8 @@ FReply SRCPanelExposedEntitiesList::OnDropOnGroup(const TSharedPtr<FDragDropOper
 {
 	checkSlow(DragTargetGroup);
 
-	if (DragDropOperation->IsOfType<FExposedEntityDragDrop>())
+	// will be changed, for now as long as you have a grouping type active you will not be able to drag fields to re-order them
+	if (DragDropOperation->IsOfType<FExposedEntityDragDrop>() && CurrentGroupType == EFieldGroupType::None)
 	{
 		if (TSharedPtr<FExposedEntityDragDrop> DragDropOp = StaticCastSharedPtr<FExposedEntityDragDrop>(DragDropOperation))
 		{
@@ -1335,6 +1381,217 @@ TSharedPtr<SWidget> SRCPanelExposedEntitiesList::OnContextMenuOpening(SRCPanelTr
 	return nullptr;
 }
 
+TSharedRef<SWidget> SRCPanelExposedEntitiesList::GetGroupMenuContentWidget()
+{
+	constexpr bool bShouldCloseWindowAfterMenuSelection = false;
+	FMenuBuilder MenuBuilder(bShouldCloseWindowAfterMenuSelection, nullptr);
+
+	// Grouping type section
+	MenuBuilder.BeginSection(FName(TEXT("RCGroupingType")), LOCTEXT("RCGroupingTypeHeader", "Group by"));
+
+	MenuBuilder.AddMenuEntry(
+		LOCTEXT("RCGroupPropertyIdLabel", "Property Id"),
+		LOCTEXT("RCGroupPropertyIdTooltip", "Group by Property Id"),
+		FSlateIcon(),
+		FUIAction(
+		FExecuteAction::CreateSP(this, &SRCPanelExposedEntitiesList::OnCreateFieldGroup, EFieldGroupType::PropertyId),
+		FCanExecuteAction(),
+		FIsActionChecked::CreateLambda([this]() { return CurrentGroupType == EFieldGroupType::PropertyId; })
+		),
+		NAME_None,
+		EUserInterfaceActionType::RadioButton);
+
+	MenuBuilder.AddMenuEntry(
+		LOCTEXT("RCGroupOwnerLabel", "Owner"),
+		LOCTEXT("RCGroupOwnerTooltip", "Group by Owner"),
+		FSlateIcon(),
+		FUIAction(
+		FExecuteAction::CreateSP(this, &SRCPanelExposedEntitiesList::OnCreateFieldGroup, EFieldGroupType::Owner),
+		FCanExecuteAction(),
+		FIsActionChecked::CreateLambda([this]() { return CurrentGroupType == EFieldGroupType::Owner; })
+		),
+		NAME_None,
+		EUserInterfaceActionType::RadioButton);
+
+	MenuBuilder.EndSection();
+
+	// Ordering type section
+	MenuBuilder.BeginSection(FName(TEXT("RCSortingType")), LOCTEXT("RCSortingTypeHeader", "Sort by"));
+
+	MenuBuilder.AddMenuEntry(
+		LOCTEXT("RCSortAscendingGroupLabel", "Ascending"),
+		LOCTEXT("RCSortAscendingGroupTooltip", "Sort groups ascending"),
+		FSlateIcon(),
+		FUIAction(
+		FExecuteAction::CreateSP(this, &SRCPanelExposedEntitiesList::OnGroupOrderChanged, ERCGroupOrder::Ascending),
+		FCanExecuteAction(),
+		FIsActionChecked::CreateLambda([this] () { return CurrentGroupSortType == ERCGroupOrder::Ascending; })
+		),
+		NAME_None,
+		EUserInterfaceActionType::RadioButton);
+
+	MenuBuilder.AddMenuEntry(
+		LOCTEXT("RCSortDescendingGroupLabel", "Descending"),
+		LOCTEXT("RCSortDescendingGroupTooltip", "Sort groups descending"),
+		FSlateIcon(),
+		FUIAction(
+		FExecuteAction::CreateSP(this, &SRCPanelExposedEntitiesList::OnGroupOrderChanged, ERCGroupOrder::Descending),
+		FCanExecuteAction(),
+		FIsActionChecked::CreateLambda([this]() { return CurrentGroupSortType == ERCGroupOrder::Descending; })
+		),
+		NAME_None,
+		EUserInterfaceActionType::RadioButton);
+
+	MenuBuilder.EndSection();
+
+	return MenuBuilder.MakeWidget();
+}
+
+void SRCPanelExposedEntitiesList::OnCreateFieldGroup(EFieldGroupType InFieldGroupType)
+{
+	CurrentGroupType = CurrentGroupType == InFieldGroupType? EFieldGroupType::None : InFieldGroupType;
+	CreateFieldGroup();
+	OrderGroups();
+}
+
+void SRCPanelExposedEntitiesList::CreateFieldGroup()
+{
+	FieldEntitiesGroups.Reset();
+	TSet<FName> GroupKeys;
+	if (CurrentGroupType == EFieldGroupType::PropertyId)
+	{
+		for (TSharedPtr<SRCPanelTreeNode> Entity : CachedFieldEntities)
+		{
+			if (TSharedPtr<SRCPanelExposedField> ExposedField = StaticCastSharedPtr<SRCPanelExposedField>(Entity))
+			{
+				if (!GroupKeys.Contains(ExposedField->GetPropertyId()))
+				{
+					GroupKeys.Add(ExposedField->GetPropertyId());
+					TSharedRef<SRCPanelExposedEntitiesGroup> GroupWidget = SNew(SRCPanelExposedEntitiesGroup, CurrentGroupType, Preset.Get())
+						.FieldKey(ExposedField->GetPropertyId())
+						.OnGroupPropertyIdChanged(this, &SRCPanelExposedEntitiesList::CreateGroupsAndSort);
+
+					GroupWidget->AssignChildren(CachedFieldEntities);
+					FieldEntitiesGroups.Add(GroupWidget);
+				}
+			}
+		}
+	}
+	else if (CurrentGroupType == EFieldGroupType::Owner)
+	{
+		for (TSharedPtr<SRCPanelTreeNode> Entity : CachedFieldEntities)
+		{
+			if (TSharedPtr<SRCPanelExposedField> ExposedField = StaticCastSharedPtr<SRCPanelExposedField>(Entity))
+			{
+				if (!GroupKeys.Contains(ExposedField->GetOwnerName()))
+				{
+					GroupKeys.Add(ExposedField->GetOwnerName());
+					TSharedRef<SRCPanelExposedEntitiesGroup> GroupWidget = SNew(SRCPanelExposedEntitiesGroup, CurrentGroupType, Preset.Get())
+						.FieldKey(ExposedField->GetOwnerName());
+
+					GroupWidget->AssignChildren(CachedFieldEntities);
+					FieldEntitiesGroups.Add(GroupWidget);
+				}
+			}
+		}
+	}
+
+	if (CurrentGroupType == EFieldGroupType::None)
+	{
+		FieldEntities = CachedFieldEntities;
+		FieldsListView->RequestTreeRefresh();
+	}
+	else if (FieldEntitiesGroups.Num())
+	{
+		FieldEntities.Empty();
+		for (const TSharedPtr<SRCPanelExposedEntitiesGroup>& Group : FieldEntitiesGroups)
+		{
+			FieldEntities.Add(Group);
+		}
+
+		bRefreshEntitiesGroups = true;
+	}
+}
+
+void SRCPanelExposedEntitiesList::OnGroupOrderChanged(ERCGroupOrder InGroupOrder)
+{
+	CurrentGroupSortType = CurrentGroupSortType == InGroupOrder? ERCGroupOrder::None : InGroupOrder;
+	OrderGroups();
+}
+
+void SRCPanelExposedEntitiesList::OrderGroups()
+{
+	if (CurrentGroupType == EFieldGroupType::None)
+	{
+		return;
+	}
+
+	if (CurrentGroupSortType == ERCGroupOrder::Ascending)
+	{
+		FieldEntities.Sort([](const TSharedPtr<SRCPanelTreeNode>& InFirst, const TSharedPtr<SRCPanelTreeNode>& InSecond)
+		{
+			if (const TSharedPtr<SRCPanelExposedEntitiesGroup>& FieldGroupFirst = StaticCastSharedPtr<SRCPanelExposedEntitiesGroup>(InFirst))
+			{
+				if (const TSharedPtr<SRCPanelExposedEntitiesGroup>& FieldGroupSecond = StaticCastSharedPtr<SRCPanelExposedEntitiesGroup>(InSecond))
+				{
+					return FieldGroupFirst->GetFieldKey().Compare(FieldGroupSecond->GetFieldKey()) <= 0;
+				}
+			}
+			return false;
+		});
+	}
+	else if (CurrentGroupSortType == ERCGroupOrder::Descending)
+	{
+		FieldEntities.Sort([](const TSharedPtr<SRCPanelTreeNode>& InFirst, const TSharedPtr<SRCPanelTreeNode>& InSecond)
+		{
+			if (const TSharedPtr<SRCPanelExposedEntitiesGroup>& FieldGroupFirst = StaticCastSharedPtr<SRCPanelExposedEntitiesGroup>(InFirst))
+			{
+				if (const TSharedPtr<SRCPanelExposedEntitiesGroup>& FieldGroupSecond = StaticCastSharedPtr<SRCPanelExposedEntitiesGroup>(InSecond))
+				{
+					return FieldGroupFirst->GetFieldKey().Compare(FieldGroupSecond->GetFieldKey()) > 0;
+				}
+			}
+			return false;
+		});
+	}
+
+	bRefreshEntitiesGroups = true;
+}
+
+void SRCPanelExposedEntitiesList::RefreshGroupsAndRestoreExpansions()
+{
+	TSet<TSharedPtr<SRCPanelTreeNode>> ExpandedItems;
+	FieldsListView->GetExpandedItems(ExpandedItems);
+	FieldsListView->RequestTreeRefresh();
+	for (const TSharedPtr<SRCPanelTreeNode>& ExpandedItem : ExpandedItems)
+	{
+		if (const TSharedPtr<SRCPanelExposedEntitiesGroup>& FieldGroup = StaticCastSharedPtr<SRCPanelExposedEntitiesGroup>(ExpandedItem))
+		{
+			if (FieldGroup->GetGroupType() == CurrentGroupType)
+			{
+				for (const TSharedPtr<SRCPanelExposedEntitiesGroup>& Group : FieldEntitiesGroups)
+				{
+					if (Group->GetFieldKey() == FieldGroup->GetFieldKey())
+					{
+						FieldsListView->SetItemExpansion(Group, true);
+						break;
+					}
+				}
+			}
+		}
+		else
+		{
+			FieldsListView->SetItemExpansion(ExpandedItem, true);
+		}
+	}
+}
+
+void SRCPanelExposedEntitiesList::CreateGroupsAndSort()
+{
+	CreateFieldGroup();
+	OrderGroups();
+}
+
 void SRCPanelExposedEntitiesList::RegisterPresetDelegates()
 {
 	FRemoteControlPresetLayout& Layout = Preset->Layout;
diff --git a/Engine/Plugins/VirtualProduction/RemoteControl/Source/RemoteControlUI/Private/UI/SRCPanelExposedEntitiesList.h b/Engine/Plugins/VirtualProduction/RemoteControl/Source/RemoteControlUI/Private/UI/SRCPanelExposedEntitiesList.h
--- a/Engine/Plugins/VirtualProduction/RemoteControl/Source/RemoteControlUI/Private/UI/SRCPanelExposedEntitiesList.h
+++ b/Engine/Plugins/VirtualProduction/RemoteControl/Source/RemoteControlUI/Private/UI/SRCPanelExposedEntitiesList.h
@@ -5,6 +5,7 @@
 #include "IRemoteControlUIModule.h"
 #include "Misc/TextFilter.h"
 #include "RemoteControlPreset.h"
+#include "SRCPanelExposedEntitiesGroup.h"
 #include "SRCPanelTreeNode.h"
 #include "Widgets/DeclarativeSyntaxSupport.h"
 #include "Widgets/SCompoundWidget.h"
@@ -21,6 +22,7 @@ struct FRemoteControlProperty;
 struct FRemoteControlPresetGroup;
 struct FRemoteControlFunction;
 class ITableRow;
+class SComboButton;
 class SRCHeaderRow;
 class SRCPanelFilter;
 class SRCPanelGroup;
@@ -38,6 +40,14 @@ enum class EEntitiesListMode : uint8
 	Protocols
 };
 
+/** Ordering types while grouping is active */
+enum class ERCGroupOrder
+{
+	None,
+	Ascending,
+	Descending
+};
+
 /** Holds information about a group drag and drop event  */
 struct FGroupDragEvent
 {
@@ -172,6 +182,32 @@ private:
 	/** Handle context menu opening on a row. */
 	TSharedPtr<SWidget> OnContextMenuOpening(SRCPanelTreeNode::ENodeType InType);
 
+	TSharedRef<SWidget> GetGroupMenuContentWidget();
+
+	/**
+	 * Called when the group type changed, if the new group type is the same as the current one it will be set to none
+	 * @param InFieldGroupType New grouping type
+	 */
+	void OnCreateFieldGroup(EFieldGroupType InFieldGroupType);
+
+	/** Group fields based on the current group type (PropertyId/Owner) */
+	void CreateFieldGroup();
+
+	/**
+	 * Called when the order type changed, if the new order type is the same as the current one it will be set to none
+	 * @param InGroupOrder New group order
+	 */
+	void OnGroupOrderChanged(ERCGroupOrder InGroupOrder);
+
+	/** Order the groups based on the current order assigned (Ascending/Descending) */
+	void OrderGroups();
+
+	/** Refresh the fields groups and restore the expansion */
+	void RefreshGroupsAndRestoreExpansions();
+
+	/** Calls both the CreateFieldGroup and the OrderGroups */
+	void CreateGroupsAndSort();
+
 	//~ Register and handle preset delegates.
 	void RegisterPresetDelegates();
 	void UnregisterPresetDelegates();
@@ -235,6 +271,16 @@ private:
 	TArray<TSharedPtr<SRCPanelGroup>> FieldGroups;
 	/** Holds all the field entities. */
 	TArray<TSharedPtr<SRCPanelTreeNode>> FieldEntities;
+	/** Cached field entities used to store the original Entities when switching groups. */
+	TArray<TSharedPtr<SRCPanelTreeNode>> CachedFieldEntities;
+	/** Holds all the exposed entities groups. */
+	TArray<TSharedPtr<SRCPanelExposedEntitiesGroup>> ExposedEntitiesGroups;
+	/** Holds the current group type */
+	EFieldGroupType CurrentGroupType = EFieldGroupType::None;
+	/** Holds the current sorting type */
+	ERCGroupOrder CurrentGroupSortType = ERCGroupOrder::None;
+	/** Holds all the entities groups currently in the list */
+	TArray<TSharedPtr<SRCPanelExposedEntitiesGroup>> FieldEntitiesGroups;
 	/** Map of field ids to field widgets. */
 	TMap<FGuid, TSharedPtr<SRCPanelTreeNode>> FieldWidgetMap;
 	/** Whether the panel is in live mode. */
@@ -271,6 +317,8 @@ private:
 	TSharedPtr<SRCPanelFilter> FilterPtr;
 	/** The text box used to search for tags. */
 	TSharedPtr<SSearchBox> SearchBoxPtr;
+	/** Button giving you the possibilities for grouping. */
+	TSharedPtr<SComboButton> ComboButtonGroupButton;
 	/** Text filter for the search text. */
     TSharedPtr<TTextFilter<const SRCPanelTreeNode&>> SearchTextFilter;
     /** Actively searched term. */
@@ -285,6 +333,7 @@ private:
 	FGuid CurrentlySelectedGroup;
 
 	bool bRefreshRequested = false;
+	bool bRefreshEntitiesGroups = false;
 
 	/** When true, widgets of the Exposed Entities List will be refreshed on Tick */
 	bool bNodesRefreshRequested = false;
diff --git a/Engine/Plugins/VirtualProduction/RemoteControl/Source/RemoteControlUI/Private/UI/SRCPanelExposedEntity.cpp b/Engine/Plugins/VirtualProduction/RemoteControl/Source/RemoteControlUI/Private/UI/SRCPanelExposedEntity.cpp
--- a/Engine/Plugins/VirtualProduction/RemoteControl/Source/RemoteControlUI/Private/UI/SRCPanelExposedEntity.cpp
+++ b/Engine/Plugins/VirtualProduction/RemoteControl/Source/RemoteControlUI/Private/UI/SRCPanelExposedEntity.cpp
@@ -431,16 +431,20 @@ void SRCPanelExposedEntity::OnLabelCommitted(const FText& InLabel, ETextCommit::
 
 void SRCPanelExposedEntity::OnPropertyIdTextCommitted(const FText& InText, ETextCommit::Type InCommitInfo)
 {
-	if (URemoteControlPreset* RCPreset = Preset.Get())
+	if (const URemoteControlPreset* RCPreset = Preset.Get())
 	{
 		if (const TSharedPtr<FRemoteControlEntity> RCEntity = GetEntity())
 		{
 			if (const TSharedPtr<FRemoteControlField> RCField = StaticCastSharedPtr<FRemoteControlField>(RCEntity))
 			{
-				RCField->PropertyId = FName(InText.ToString());
-				PropertyIdLabel = RCField->PropertyId;
-				RCPreset->UpdateIdentifiedField(RCField.ToSharedRef());
-				OnPropertyIdRenamed().ExecuteIfBound(PropertyIdLabel);
+				const FName NewId = FName(InText.ToString());
+				if (RCField->PropertyId.Compare(NewId) != 0)
+				{
+					RCField->PropertyId = NewId;
+					PropertyIdLabel = NewId;
+					RCPreset->UpdateIdentifiedField(RCField.ToSharedRef());
+					OnPropertyIdRenamed().ExecuteIfBound(NewId);
+				}
 			}
 		}
 	}
diff --git a/Engine/Plugins/VirtualProduction/RemoteControl/Source/RemoteControlUI/Private/UI/SRCPanelExposedField.cpp b/Engine/Plugins/VirtualProduction/RemoteControl/Source/RemoteControlUI/Private/UI/SRCPanelExposedField.cpp
--- a/Engine/Plugins/VirtualProduction/RemoteControl/Source/RemoteControlUI/Private/UI/SRCPanelExposedField.cpp
+++ b/Engine/Plugins/VirtualProduction/RemoteControl/Source/RemoteControlUI/Private/UI/SRCPanelExposedField.cpp
@@ -177,6 +177,11 @@ FName SRCPanelExposedField::GetFieldLabel() const
 	return CachedLabel;
 }
 
+FName SRCPanelExposedField::GetOwnerName() const
+{
+	return CachedOwnerName;
+}
+
 EExposedFieldType SRCPanelExposedField::GetFieldType() const
 {
 	if (TSharedPtr<FRemoteControlField> Field = WeakField.Pin())
diff --git a/Engine/Plugins/VirtualProduction/RemoteControl/Source/RemoteControlUI/Private/UI/SRCPanelExposedField.h b/Engine/Plugins/VirtualProduction/RemoteControl/Source/RemoteControlUI/Private/UI/SRCPanelExposedField.h
--- a/Engine/Plugins/VirtualProduction/RemoteControl/Source/RemoteControlUI/Private/UI/SRCPanelExposedField.h
+++ b/Engine/Plugins/VirtualProduction/RemoteControl/Source/RemoteControlUI/Private/UI/SRCPanelExposedField.h
@@ -73,6 +73,9 @@ struct SRCPanelExposedField : public SRCPanelExposedEntity
 	/** Returns this widget's underlying objects. */
 	void GetBoundObjects(TSet<UObject*>& OutBoundObjects) const;
 
+	/** Get the owner name of this field */
+	FName GetOwnerName() const;
+
 private:
 	/** Construct a property widget. */
 	TSharedRef<SWidget> ConstructWidget();
diff --git a/Engine/Plugins/VirtualProduction/RemoteControl/Source/RemoteControlUI/Private/UI/SRCPanelTreeNode.h b/Engine/Plugins/VirtualProduction/RemoteControl/Source/RemoteControlUI/Private/UI/SRCPanelTreeNode.h
--- a/Engine/Plugins/VirtualProduction/RemoteControl/Source/RemoteControlUI/Private/UI/SRCPanelTreeNode.h
+++ b/Engine/Plugins/VirtualProduction/RemoteControl/Source/RemoteControlUI/Private/UI/SRCPanelTreeNode.h
@@ -53,7 +53,8 @@ struct SRCPanelTreeNode : public SCompoundWidget, public IHasProtocolExtensibili
 		Field,
 		FieldChild,
 		Actor,
-		Material
+		Material,
+		FieldGroup
 	};
 
 	virtual ~SRCPanelTreeNode() {}
diff --git a/Engine/Plugins/VirtualProduction/RemoteControl/Source/RemoteControlUI/Private/UI/SRemoteControlPanel.cpp b/Engine/Plugins/VirtualProduction/RemoteControl/Source/RemoteControlUI/Private/UI/SRemoteControlPanel.cpp
--- a/Engine/Plugins/VirtualProduction/RemoteControl/Source/RemoteControlUI/Private/UI/SRemoteControlPanel.cpp
+++ b/Engine/Plugins/VirtualProduction/RemoteControl/Source/RemoteControlUI/Private/UI/SRemoteControlPanel.cpp
@@ -1767,7 +1767,8 @@ void SRemoteControlPanel::UpdateEntityDetailsView(const TSharedPtr<SRCPanelTreeN
 	if (LastSelectedEntity)
 	{
 		if (LastSelectedEntity->GetRCType() != SRCPanelTreeNode::Group &&
-			LastSelectedEntity->GetRCType() != SRCPanelTreeNode::FieldChild) // Field Child does not contain entity ID, that is why it should not be processed
+			LastSelectedEntity->GetRCType() != SRCPanelTreeNode::FieldChild &&
+			LastSelectedEntity->GetRCType() != SRCPanelTreeNode::FieldGroup) // Field Child does not contain entity ID, that is why it should not be processed
 		{
 			const TSharedPtr<FRemoteControlEntity> Entity = Preset->GetExposedEntity<FRemoteControlEntity>(LastSelectedEntity->GetRCId()).Pin();
 			SelectedEntityPtr = RemoteControlPanelUtils::GetEntityOnScope(Entity, Preset->GetExposedEntityType(LastSelectedEntity->GetRCId()));
