[RemoteControl]
Add button in the action to add all selected properties

diff --git a/Engine/Plugins/VirtualProduction/RemoteControl/Source/RemoteControlUI/Private/UI/Action/SRCActionPanel.cpp b/Engine/Plugins/VirtualProduction/RemoteControl/Source/RemoteControlUI/Private/UI/Action/SRCActionPanel.cpp
--- a/Engine/Plugins/VirtualProduction/RemoteControl/Source/RemoteControlUI/Private/UI/Action/SRCActionPanel.cpp
+++ b/Engine/Plugins/VirtualProduction/RemoteControl/Source/RemoteControlUI/Private/UI/Action/SRCActionPanel.cpp
@@ -30,6 +30,7 @@
 #include "UI/Behaviour/RCBehaviourModel.h"
 #include "UI/Panels/SRCDockPanel.h"
 #include "UI/RemoteControlPanelStyle.h"
+#include "UI/SRCPanelExposedEntitiesList.h"
 #include "UI/SRemoteControlPanel.h"
 
 #include "Widgets/Input/SButton.h"
@@ -152,7 +153,29 @@ void SRCActionPanel::UpdateWrappedWidget(TSharedPtr<FRCBehaviourModel> InBehavio
 				]
 			];
 
+		// Add All Button
+		const TSharedRef<SWidget> AddAllSelectedActionsButton = SNew(SButton)
+			.AddMetaData<FTagMetaData>(FTagMetaData(TEXT("Add All Selected Fields")))
+			.HAlign(HAlign_Center)
+			.VAlign(VAlign_Center)
+			.ForegroundColor(FSlateColor::UseForeground())
+			.ButtonStyle(&RCPanelStyle->FlatButtonStyle)
+			.ToolTipText(LOCTEXT("RCAddAllSelectedToolTip", "Adds all the selected fields."))
+			.OnClicked(this, &SRCActionPanel::OnAddAllSelectedFields)
+			.Visibility(this, &SRCActionPanel::HandleAddAllButtonVisibility)
+			[
+				SNew(SBox)
+				.WidthOverride(RCPanelStyle->IconSize.X)
+				.HeightOverride(RCPanelStyle->IconSize.Y)
+				[
+					SNew(SImage)
+					.ColorAndOpacity(FSlateColor::UseForeground())
+					.Image(FAppStyle::GetBrush("DataLayerBrowser.AddSelection"))
+				]
+			];
+
 		ActionDockPanel->AddHeaderToolbarItem(EToolbar::Left, AddNewActionButton);
+		ActionDockPanel->AddHeaderToolbarItem(EToolbar::Right, AddAllSelectedActionsButton);
 		ActionDockPanel->AddHeaderToolbarItem(EToolbar::Right, AddAllActionsButton);
 
 		// Header Dock Panel
@@ -446,6 +469,66 @@ FReply SRCActionPanel::OnAddAllFields()
 	return FReply::Handled();
 }
 
+FReply SRCActionPanel::OnAddAllSelectedFields()
+{
+	if (!SelectedBehaviourItemWeakPtr.IsValid() || !PanelWeakPtr.IsValid())
+	{
+		return FReply::Handled();
+	}
+
+	const TSharedPtr<SRemoteControlPanel> RCPanel = PanelWeakPtr.Pin();
+	if (const TSharedPtr<SRCPanelExposedEntitiesList> RCEntitiesList = RCPanel->GetEntityList())
+	{
+		if (const TSharedPtr<FRCBehaviourModel>& BehaviourItem = SelectedBehaviourItemWeakPtr.Pin())
+		{
+			FScopedTransaction Transaction(LOCTEXT("RCAddAllSelectedActionsTransaction", "Add All Selected Fields"));
+
+			auto AddSelectedActionLambda = [this, BehaviourItem] (const TSharedPtr<SRCPanelTreeNode>& RCEntity)
+			{
+				if (URemoteControlPreset* Preset = GetPreset())
+				{
+					const TWeakPtr<FRemoteControlField> RCWeakField = Preset->GetExposedEntity<FRemoteControlField>(RCEntity->GetRCId());
+					if (const TSharedPtr<FRemoteControlField> RCField = RCWeakField.Pin())
+					{
+						const URCBehaviour* Behaviour = BehaviourItem->GetBehaviour();
+
+						if (Behaviour && Behaviour->CanHaveActionForField(RCField))
+						{
+							AddAction(RCField.ToSharedRef());
+						}
+					}
+				}
+			};
+
+			for (const TSharedPtr<SRCPanelTreeNode>& RCEntity : RCEntitiesList->GetSelectedEntities())
+			{
+				if (RCEntity->GetRCId().IsValid())
+				{
+					AddSelectedActionLambda(RCEntity);
+				}
+				else if (RCEntity->GetRCType() == SRCPanelTreeNode::FieldGroup)
+				{
+					if (const TSharedPtr<SRCPanelExposedEntitiesGroup> RCFieldGroup = StaticCastSharedPtr<SRCPanelExposedEntitiesGroup>(RCEntity))
+					{
+						TArray<TSharedPtr<SRCPanelTreeNode>> RCFieldGroupEntities;
+						RCFieldGroup->GetNodeChildren(RCFieldGroupEntities);
+
+						for (const TSharedPtr<SRCPanelTreeNode>& RCFieldGroupEntity : RCFieldGroupEntities)
+						{
+							if (RCFieldGroupEntity->GetRCId().IsValid())
+							{
+								AddSelectedActionLambda(RCFieldGroupEntity);
+							}
+						}
+					}
+				}
+			}
+		}
+	}
+
+	return FReply::Handled();
+}
+
 void SRCActionPanel::OnRemoteControlFieldAdded(const FGuid& GroupId, const FGuid& FieldId, int32 FieldPosition)
 {
 	bAddActionMenuNeedsRefresh = true;
diff --git a/Engine/Plugins/VirtualProduction/RemoteControl/Source/RemoteControlUI/Private/UI/Action/SRCActionPanel.h b/Engine/Plugins/VirtualProduction/RemoteControl/Source/RemoteControlUI/Private/UI/Action/SRCActionPanel.h
--- a/Engine/Plugins/VirtualProduction/RemoteControl/Source/RemoteControlUI/Private/UI/Action/SRCActionPanel.h
+++ b/Engine/Plugins/VirtualProduction/RemoteControl/Source/RemoteControlUI/Private/UI/Action/SRCActionPanel.h
@@ -137,6 +137,9 @@ private:
 	/** Handles click event for Add All button; Adds all possible actions for the active Remote Control Preset*/
 	FReply OnAddAllFields();
 
+	/** Handles click event for Add All Selected button; Adds all selected actions for the active Remote Control Preset*/
+	FReply OnAddAllSelectedFields();
+
 	/** Event invoked when a new remote control field has been added to the Remote Control Preset associated with this Action panel */
 	void OnRemoteControlFieldAdded(const FGuid& GroupId, const FGuid& FieldId, int32 FieldPosition);
 
