[RemoteControl]
Bind controller of type float/integer can now bind to vectors

diff --git a/Engine/Plugins/VirtualProduction/RemoteControl/Source/RemoteControlLogic/Private/Action/Bind/RCPropertyBindAction.cpp b/Engine/Plugins/VirtualProduction/RemoteControl/Source/RemoteControlLogic/Private/Action/Bind/RCPropertyBindAction.cpp
--- a/Engine/Plugins/VirtualProduction/RemoteControl/Source/RemoteControlLogic/Private/Action/Bind/RCPropertyBindAction.cpp
+++ b/Engine/Plugins/VirtualProduction/RemoteControl/Source/RemoteControlLogic/Private/Action/Bind/RCPropertyBindAction.cpp
@@ -299,6 +299,28 @@ void URCPropertyBindAction::Execute() const
 
 				Handle->SetValue(StringValue);
 			}
+			// Numeric to Vector/Vector2D/Rotator
+			else if (const FStructProperty* TargetStructProperty = CastField<FStructProperty>(RemoteControlProperty))
+			{
+				// FVector
+				if (TargetStructProperty->Struct == TBaseStructure<FVector>::Get())
+				{
+					const FVector VectorValue(NumericValue, NumericValue, NumericValue);
+					Handle->SetValue(VectorValue);
+				}
+				// FVector2D
+				else if (TargetStructProperty->Struct == TBaseStructure<FVector2D>::Get())
+				{
+					const FVector2D Vector2DValue(NumericValue, NumericValue);
+					Handle->SetValue(Vector2DValue);
+				}
+				// FRotator
+				else if (TargetStructProperty->Struct == TBaseStructure<FRotator>::Get())
+				{
+					const FRotator RotatorValue(NumericValue, NumericValue, NumericValue);
+					Handle->SetValue(RotatorValue);
+				} 
+			}
 		}
 		// Boolean Controller
 		else if (ControllerAsProperty->IsA(FBoolProperty::StaticClass()))
diff --git a/Engine/Plugins/VirtualProduction/RemoteControl/Source/RemoteControlLogic/Private/Behaviour/Builtin/Bind/RCBehaviourBind.cpp b/Engine/Plugins/VirtualProduction/RemoteControl/Source/RemoteControlLogic/Private/Behaviour/Builtin/Bind/RCBehaviourBind.cpp
--- a/Engine/Plugins/VirtualProduction/RemoteControl/Source/RemoteControlLogic/Private/Behaviour/Builtin/Bind/RCBehaviourBind.cpp
+++ b/Engine/Plugins/VirtualProduction/RemoteControl/Source/RemoteControlLogic/Private/Behaviour/Builtin/Bind/RCBehaviourBind.cpp
@@ -110,6 +110,33 @@ static bool EvaluateBindCompatibility(const TMap<FFieldClass*, TArray<FFieldClas
 	return false;
 }
 
+static bool EvaluateBindCompatibility(const TMap<FFieldClass*, TArray<UScriptStruct*>>& InBindingsMap, const FFieldClass* InControllerPropertyClass, const FProperty* InRemoteControlProperty)
+{
+	TArray<FFieldClass*> Keys;
+	InBindingsMap.GenerateKeyArray(Keys);
+
+	if (const FStructProperty* RCFieldStructProperty = CastField<FStructProperty>(InRemoteControlProperty))
+	{
+		for (const FFieldClass* PropertyType : Keys)
+		{
+			if (InControllerPropertyClass->IsChildOf(PropertyType))
+			{
+				if (const TArray<UScriptStruct*>* SupportedBindings = InBindingsMap.Find(PropertyType))
+				{
+					for (const UScriptStruct* SupportedBinding : *SupportedBindings)
+					{
+						if (RCFieldStructProperty->Struct && RCFieldStructProperty->Struct->IsChildOf(SupportedBinding))
+						{
+							return true;
+						}
+					}
+				}
+			}
+		}
+	}
+
+	return false;
+}
 
 static bool EvaluateCustomControllerBindCompatibility(const URCController* InController,  const FFieldClass* InControllerPropertyClass, const FFieldClass* InRemoteControlPropertyClass, const FProperty* InRemoteControlProperty)
 {
@@ -195,6 +222,14 @@ bool URCBehaviourBind::CanHaveActionForField(URCController* Controller, TSharedR
 		{ FNumericProperty::StaticClass(),  /* --> */   { FStrProperty::StaticClass(),     FTextProperty::StaticClass(),  FNameProperty::StaticClass() } }
 	};
 
+	// Indirect Binding (via numeric conversion)
+	//
+	const static TMap<FFieldClass*, TArray<UScriptStruct*>> SupportedNumericConversionsStructMap =
+	{
+		/* Controller Type */                           /* Supported Remote Control Property Types */
+		{ FNumericProperty::StaticClass(),  /* --> */   { TBaseStructure<FVector>::Get(), TBaseStructure<FVector2D>::Get(), TBaseStructure<FRotator>::Get(), } }
+	};
+
 	// Indirect Binding (for Structs)
 	//
 	const static TMap<UScriptStruct*, TArray<UScriptStruct*>> SupportedStructConversions =
@@ -219,6 +254,11 @@ bool URCBehaviourBind::CanHaveActionForField(URCController* Controller, TSharedR
 	{
 		return true;
 	}
+	// Indirect Binding (Numeric to struct)
+	else if (EvaluateBindCompatibility(SupportedNumericConversionsStructMap, ControllerPropertyClass, RemoteControlProperty))
+	{
+		return true;
+	}
 	// Indirect Binding (Numeric Conversion)
 	else if (bInAllowNumericInputAsStrings)
 	{
