[RemoteControl]
Update value option when right clicking an Action.

diff --git a/Engine/Plugins/VirtualProduction/RemoteControl/Source/RemoteControl/Private/RCVirtualProperty.cpp b/Engine/Plugins/VirtualProduction/RemoteControl/Source/RemoteControl/Private/RCVirtualProperty.cpp
--- a/Engine/Plugins/VirtualProduction/RemoteControl/Source/RemoteControl/Private/RCVirtualProperty.cpp
+++ b/Engine/Plugins/VirtualProduction/RemoteControl/Source/RemoteControl/Private/RCVirtualProperty.cpp
@@ -853,6 +853,12 @@ bool URCVirtualPropertySelfContainer::DuplicatePropertyWithCopy(URCVirtualProper
 	return false;
 }
 
+bool URCVirtualPropertySelfContainer::UpdateValueWithProperty(const FProperty* InProperty, const void* InPropertyContainer)
+{
+	const EPropertyBagResult Result = Bag.SetValue(PropertyName, InProperty, InPropertyContainer);
+	return Result == EPropertyBagResult::Success;
+}
+
 void URCVirtualPropertySelfContainer::Reset()
 {
 	PropertyName = NAME_None;
diff --git a/Engine/Plugins/VirtualProduction/RemoteControl/Source/RemoteControl/Private/RemoteControlField.cpp b/Engine/Plugins/VirtualProduction/RemoteControl/Source/RemoteControl/Private/RemoteControlField.cpp
--- a/Engine/Plugins/VirtualProduction/RemoteControl/Source/RemoteControl/Private/RemoteControlField.cpp
+++ b/Engine/Plugins/VirtualProduction/RemoteControl/Source/RemoteControl/Private/RemoteControlField.cpp
@@ -260,6 +260,19 @@ FProperty* FRemoteControlProperty::GetProperty() const
 	return nullptr;
 }
 
+void* FRemoteControlProperty::GetFieldContainerAddress() const
+{
+	// Make a copy in order to preserve constness.
+	FRCFieldPathInfo FieldPathCopy = FieldPathInfo;
+	TArray<UObject*> Objects = GetBoundObjects();
+	if (Objects.Num() && FieldPathCopy.Resolve(Objects[0]))
+	{
+		const FRCFieldResolvedData Data = FieldPathCopy.GetResolvedData();
+		return Data.ContainerAddress;
+	}
+	return nullptr;
+}
+
 TSharedPtr<IRemoteControlPropertyHandle> FRemoteControlProperty::GetPropertyHandle() const 
 {
 	TSharedPtr<FRemoteControlProperty> ThisPtr = Owner->GetExposedEntity<FRemoteControlProperty>(GetId()).Pin();
diff --git a/Engine/Plugins/VirtualProduction/RemoteControl/Source/RemoteControl/Public/RCVirtualProperty.h b/Engine/Plugins/VirtualProduction/RemoteControl/Source/RemoteControl/Public/RCVirtualProperty.h
--- a/Engine/Plugins/VirtualProduction/RemoteControl/Source/RemoteControl/Public/RCVirtualProperty.h
+++ b/Engine/Plugins/VirtualProduction/RemoteControl/Source/RemoteControl/Public/RCVirtualProperty.h
@@ -406,6 +406,16 @@ public:
 	 */
 	bool DuplicatePropertyWithCopy(URCVirtualPropertyBase* InVirtualProperty);
 
+	/**
+ 	 * Update property bag value using the property and the container passed.
+ 	 *
+ 	 * @param InProperty Property used to update the bag property value
+ 	 * @param InPropertyContainer Container to take the value using the property
+ 	 *
+ 	 * @return true if SetValue of the bag returned Success otherwise false
+ 	 */
+	bool UpdateValueWithProperty(const FProperty* InProperty, const void* InPropertyContainer);
+
 	/** Resets the property bag instance to empty and remove Virtual Property data */
 	void Reset();
 
diff --git a/Engine/Plugins/VirtualProduction/RemoteControl/Source/RemoteControl/Public/RemoteControlField.h b/Engine/Plugins/VirtualProduction/RemoteControl/Source/RemoteControl/Public/RemoteControlField.h
--- a/Engine/Plugins/VirtualProduction/RemoteControl/Source/RemoteControl/Public/RemoteControlField.h
+++ b/Engine/Plugins/VirtualProduction/RemoteControl/Source/RemoteControl/Public/RemoteControlField.h
@@ -196,6 +196,13 @@ public:
 	 */
 	FProperty* GetProperty() const;
 
+	/**
+ 	 * Get the container address.
+ 	 * @return The container address of the value exposed by this field or nullptr if it couldn't be resolved.
+ 	 * @note This field's binding must be valid to get the container address.
+ 	 */
+	void* GetFieldContainerAddress() const;
+
 	/**
 	 * Get the property handle with ability set and get property value directly.
 	 * @return The property handle for exposed property.
diff --git a/Engine/Plugins/VirtualProduction/RemoteControl/Source/RemoteControlLogic/Private/Action/RCPropertyAction.cpp b/Engine/Plugins/VirtualProduction/RemoteControl/Source/RemoteControlLogic/Private/Action/RCPropertyAction.cpp
--- a/Engine/Plugins/VirtualProduction/RemoteControl/Source/RemoteControlLogic/Private/Action/RCPropertyAction.cpp
+++ b/Engine/Plugins/VirtualProduction/RemoteControl/Source/RemoteControlLogic/Private/Action/RCPropertyAction.cpp
@@ -8,6 +8,7 @@
 #include "IRemoteControlPropertyHandle.h"
 #include "IStructSerializerBackend.h"
 #include "RCVirtualProperty.h"
+#include "RemoteControlField.h"
 #include "RemoteControlPreset.h"
 #include "StructSerializer.h"
 #include "Action/RCAction.h"
@@ -82,6 +83,18 @@ FProperty* URCPropertyAction::GetProperty() const
 	return nullptr;
 }
 
+void URCPropertyAction::UpdateValueBasedOnRCProperty() const
+{
+	if (const TSharedPtr<FRemoteControlProperty> RCExposedProperty = GetRemoteControlProperty())
+	{
+		const FProperty* RCProperty = RCExposedProperty->GetProperty();
+		if (PropertySelfContainer)
+		{
+			PropertySelfContainer->UpdateValueWithProperty(RCProperty, RCExposedProperty->GetFieldContainerAddress());
+		}
+	}
+}
+
 FName URCAction::GetExposedFieldLabel() const
 {
 	if (const URemoteControlPreset* Preset = PresetWeakPtr.Get())
diff --git a/Engine/Plugins/VirtualProduction/RemoteControl/Source/RemoteControlLogic/Public/Action/RCPropertyAction.h b/Engine/Plugins/VirtualProduction/RemoteControl/Source/RemoteControlLogic/Public/Action/RCPropertyAction.h
--- a/Engine/Plugins/VirtualProduction/RemoteControl/Source/RemoteControlLogic/Public/Action/RCPropertyAction.h
+++ b/Engine/Plugins/VirtualProduction/RemoteControl/Source/RemoteControlLogic/Public/Action/RCPropertyAction.h
@@ -27,6 +27,9 @@ public:
 	/* Returns the FProperty underlying this Action */
 	FProperty* GetProperty() const;
 
+	/* Update the value of this Action with the current value of the RC property it is bound to */
+	void UpdateValueBasedOnRCProperty() const;
+
 public:
 	/** Virtual Property Container */
 	UPROPERTY()
diff --git a/Engine/Plugins/VirtualProduction/RemoteControl/Source/RemoteControlUI/Private/Commands/RemoteControlCommands.cpp b/Engine/Plugins/VirtualProduction/RemoteControl/Source/RemoteControlUI/Private/Commands/RemoteControlCommands.cpp
--- a/Engine/Plugins/VirtualProduction/RemoteControl/Source/RemoteControlUI/Private/Commands/RemoteControlCommands.cpp
+++ b/Engine/Plugins/VirtualProduction/RemoteControl/Source/RemoteControlUI/Private/Commands/RemoteControlCommands.cpp
@@ -45,6 +45,9 @@ void FRemoteControlCommands::RegisterCommands()
 
 	// Duplicate Item
 	UI_COMMAND(DuplicateItem, "Duplicate", "Duplicate the selected UI item", EUserInterfaceActionType::Button, FInputChord(EKeys::D, EModifierKey::Control));
+
+	// Update Value
+	UI_COMMAND(UpdateValue, "Update Value", "Update the selected UI item value with the one in the fields list", EUserInterfaceActionType::Button, FInputChord(EKeys::U, EModifierKey::Control));
 }
 
 #undef LOCTEXT_NAMESPACE
diff --git a/Engine/Plugins/VirtualProduction/RemoteControl/Source/RemoteControlUI/Private/Commands/RemoteControlCommands.h b/Engine/Plugins/VirtualProduction/RemoteControl/Source/RemoteControlUI/Private/Commands/RemoteControlCommands.h
--- a/Engine/Plugins/VirtualProduction/RemoteControl/Source/RemoteControlUI/Private/Commands/RemoteControlCommands.h
+++ b/Engine/Plugins/VirtualProduction/RemoteControl/Source/RemoteControlUI/Private/Commands/RemoteControlCommands.h
@@ -44,25 +44,30 @@ public:
 	/**
 	 * Holds the information about UI Command that deletes currently selected group/exposed entity.
 	 */
-	TSharedPtr< FUICommandInfo > DeleteEntity;
+	TSharedPtr<FUICommandInfo> DeleteEntity;
 
 	/**
 	 * Holds the information about UI Command that  renames selected group/exposed entity.
 	 */
-	TSharedPtr< FUICommandInfo > RenameEntity;
+	TSharedPtr<FUICommandInfo> RenameEntity;
 
 	/**
 	 * UI Command for copying a UI item in the Remote Control preset. Currently used for Logic panel
 	 */
-	TSharedPtr< FUICommandInfo > CopyItem;
+	TSharedPtr<FUICommandInfo> CopyItem;
 
 	/**
 	 * UI Command for pasting a UI item in the Remote Control preset. Currently used for Logic panel
 	 */
-	TSharedPtr< FUICommandInfo > PasteItem;
+	TSharedPtr<FUICommandInfo> PasteItem;
 
 	/**
 	 * UI Command for duplicating a UI item in the Remote Control preset. Currently used for Logic panel
 	 */
-	TSharedPtr< FUICommandInfo > DuplicateItem;
+	TSharedPtr<FUICommandInfo> DuplicateItem;
+
+	/**
+	 * UI Command for updating tha action in the action list with the value in the field list. Currently used for Logic panel
+	 */
+	TSharedPtr<FUICommandInfo> UpdateValue;
 };
diff --git a/Engine/Plugins/VirtualProduction/RemoteControl/Source/RemoteControlUI/Private/RemoteControlUIModule.cpp b/Engine/Plugins/VirtualProduction/RemoteControl/Source/RemoteControlUI/Private/RemoteControlUIModule.cpp
--- a/Engine/Plugins/VirtualProduction/RemoteControl/Source/RemoteControlUI/Private/RemoteControlUIModule.cpp
+++ b/Engine/Plugins/VirtualProduction/RemoteControl/Source/RemoteControlUI/Private/RemoteControlUIModule.cpp
@@ -475,6 +475,7 @@ void FRemoteControlUIModule::UnbindRemoteControlCommands()
 		ActionList.UnmapAction(Commands.CopyItem);
 		ActionList.UnmapAction(Commands.PasteItem);
 		ActionList.UnmapAction(Commands.DuplicateItem);
+		ActionList.UnmapAction(Commands.UpdateValue);
 	}
 
 	FRemoteControlCommands::Unregister();
diff --git a/Engine/Plugins/VirtualProduction/RemoteControl/Source/RemoteControlUI/Private/UI/BaseLogicUI/SRCLogicPanelListBase.cpp b/Engine/Plugins/VirtualProduction/RemoteControl/Source/RemoteControlUI/Private/UI/BaseLogicUI/SRCLogicPanelListBase.cpp
--- a/Engine/Plugins/VirtualProduction/RemoteControl/Source/RemoteControlUI/Private/UI/BaseLogicUI/SRCLogicPanelListBase.cpp
+++ b/Engine/Plugins/VirtualProduction/RemoteControl/Source/RemoteControlUI/Private/UI/BaseLogicUI/SRCLogicPanelListBase.cpp
@@ -52,10 +52,13 @@ TSharedPtr<SWidget> SRCLogicPanelListBase::GetContextMenuWidget()
 	// 2. Duplicate
 	MenuBuilder.AddMenuEntry(Commands.DuplicateItem);
 
-	// 3. Delete
+	// 3. Update
+	MenuBuilder.AddMenuEntry(Commands.UpdateValue);
+
+	// 4. Delete
 	MenuBuilder.AddMenuEntry(Commands.DeleteEntity);
 
-	// 4. Delete All
+	// 5. Delete All
 	FUIAction Action(FExecuteAction::CreateSP(this, &SRCLogicPanelListBase::RequestDeleteAllItems)
 	, FCanExecuteAction::CreateSP(this, &SRCLogicPanelListBase::CanDeleteAllItems));
 
diff --git a/Engine/Plugins/VirtualProduction/RemoteControl/Source/RemoteControlUI/Private/UI/SRemoteControlPanel.cpp b/Engine/Plugins/VirtualProduction/RemoteControl/Source/RemoteControlUI/Private/UI/SRemoteControlPanel.cpp
--- a/Engine/Plugins/VirtualProduction/RemoteControl/Source/RemoteControlUI/Private/UI/SRemoteControlPanel.cpp
+++ b/Engine/Plugins/VirtualProduction/RemoteControl/Source/RemoteControlUI/Private/UI/SRemoteControlPanel.cpp
@@ -1406,6 +1406,11 @@ void SRemoteControlPanel::BindRemoteControlCommands()
 		Commands.DuplicateItem,
 		FExecuteAction::CreateSP(this, &SRemoteControlPanel::DuplicateItem_Execute),
 		FCanExecuteAction::CreateSP(this, &SRemoteControlPanel::CanDuplicateItem));
+
+	ActionList.MapAction(
+		Commands.UpdateValue,
+		FExecuteAction::CreateSP(this, &SRemoteControlPanel::UpdateValue_Execute),
+		FCanExecuteAction::CreateSP(this, &SRemoteControlPanel::CanUpdateValue));
 }
 
 void SRemoteControlPanel::RegisterEvents()
@@ -2597,6 +2602,44 @@ bool SRemoteControlPanel::CanDuplicateItem() const
 	return false;
 }
 
+void SRemoteControlPanel::UpdateValue_Execute()
+{
+	if (const TSharedPtr<SRCLogicPanelBase> ActiveLogicPanel = GetActiveLogicPanel())
+	{
+		if (const TSharedPtr<FRCActionModel> RCActionLogicItem = StaticCastSharedPtr<FRCActionModel>(ActiveLogicPanel->GetSelectedLogicItem()))
+		{
+			if (const URCPropertyAction* RCPropertyAction = Cast<URCPropertyAction>(RCActionLogicItem->GetAction()))
+			{
+				RCPropertyAction->UpdateValueBasedOnRCProperty();
+			}
+		}
+	}
+}
+
+bool SRemoteControlPanel::CanUpdateValue() const
+{
+	if (bIsInLiveMode)
+	{
+		return false;
+	}
+	
+	if (const TSharedPtr<SRCLogicPanelBase> ActiveLogicPanel = GetActiveLogicPanel())
+	{
+		if (const TSharedPtr<FRCActionModel> RCActionLogicItem = StaticCastSharedPtr<FRCActionModel>(ActiveLogicPanel->GetSelectedLogicItem()))
+		{
+			if (const URemoteControlPreset* RCPreset = GetPreset())
+			{
+				if (const URCAction* RCAction = RCActionLogicItem->GetAction())
+				{
+					return RCAction->IsA<URCPropertyAction>() && RCPreset->GetExposedEntity(RCAction->ExposedFieldId).IsValid();
+				}
+			}
+		}
+	}
+
+	return false;
+}
+
 void SRemoteControlPanel::LoadSettings(const FGuid& InInstanceId) const
 {
 	const FString SettingsString = InInstanceId.ToString();
diff --git a/Engine/Plugins/VirtualProduction/RemoteControl/Source/RemoteControlUI/Private/UI/SRemoteControlPanel.h b/Engine/Plugins/VirtualProduction/RemoteControl/Source/RemoteControlUI/Private/UI/SRemoteControlPanel.h
--- a/Engine/Plugins/VirtualProduction/RemoteControl/Source/RemoteControlUI/Private/UI/SRemoteControlPanel.h
+++ b/Engine/Plugins/VirtualProduction/RemoteControl/Source/RemoteControlUI/Private/UI/SRemoteControlPanel.h
@@ -332,6 +332,12 @@ private:
 	/** Called to test if user is able to Duplicate a logic UI item. */
 	bool CanDuplicateItem() const;
 
+	/** Called when user attempts to Update a logic UI item. */
+	void UpdateValue_Execute();
+
+	/** Called to test if user is able to Update a logic UI item. */
+	bool CanUpdateValue() const;
+
 	/** Loads settings from config based on the preset identifier. */
 	void LoadSettings(const FGuid& InInstanceId) const;
 
