[RemoteControl]
Disable UpdateValue command in action for Bind behaviour.

diff --git a/Engine/Plugins/VirtualProduction/RemoteControl/Source/RemoteControlUI/Private/UI/SRemoteControlPanel.cpp b/Engine/Plugins/VirtualProduction/RemoteControl/Source/RemoteControlUI/Private/UI/SRemoteControlPanel.cpp
--- a/Engine/Plugins/VirtualProduction/RemoteControl/Source/RemoteControlUI/Private/UI/SRemoteControlPanel.cpp
+++ b/Engine/Plugins/VirtualProduction/RemoteControl/Source/RemoteControlUI/Private/UI/SRemoteControlPanel.cpp
@@ -4,6 +4,7 @@
 
 #include "Action/SRCActionPanel.h"
 #include "ActorEditorUtils.h"
+#include "Behaviour/Builtin/Bind/RCBehaviourBind.h"
 #include "Behaviour/SRCBehaviourPanel.h"
 #include "ClassViewerFilter.h"
 #include "ClassViewerModule.h"
@@ -2631,7 +2632,13 @@ bool SRemoteControlPanel::CanUpdateValue() const
 			{
 				if (const URCAction* RCAction = RCActionLogicItem->GetAction())
 				{
-					return RCAction->IsA<URCPropertyAction>() && RCPreset->GetExposedEntity(RCAction->ExposedFieldId).IsValid();
+					if (const TSharedPtr<FRCBehaviourModel> ParentBehaviour = RCActionLogicItem->GetParentBehaviour())
+					{
+						if (const URCBehaviour* Behaviour = ParentBehaviour->GetBehaviour())
+						{
+							return RCAction->IsA<URCPropertyAction>() && !Behaviour->IsA<URCBehaviourBind>() && RCPreset->GetExposedEntity(RCAction->ExposedFieldId).IsValid();
+						}
+					}
 				}
 			}
 		}
