[Remote Control] Address issues with static global variables

Issues with these globals have been found while loading a Module referencing Remote Control on Mac OS.
The module was being loaded in PostConfigInit phase.

Possibly, the Remote Control globals were being initialized before the UObject system was loaded, thus causing a crash.

This CL converts those globals to static members of static functions, so that they are initialized only when requested, implying that the UObject system is loaded by then.

diff --git a/Engine/Plugins/VirtualProduction/RemoteControl/Source/RemoteControlProtocolWidgets/Private/Widgets/Masks/SRCProtocolBindingMask.cpp b/Engine/Plugins/VirtualProduction/RemoteControl/Source/RemoteControlProtocolWidgets/Private/Widgets/Masks/SRCProtocolBindingMask.cpp
--- a/Engine/Plugins/VirtualProduction/RemoteControl/Source/RemoteControlProtocolWidgets/Private/Widgets/Masks/SRCProtocolBindingMask.cpp
+++ b/Engine/Plugins/VirtualProduction/RemoteControl/Source/RemoteControlProtocolWidgets/Private/Widgets/Masks/SRCProtocolBindingMask.cpp
@@ -47,7 +47,7 @@ EMaskingType SRCProtocolBindingMask::GetMaskingType()
 		{
 			if (const FStructProperty* StructProperty = CastField<FStructProperty>(Property))
 			{
-				if (const EMaskingType* StructToMaskingType = RemoteControlProtocolMasking::StructsToMaskingTypes.Find(StructProperty->Struct))
+				if (const EMaskingType* StructToMaskingType = FRemoteControlProtocolMasking::GetStructsToMaskingTypes().Find(StructProperty->Struct))
 				{
 					return *StructToMaskingType;
 				}
@@ -66,7 +66,7 @@ bool SRCProtocolBindingMask::HasOptionalMask() const
 		{
 			if (const FStructProperty* StructProperty = CastField<FStructProperty>(Property))
 			{
-				return RemoteControlProtocolMasking::OptionalMaskStructs.Contains(StructProperty->Struct);
+				return FRemoteControlProtocolMasking::GetOptionalMaskStructs().Contains(StructProperty->Struct);
 			}
 		}
 	}
diff --git a/Engine/Plugins/VirtualProduction/RemoteControl/Source/RemoteControlProtocolWidgets/Private/Widgets/Masks/SRCProtocolMask.cpp b/Engine/Plugins/VirtualProduction/RemoteControl/Source/RemoteControlProtocolWidgets/Private/Widgets/Masks/SRCProtocolMask.cpp
--- a/Engine/Plugins/VirtualProduction/RemoteControl/Source/RemoteControlProtocolWidgets/Private/Widgets/Masks/SRCProtocolMask.cpp
+++ b/Engine/Plugins/VirtualProduction/RemoteControl/Source/RemoteControlProtocolWidgets/Private/Widgets/Masks/SRCProtocolMask.cpp
@@ -45,7 +45,7 @@ EMaskingType SRCProtocolMask::GetMaskingType()
 			{
 				if (const FStructProperty* StructProperty = CastField<FStructProperty>(RCProperty->GetProperty()))
 				{
-					if (const EMaskingType* StructToMaskingType = RemoteControlProtocolMasking::StructsToMaskingTypes.Find(StructProperty->Struct))
+					if (const EMaskingType* StructToMaskingType = FRemoteControlProtocolMasking::GetStructsToMaskingTypes().Find(StructProperty->Struct))
 					{
 						return *StructToMaskingType;
 					}
@@ -65,7 +65,7 @@ bool SRCProtocolMask::HasOptionalMask() const
 		{
 			if (const FStructProperty* StructProperty = CastField<FStructProperty>(RCProperty->GetProperty()))
 			{
-				return RemoteControlProtocolMasking::OptionalMaskStructs.Contains(StructProperty->Struct);
+				return FRemoteControlProtocolMasking::GetOptionalMaskStructs().Contains(StructProperty->Struct);
 			}
 		}
 	}
diff --git a/Engine/Plugins/VirtualProduction/RemoteControl/Source/RemoteControlProtocolWidgets/Private/Widgets/Masks/SRCProtocolMaskTriplet.cpp b/Engine/Plugins/VirtualProduction/RemoteControl/Source/RemoteControlProtocolWidgets/Private/Widgets/Masks/SRCProtocolMaskTriplet.cpp
--- a/Engine/Plugins/VirtualProduction/RemoteControl/Source/RemoteControlProtocolWidgets/Private/Widgets/Masks/SRCProtocolMaskTriplet.cpp
+++ b/Engine/Plugins/VirtualProduction/RemoteControl/Source/RemoteControlProtocolWidgets/Private/Widgets/Masks/SRCProtocolMaskTriplet.cpp
@@ -11,6 +11,33 @@
 
 #define LOCTEXT_NAMESPACE "SRCProtocolMaskTriplet"
 
+const TMap<UScriptStruct*, EMaskingType>& FRemoteControlProtocolMasking::GetStructsToMaskingTypes()
+{
+	static TMap<UScriptStruct*, EMaskingType> StructsToMaskingTypes = { { TBaseStructure<FColor>::Get(), EMaskingType::Color }
+		, { TBaseStructure<FLinearColor>::Get(), EMaskingType::Color }
+		, { TBaseStructure<FRotator>::Get(), EMaskingType::Rotator }
+		, { TBaseStructure<FVector>::Get(), EMaskingType::Vector }
+		, { TBaseStructure<FIntVector>::Get(), EMaskingType::Vector }
+		, { TBaseStructure<FVector4>::Get(), EMaskingType::Quat }
+		, { TBaseStructure<FIntVector4>::Get(), EMaskingType::Quat }
+		, { TBaseStructure<FQuat>::Get(), EMaskingType::Quat }
+	};
+
+	return StructsToMaskingTypes;
+}
+
+const TSet<UScriptStruct*>& FRemoteControlProtocolMasking::GetOptionalMaskStructs()
+{
+	static TSet<UScriptStruct*> OptionalMaskStructs = {
+		TBaseStructure<FLinearColor>::Get(),
+		TBaseStructure<FVector4>::Get(),
+		TBaseStructure<FQuat>::Get(),
+		TBaseStructure<FIntVector4>::Get()
+	};
+
+	return OptionalMaskStructs;
+}
+
 void SRCProtocolMaskTriplet::Construct(const FArguments& InArgs)
 {
 	WidgetStyle = &FProtocolPanelStyle::Get()->GetWidgetStyle<FProtocolWidgetStyle>("ProtocolsPanel.Widgets.Mask");
diff --git a/Engine/Plugins/VirtualProduction/RemoteControl/Source/RemoteControlProtocolWidgets/Public/Widgets/Masks/SRCProtocolMaskTriplet.h b/Engine/Plugins/VirtualProduction/RemoteControl/Source/RemoteControlProtocolWidgets/Public/Widgets/Masks/SRCProtocolMaskTriplet.h
--- a/Engine/Plugins/VirtualProduction/RemoteControl/Source/RemoteControlProtocolWidgets/Public/Widgets/Masks/SRCProtocolMaskTriplet.h
+++ b/Engine/Plugins/VirtualProduction/RemoteControl/Source/RemoteControlProtocolWidgets/Public/Widgets/Masks/SRCProtocolMaskTriplet.h
@@ -38,25 +38,12 @@ enum class EMaskingType : uint8
 	Unsupported
 };
 
-namespace RemoteControlProtocolMasking
+class FRemoteControlProtocolMasking
 {
-	static TMap<UScriptStruct*, EMaskingType> StructsToMaskingTypes = { { TBaseStructure<FColor>::Get(), EMaskingType::Color }
-		, { TBaseStructure<FLinearColor>::Get(), EMaskingType::Color }
-		, { TBaseStructure<FRotator>::Get(), EMaskingType::Rotator }
-		, { TBaseStructure<FVector>::Get(), EMaskingType::Vector }
-		, { TBaseStructure<FIntVector>::Get(), EMaskingType::Vector }
-		, { TBaseStructure<FVector4>::Get(), EMaskingType::Quat }
-		, { TBaseStructure<FIntVector4>::Get(), EMaskingType::Quat }
-		, { TBaseStructure<FQuat>::Get(), EMaskingType::Quat }
-	};
-
-	static TSet<UScriptStruct*> OptionalMaskStructs = {
-		TBaseStructure<FLinearColor>::Get(),
-		TBaseStructure<FVector4>::Get(),
-		TBaseStructure<FQuat>::Get(),
-		TBaseStructure<FIntVector4>::Get()
-	};
-}
+public:
+	static const TMap<UScriptStruct*, EMaskingType>& GetStructsToMaskingTypes();
+	static const TSet<UScriptStruct*>& GetOptionalMaskStructs();
+};
 
 /**
  * A widget that represents a triplet mask.
diff --git a/Engine/Plugins/VirtualProduction/RemoteControl/Source/RemoteControlUI/Private/RemoteControlUIModule.cpp b/Engine/Plugins/VirtualProduction/RemoteControl/Source/RemoteControlUI/Private/RemoteControlUIModule.cpp
--- a/Engine/Plugins/VirtualProduction/RemoteControl/Source/RemoteControlUI/Private/RemoteControlUIModule.cpp
+++ b/Engine/Plugins/VirtualProduction/RemoteControl/Source/RemoteControlUI/Private/RemoteControlUIModule.cpp
@@ -274,10 +274,6 @@ void FRemoteControlUIModule::ShutdownModule()
 	UnbindRemoteControlCommands();
 	FRemoteControlPanelStyle::Shutdown();
 	FRemoteControlExposeMenuStyle::Shutdown();
-	SRemoteControlPanel::Shutdown();
-	SRCActionPanel::Shutdown();
-	SRCBehaviourPanel::Shutdown();
-	SRCBehaviourConditional::Shutdown();
 }
 
 FDelegateHandle FRemoteControlUIModule::AddPropertyFilter(FOnDisplayExposeIcon OnDisplayExposeIcon)
diff --git a/Engine/Plugins/VirtualProduction/RemoteControl/Source/RemoteControlUI/Private/UI/Action/SRCActionPanel.cpp b/Engine/Plugins/VirtualProduction/RemoteControl/Source/RemoteControlUI/Private/UI/Action/SRCActionPanel.cpp
--- a/Engine/Plugins/VirtualProduction/RemoteControl/Source/RemoteControlUI/Private/UI/Action/SRCActionPanel.cpp
+++ b/Engine/Plugins/VirtualProduction/RemoteControl/Source/RemoteControlUI/Private/UI/Action/SRCActionPanel.cpp
@@ -40,16 +40,19 @@
 
 #define LOCTEXT_NAMESPACE "SRCActionPanel"
 
-TSharedPtr<SBox> SRCActionPanel::NoneSelectedWidget = SNew(SBox)
-			.Padding(0.f)
-			.VAlign(VAlign_Center)
-			.HAlign(HAlign_Center)
-			[
-				SNew(STextBlock)
-				.Text(LOCTEXT("NoneSelected", "Select a behavior to view its actions."))
-				.TextStyle(&FAppStyle::GetWidgetStyle<FTextBlockStyle>("NormalText"))
-				.Justification(ETextJustify::Center)
-			];
+TSharedRef<SBox> SRCActionPanel::CreateNoneSelectedWidget()
+{
+	return SNew(SBox)
+	.Padding(0.f)
+	.VAlign(VAlign_Center)
+	.HAlign(HAlign_Center)
+	[
+		SNew(STextBlock)
+		.Text(LOCTEXT("NoneSelected", "Select a behavior to view its actions."))
+		.TextStyle(&FAppStyle::GetWidgetStyle<FTextBlockStyle>("NormalText"))
+		.Justification(ETextJustify::Center)
+	];
+}
 
 BEGIN_SLATE_FUNCTION_BUILD_OPTIMIZATION
 
@@ -82,11 +85,6 @@ void SRCActionPanel::Construct(const FArguments& InArgs, const TSharedRef<SRemot
 
 END_SLATE_FUNCTION_BUILD_OPTIMIZATION
 
-void SRCActionPanel::Shutdown()
-{
-	NoneSelectedWidget.Reset();
-}
-
 void SRCActionPanel::OnBehaviourSelectionChanged(TSharedPtr<FRCBehaviourModel> InBehaviourItem)
 {
 	SelectedBehaviourItemWeakPtr = InBehaviourItem;
@@ -199,7 +197,7 @@ void SRCActionPanel::UpdateWrappedWidget(TSharedPtr<FRCBehaviourModel> InBehavio
 	}
 	else
 	{
-		WrappedBoxWidget->SetContent(NoneSelectedWidget.ToSharedRef());
+		WrappedBoxWidget->SetContent(CreateNoneSelectedWidget());
 	}
 }
 
@@ -773,4 +771,4 @@ EVisibility SRCActionPanel::HandleAddAllButtonVisibility() const
 	return EVisibility::Collapsed;
 }
 
-#undef LOCTEXT_NAMESPACE
\ No newline at end of file
+#undef LOCTEXT_NAMESPACE
diff --git a/Engine/Plugins/VirtualProduction/RemoteControl/Source/RemoteControlUI/Private/UI/Action/SRCActionPanel.h b/Engine/Plugins/VirtualProduction/RemoteControl/Source/RemoteControlUI/Private/UI/Action/SRCActionPanel.h
--- a/Engine/Plugins/VirtualProduction/RemoteControl/Source/RemoteControlUI/Private/UI/Action/SRCActionPanel.h
+++ b/Engine/Plugins/VirtualProduction/RemoteControl/Source/RemoteControlUI/Private/UI/Action/SRCActionPanel.h
@@ -37,9 +37,6 @@ public:
 	/** Constructs this widget with InArgs */
 	void Construct(const FArguments& InArgs, const TSharedRef<SRemoteControlPanel>& InPanel);
 
-	/** Shutdown panel */
-	static void Shutdown();
-
 	/** Whether the Actions list widget currently has focus.*/
 	bool IsListFocused() const;
 
@@ -147,6 +144,8 @@ private:
 	void OnRemoteControlFieldDeleted(const FGuid& GroupId, const FGuid& FieldId, int32 FieldPosition);
 
 private:
+	/** Helper widget for behavior details. */
+	static TSharedRef<SBox> CreateNoneSelectedWidget();
 
 	void DuplicateAction(URCAction* InAction);
 
@@ -164,15 +163,12 @@ private:
 	/** Behaviour specific details widget*/
 	TSharedPtr<SRCBehaviourDetails> BehaviourDetailsWidget;
 
-	/** Helper widget for behavior details. */
-	static TSharedPtr<SBox> NoneSelectedWidget;
-
 	/** Panel Style reference. */
-	const FRCPanelStyle* RCPanelStyle;
+	const FRCPanelStyle* RCPanelStyle = nullptr;
 
 	/** Cached menu widget for Add New Action */
 	TSharedPtr<SWidget> AddNewActionMenuWidget;
 
 	/** Whether the Add Actions menu list is outdated and needs to be refreshed */
 	bool bAddActionMenuNeedsRefresh = false;
-};
\ No newline at end of file
+};
diff --git a/Engine/Plugins/VirtualProduction/RemoteControl/Source/RemoteControlUI/Private/UI/Behaviour/Builtin/Conditional/SRCBehaviourConditional.cpp b/Engine/Plugins/VirtualProduction/RemoteControl/Source/RemoteControlUI/Private/UI/Behaviour/Builtin/Conditional/SRCBehaviourConditional.cpp
--- a/Engine/Plugins/VirtualProduction/RemoteControl/Source/RemoteControlUI/Private/UI/Behaviour/Builtin/Conditional/SRCBehaviourConditional.cpp
+++ b/Engine/Plugins/VirtualProduction/RemoteControl/Source/RemoteControlUI/Private/UI/Behaviour/Builtin/Conditional/SRCBehaviourConditional.cpp
@@ -22,18 +22,6 @@
 
 class SPositiveActionButton;
 
-namespace UE::SRCBehaviourConditional::Private
-{
-	static TSharedPtr<ERCBehaviourConditionType> ConditionTypeIsEqual = MakeShared<ERCBehaviourConditionType>(ERCBehaviourConditionType::IsEqual);
-	static TSharedPtr<ERCBehaviourConditionType> ConditionTypeIsLesserThan = MakeShared<ERCBehaviourConditionType>(ERCBehaviourConditionType::IsLesserThan);
-	static TSharedPtr<ERCBehaviourConditionType> ConditionTypeIsGreaterThan = MakeShared<ERCBehaviourConditionType>(ERCBehaviourConditionType::IsGreaterThan);
-	static TSharedPtr<ERCBehaviourConditionType> ConditionTypeIsLesserThanOrEqualTo = MakeShared<ERCBehaviourConditionType>(ERCBehaviourConditionType::IsLesserThanOrEqualTo);
-	static TSharedPtr<ERCBehaviourConditionType> ConditionTypeIsGreaterThanOrEqualTo = MakeShared<ERCBehaviourConditionType>(ERCBehaviourConditionType::IsGreaterThanOrEqualTo);
-	static TSharedPtr<ERCBehaviourConditionType> ConditionTypeElse = MakeShared<ERCBehaviourConditionType>(ERCBehaviourConditionType::Else);
-}
-
-using namespace UE::SRCBehaviourConditional::Private;
-
 BEGIN_SLATE_FUNCTION_BUILD_OPTIMIZATION
 
 void SRCBehaviourConditional::Construct(const FArguments& InArgs, TSharedRef<FRCBehaviourConditionalModel> InBehaviourItem)
@@ -46,17 +34,22 @@ void SRCBehaviourConditional::Construct(const FArguments& InArgs, TSharedRef<FRC
 	{
 		if (URCController* Controller = Behaviour->ControllerWeakPtr.Get())
 		{
-			Conditions.Add(ConditionTypeIsEqual);
+			auto AddCondition = [this](ERCBehaviourConditionType InConditionType)
+			{
+				Conditions.Add(MakeShared<ERCBehaviourConditionType>(InConditionType));
+			};
+
+			AddCondition(ERCBehaviourConditionType::IsEqual);
 
 			if (Controller->IsNumericType())
 			{
-				Conditions.Add(ConditionTypeIsLesserThan);
-				Conditions.Add(ConditionTypeIsGreaterThan);
-				Conditions.Add(ConditionTypeIsLesserThanOrEqualTo);
-				Conditions.Add(ConditionTypeIsGreaterThanOrEqualTo);
+				AddCondition(ERCBehaviourConditionType::IsLesserThan);
+				AddCondition(ERCBehaviourConditionType::IsGreaterThan);
+				AddCondition(ERCBehaviourConditionType::IsLesserThanOrEqualTo);
+				AddCondition(ERCBehaviourConditionType::IsGreaterThanOrEqualTo);
 			}
 
-			Conditions.Add(ConditionTypeElse);
+			AddCondition(ERCBehaviourConditionType::Else);
 		}
 	}
 
@@ -118,17 +111,6 @@ void SRCBehaviourConditional::Construct(const FArguments& InArgs, TSharedRef<FRC
 
 END_SLATE_FUNCTION_BUILD_OPTIMIZATION
 
-void SRCBehaviourConditional::Shutdown()
-{
-	// Reset all static shared pointers
-	ConditionTypeIsEqual.Reset();
-	ConditionTypeIsLesserThan.Reset();
-	ConditionTypeIsGreaterThan.Reset();
-	ConditionTypeIsLesserThanOrEqualTo.Reset();
-	ConditionTypeIsGreaterThan.Reset();
-	ConditionTypeElse.Reset();
-}
-
 void SRCBehaviourConditional::RefreshPropertyWidget()
 {
 	if (ComparandFieldBoxWidget)
@@ -136,7 +118,7 @@ void SRCBehaviourConditional::RefreshPropertyWidget()
 		if (TSharedPtr<const FRCBehaviourConditionalModel> BehaviourItem = ConditionalBehaviourItemWeakPtr.Pin())
 		{
 			ComparandFieldBoxWidget->SetContent(BehaviourItem->GetComparandFieldWidget());
-		}		
+		}
 	}
 }
 
@@ -186,4 +168,4 @@ void SRCBehaviourConditional::OnConditionsListSelectionChanged(TSharedPtr<ERCBeh
 	ComparandFieldBoxWidget->SetVisibility(bComparandFieldVisibility);
 }
 
-#undef LOCTEXT_NAMESPACE
\ No newline at end of file
+#undef LOCTEXT_NAMESPACE
diff --git a/Engine/Plugins/VirtualProduction/RemoteControl/Source/RemoteControlUI/Private/UI/Behaviour/Builtin/Conditional/SRCBehaviourConditional.h b/Engine/Plugins/VirtualProduction/RemoteControl/Source/RemoteControlUI/Private/UI/Behaviour/Builtin/Conditional/SRCBehaviourConditional.h
--- a/Engine/Plugins/VirtualProduction/RemoteControl/Source/RemoteControlUI/Private/UI/Behaviour/Builtin/Conditional/SRCBehaviourConditional.h
+++ b/Engine/Plugins/VirtualProduction/RemoteControl/Source/RemoteControlUI/Private/UI/Behaviour/Builtin/Conditional/SRCBehaviourConditional.h
@@ -34,9 +34,6 @@ public:
 	/** Constructs this widget with InArgs */
 	void Construct(const FArguments& InArgs, TSharedRef<FRCBehaviourConditionalModel> InBehaviourItem);
 
-	/** Shutdown panel  (used to clear static shared pointers for the Conditions List) */
-	static void Shutdown();
-
 	/** Refreshes the Comparand input field.
 	* This needs to be performed each time the user creates a new Action, because the previous Comparand field has been consumed
 	*and so we create a new virtual property Comparand with whose memory the UI widget needs to be rebound.*/
diff --git a/Engine/Plugins/VirtualProduction/RemoteControl/Source/RemoteControlUI/Private/UI/Behaviour/SRCBehaviourPanel.cpp b/Engine/Plugins/VirtualProduction/RemoteControl/Source/RemoteControlUI/Private/UI/Behaviour/SRCBehaviourPanel.cpp
--- a/Engine/Plugins/VirtualProduction/RemoteControl/Source/RemoteControlUI/Private/UI/Behaviour/SRCBehaviourPanel.cpp
+++ b/Engine/Plugins/VirtualProduction/RemoteControl/Source/RemoteControlUI/Private/UI/Behaviour/SRCBehaviourPanel.cpp
@@ -3,22 +3,21 @@
 #include "SRCBehaviourPanel.h"
 #include "SRCBehaviourPanelList.h"
 
-#include "Behaviour/RCBehaviourBlueprintNode.h"
 #include "Behaviour/RCBehaviour.h"
+#include "Behaviour/RCBehaviourBlueprintNode.h"
 #include "Behaviour/RCBehaviourNode.h"
 #include "Controller/RCController.h"
 #include "Framework/MultiBox/MultiBoxBuilder.h"
 #include "Kismet2/KismetEditorUtilities.h"
 #include "Misc/MessageDialog.h"
-#include "RemoteControlPreset.h"
 
 #include "SlateOptMacros.h"
 #include "Styling/RemoteControlStyles.h"
 
+#include "UI/Controller/RCControllerModel.h"
 #include "UI/Panels/SRCDockPanel.h"
 #include "UI/RemoteControlPanelStyle.h"
 #include "UI/SRemoteControlPanel.h"
-#include "UI/Controller/RCControllerModel.h"
 #include "UObject/UObjectIterator.h"
 #include "Widgets/Input/SComboButton.h"
 
@@ -28,17 +27,20 @@
 
 #define LOCTEXT_NAMESPACE "SRCBehaviourPanel"
 
-TSharedPtr<SBox> SRCBehaviourPanel::NoneSelectedWidget = SNew(SBox)
-			.Padding(10.f)
-			.VAlign(VAlign_Center)
-			.HAlign(HAlign_Center)
-			[
-				SNew(STextBlock)
-				.Text(LOCTEXT("NoneSelected", "Select a controller\nto view its behaviors."))
-				.TextStyle(&FAppStyle::GetWidgetStyle<FTextBlockStyle>("NormalText"))
-				.Justification(ETextJustify::Center)
-				.AutoWrapText(true)
-			];
+TSharedRef<SBox> SRCBehaviourPanel::CreateNoneSelectedWidget()
+{
+	return SNew(SBox)
+		.Padding(10.f)
+		.VAlign(VAlign_Center)
+		.HAlign(HAlign_Center)
+		[
+			SNew(STextBlock)
+			.Text(LOCTEXT("NoneSelected", "Select a controller\nto view its behaviors."))
+			.TextStyle(&FAppStyle::GetWidgetStyle<FTextBlockStyle>("NormalText"))
+			.Justification(ETextJustify::Center)
+			.AutoWrapText(true)
+		];
+}
 
 void SRCBehaviourPanel::Construct(const FArguments& InArgs, const TSharedRef<SRemoteControlPanel>& InPanel)
 {
@@ -59,11 +61,6 @@ void SRCBehaviourPanel::Construct(const FArguments& InArgs, const TSharedRef<SRe
 	InPanel->OnControllerSelectionChanged.AddSP(this, &SRCBehaviourPanel::OnControllerSelectionChanged);
 }
 
-void SRCBehaviourPanel::Shutdown()
-{
-	NoneSelectedWidget.Reset();
-}
-
 void SRCBehaviourPanel::OnControllerSelectionChanged(TSharedPtr<FRCControllerModel> InControllerItem)
 {
 	SelectedControllerItemWeakPtr = InControllerItem;
@@ -116,7 +113,7 @@ void SRCBehaviourPanel::UpdateWrappedWidget(TSharedPtr<FRCControllerModel> InCon
 	}
 	else
 	{
-		WrappedBoxWidget->SetContent(NoneSelectedWidget.ToSharedRef());
+		WrappedBoxWidget->SetContent(CreateNoneSelectedWidget());
 	}
 }
 
@@ -370,4 +367,4 @@ FReply SRCBehaviourPanel::RequestDeleteAllItems()
 	return FReply::Handled();
 }
 
-#undef LOCTEXT_NAMESPACE
\ No newline at end of file
+#undef LOCTEXT_NAMESPACE
diff --git a/Engine/Plugins/VirtualProduction/RemoteControl/Source/RemoteControlUI/Private/UI/Behaviour/SRCBehaviourPanel.h b/Engine/Plugins/VirtualProduction/RemoteControl/Source/RemoteControlUI/Private/UI/Behaviour/SRCBehaviourPanel.h
--- a/Engine/Plugins/VirtualProduction/RemoteControl/Source/RemoteControlUI/Private/UI/Behaviour/SRCBehaviourPanel.h
+++ b/Engine/Plugins/VirtualProduction/RemoteControl/Source/RemoteControlUI/Private/UI/Behaviour/SRCBehaviourPanel.h
@@ -30,9 +30,6 @@ public:
 	/** Constructs this widget with InArgs */
 	void Construct(const FArguments& InArgs, const TSharedRef<SRemoteControlPanel>& InPanel);
 
-	/** Shutdown panel */
-	static void Shutdown();
-
 	/** Whether the Behaviour list widget currently has focus. Used for Delete Item UI command */
 	bool IsListFocused() const;
 
@@ -69,6 +66,9 @@ protected:
 	virtual FReply RequestDeleteAllItems() override;
 
 private:
+	/** Get a Helper widget for behavior details. */
+	static TSharedRef<SBox> CreateNoneSelectedWidget();
+
 	/** Duplicates a given Behaviour object*/
 	void DuplicateBehaviour(URCBehaviour* InBehaviour);
 	
@@ -94,7 +94,6 @@ private:
 	FReply OnClickEmptyButton();
 
 private:
-
 	/** The parent Controller that this Behaviour panel is associated with */
 	TWeakPtr<FRCControllerModel> SelectedControllerItemWeakPtr = nullptr;
 	
@@ -104,9 +103,6 @@ private:
 	/** Widget representing List of Behaviours */
 	TSharedPtr<class SRCBehaviourPanelList> BehaviourPanelList;
 
-	/** Helper widget for behavior details. */
-	static TSharedPtr<SBox> NoneSelectedWidget;
-
 	/** Panel Style reference. */
-	const FRCPanelStyle* RCPanelStyle;
+	const FRCPanelStyle* RCPanelStyle = nullptr;
 };
diff --git a/Engine/Plugins/VirtualProduction/RemoteControl/Source/RemoteControlUI/Private/UI/SRemoteControlPanel.cpp b/Engine/Plugins/VirtualProduction/RemoteControl/Source/RemoteControlUI/Private/UI/SRemoteControlPanel.cpp
--- a/Engine/Plugins/VirtualProduction/RemoteControl/Source/RemoteControlUI/Private/UI/SRemoteControlPanel.cpp
+++ b/Engine/Plugins/VirtualProduction/RemoteControl/Source/RemoteControlUI/Private/UI/SRemoteControlPanel.cpp
@@ -86,16 +86,19 @@ const FName SRemoteControlPanel::DefaultRemoteControlPanelToolBarName("RemoteCon
 const FName SRemoteControlPanel::AuxiliaryRemoteControlPanelToolBarName("RemoteControlPanel.AuxiliaryToolBar");
 const float SRemoteControlPanel::MinimumPanelWidth = 640.f;
 
-TSharedPtr<SBox> SRemoteControlPanel::NoneSelectedWidget = SNew(SBox)
-			.Padding(0.f)
-			.HAlign(HAlign_Center)
-			.VAlign(VAlign_Center)
-			[
-				SNew(STextBlock)
-				.Text(LOCTEXT("NoneSelected", "Select an entity to view details."))
-				.TextStyle(&FAppStyle::GetWidgetStyle<FTextBlockStyle>("NormalText"))
-				.Justification(ETextJustify::Center)
-			];
+TSharedRef<SBox> SRemoteControlPanel::CreateNoneSelectedWidget()
+{
+	return SNew(SBox)
+		.Padding(0.f)
+		.HAlign(HAlign_Center)
+		.VAlign(VAlign_Center)
+		[
+			SNew(STextBlock)
+			.Text(LOCTEXT("NoneSelected", "Select an entity to view details."))
+			.TextStyle(&FAppStyle::GetWidgetStyle<FTextBlockStyle>("NormalText"))
+			.Justification(ETextJustify::Center)
+		];
+}
 
 namespace RemoteControlPanelUtils
 {
@@ -757,11 +760,6 @@ SRemoteControlPanel::~SRemoteControlPanel()
 	}
 }
 
-void SRemoteControlPanel::Shutdown()
-{
-	NoneSelectedWidget.Reset();
-}
-
 void SRemoteControlPanel::Tick(const FGeometry& AllottedGeometry, const double InCurrentTime, const float InDeltaTime)
 {
 	if (bMaterialsCompiledThisFrame)
@@ -1758,7 +1756,7 @@ TSharedRef<SWidget> SRemoteControlPanel::CreateEntityDetailsView()
 	}
 	else
 	{
-		WrappedEntityDetailsView->SetContent(NoneSelectedWidget.ToSharedRef());
+		WrappedEntityDetailsView->SetContent(CreateNoneSelectedWidget());
 	}
 
 	return EntityDetailsDockPanel;
@@ -1794,7 +1792,7 @@ void SRemoteControlPanel::UpdateEntityDetailsView(const TSharedPtr<SRCPanelTreeN
 	}
 	else
 	{
-		WrappedEntityDetailsView->SetContent(NoneSelectedWidget.ToSharedRef());
+		WrappedEntityDetailsView->SetContent(CreateNoneSelectedWidget());
 	}
 
 	static const FName ProtocolWidgetsModuleName = "RemoteControlProtocolWidgets";
@@ -1809,13 +1807,13 @@ void SRemoteControlPanel::UpdateEntityDetailsView(const TSharedPtr<SRCPanelTreeN
 			}
 			else
 			{
-				EntityProtocolDetails->SetContent(NoneSelectedWidget.ToSharedRef());
+				EntityProtocolDetails->SetContent(CreateNoneSelectedWidget());
 			}
 		}
 	}
 	else
 	{
-		EntityProtocolDetails->SetContent(NoneSelectedWidget.ToSharedRef());
+		EntityProtocolDetails->SetContent(CreateNoneSelectedWidget());
 	}
 
 	// Trigger search to list the search results specific to selected group.
diff --git a/Engine/Plugins/VirtualProduction/RemoteControl/Source/RemoteControlUI/Private/UI/SRemoteControlPanel.h b/Engine/Plugins/VirtualProduction/RemoteControl/Source/RemoteControlUI/Private/UI/SRemoteControlPanel.h
--- a/Engine/Plugins/VirtualProduction/RemoteControl/Source/RemoteControlUI/Private/UI/SRemoteControlPanel.h
+++ b/Engine/Plugins/VirtualProduction/RemoteControl/Source/RemoteControlUI/Private/UI/SRemoteControlPanel.h
@@ -75,7 +75,6 @@ public:
 
 	void Construct(const FArguments& InArgs, URemoteControlPreset* InPreset, TSharedPtr<IToolkitHost> InToolkitHost);
 	~SRemoteControlPanel();
-	static void Shutdown();
 
 	//~ Begin SWidget interface
 	virtual void Tick(const FGeometry& AllottedGeometry, const double InCurrentTime, const float InDeltaTime);
@@ -164,6 +163,9 @@ public:
 
 private:
 
+	/** Returns a Helper widget for entity details view and protocol details view. */
+	static TSharedRef<SBox> CreateNoneSelectedWidget();
+
 	//~ Remote Control Commands
 	void BindRemoteControlCommands();
 
@@ -380,8 +382,6 @@ private:
 	TSharedPtr<class IStructureDetailsView> EntityDetailsView;
 	/** Wrapper widget for entity details view. */
 	TSharedPtr<SBorder> WrappedEntityDetailsView;
-	/** Helper widget for entity details view and protocol details view. */
-	static TSharedPtr<SBox> NoneSelectedWidget;
 	/** Holds the field's protocol details. */
 	TSharedPtr<SBox> EntityProtocolDetails;
 	/** Whether to show the rebind all button. */
