[RemoteControl]
Rebinding Property Still Sends Input Original Asset

diff --git a/Engine/Plugins/VirtualProduction/RemoteControl/Source/RemoteControl/Private/RemoteControlEntity.cpp b/Engine/Plugins/VirtualProduction/RemoteControl/Source/RemoteControl/Private/RemoteControlEntity.cpp
--- a/Engine/Plugins/VirtualProduction/RemoteControl/Source/RemoteControl/Private/RemoteControlEntity.cpp
+++ b/Engine/Plugins/VirtualProduction/RemoteControl/Source/RemoteControl/Private/RemoteControlEntity.cpp
@@ -81,6 +81,10 @@ void FRemoteControlEntity::BindObject(UObject* InObjectToBind)
 
 	}
 
+	// The order of this delegate needs to be this way
+	// The first will update the Entity widget with the new data
+	// The second one will refresh the List
+	Owner->OnEntityRebind().Broadcast(Id);
 	OnEntityModifiedDelegate.ExecuteIfBound(Id);
 }
 
diff --git a/Engine/Plugins/VirtualProduction/RemoteControl/Source/RemoteControl/Public/RemoteControlPreset.h b/Engine/Plugins/VirtualProduction/RemoteControl/Source/RemoteControl/Public/RemoteControlPreset.h
--- a/Engine/Plugins/VirtualProduction/RemoteControl/Source/RemoteControl/Public/RemoteControlPreset.h
+++ b/Engine/Plugins/VirtualProduction/RemoteControl/Source/RemoteControl/Public/RemoteControlPreset.h
@@ -600,7 +600,10 @@ public:
 	 * Delegate called when the exposed entity wrapper itself is updated (ie. binding change, rename)  
 	 */
 	FOnPresetEntitiesUpdatedEvent& OnEntitiesUpdated() { return OnEntitiesUpdatedDelegate; }
-	
+
+	DECLARE_MULTICAST_DELEGATE_OneParam(FOnEntityRebind, const FGuid&)
+	FOnEntityRebind& OnEntityRebind() { return OnEntityRebindDelegate; }
+
 	DECLARE_MULTICAST_DELEGATE_TwoParams(FOnPresetPropertyExposed, URemoteControlPreset* /*Preset*/, FName /*ExposedLabel*/);
 	UE_DEPRECATED(4.27, "This delegate is deprecated, use OnEntityExposed instead.")
 	FOnPresetPropertyExposed& OnPropertyExposed() { return OnPropertyExposedDelegate; }
@@ -884,6 +887,8 @@ private:
 	FOnPresetEntityEvent OnEntityUnexposedDelegate;
 	/** Delegate triggered when entities are modified and may need to be re-resolved. */
 	FOnPresetEntitiesUpdatedEvent OnEntitiesUpdatedDelegate;
+	/** Delegate triggered when an Entity is rebound */
+	FOnEntityRebind OnEntityRebindDelegate;
 	/** Delegate triggered when an exposed property value has changed. */
 	FOnPresetExposedPropertiesModified OnPropertyChangedDelegate;
 	/** Delegate triggered when a new property has been exposed. */
diff --git a/Engine/Plugins/VirtualProduction/RemoteControl/Source/RemoteControlUI/Private/UI/SRCPanelExposedEntitiesList.cpp b/Engine/Plugins/VirtualProduction/RemoteControl/Source/RemoteControlUI/Private/UI/SRCPanelExposedEntitiesList.cpp
--- a/Engine/Plugins/VirtualProduction/RemoteControl/Source/RemoteControlUI/Private/UI/SRCPanelExposedEntitiesList.cpp
+++ b/Engine/Plugins/VirtualProduction/RemoteControl/Source/RemoteControlUI/Private/UI/SRCPanelExposedEntitiesList.cpp
@@ -1266,6 +1266,18 @@ void SRCPanelExposedEntitiesList::OnDeleteGroup(const FGuid& GroupId)
 	Preset->Layout.DeleteGroup(GroupId);
 }
 
+void SRCPanelExposedEntitiesList::OnEntityRebind(const FGuid& InEntityGuid)
+{
+	const TSharedPtr<SRCPanelTreeNode>* FoundNode = FieldEntities.FindByPredicate([InEntityGuid] (const TSharedPtr<SRCPanelTreeNode>& InNode) { return InNode->GetRCId() == InEntityGuid; } );
+	if (FoundNode && FoundNode->IsValid())
+	{
+		if (const TSharedPtr<SRCPanelExposedEntity>& Entity = StaticCastSharedPtr<SRCPanelExposedEntity>(*FoundNode))
+		{
+			Entity->Refresh();
+		}
+	}
+}
+
 void SRCPanelExposedEntitiesList::SelectActorsInlevel(const TArray<UObject*>& Objects)
 {
 	if (GEditor)
@@ -1603,6 +1615,7 @@ void SRCPanelExposedEntitiesList::RegisterPresetDelegates()
 	Layout.OnFieldDeleted().AddSP(this, &SRCPanelExposedEntitiesList::OnFieldDeleted);
 	Layout.OnFieldOrderChanged().AddSP(this, &SRCPanelExposedEntitiesList::OnFieldOrderChanged);
 	Preset->OnEntitiesUpdated().AddSP(this, &SRCPanelExposedEntitiesList::OnEntitiesUpdated);
+	Preset->OnEntityRebind().AddSP(this, &SRCPanelExposedEntitiesList::OnEntityRebind);
 }
 
 void SRCPanelExposedEntitiesList::UnregisterPresetDelegates()
@@ -1611,6 +1624,7 @@ void SRCPanelExposedEntitiesList::UnregisterPresetDelegates()
 	{
 		FRemoteControlPresetLayout& Layout = Preset->Layout;
 		Preset->OnEntitiesUpdated().RemoveAll(this);
+		Preset->OnEntityRebind().RemoveAll(this);
 		Layout.OnFieldOrderChanged().RemoveAll(this);
 		Layout.OnFieldDeleted().RemoveAll(this);
 		Layout.OnFieldAdded().RemoveAll(this);
diff --git a/Engine/Plugins/VirtualProduction/RemoteControl/Source/RemoteControlUI/Private/UI/SRCPanelExposedEntitiesList.h b/Engine/Plugins/VirtualProduction/RemoteControl/Source/RemoteControlUI/Private/UI/SRCPanelExposedEntitiesList.h
--- a/Engine/Plugins/VirtualProduction/RemoteControl/Source/RemoteControlUI/Private/UI/SRCPanelExposedEntitiesList.h
+++ b/Engine/Plugins/VirtualProduction/RemoteControl/Source/RemoteControlUI/Private/UI/SRCPanelExposedEntitiesList.h
@@ -160,6 +160,8 @@ private:
 	FReply OnCreateGroup();
 	/** Handles group deletion. */
 	void OnDeleteGroup(const FGuid& GroupId);
+	/** Called when an Entity is rebound */
+	void OnEntityRebind(const FGuid& InEntityGuid);
 	/** Select actors in the current level. */
 	void SelectActorsInlevel(const TArray<UObject*>& Objects);
 	//~ Register to engine/editor events in order to correctly update widgets.
