[RemoteControl]
Add multi selection for Actions and support for multi-selection for the commands.

diff --git a/Engine/Plugins/VirtualProduction/RemoteControl/Source/RemoteControl/Private/RCVirtualProperty.cpp b/Engine/Plugins/VirtualProduction/RemoteControl/Source/RemoteControl/Private/RCVirtualProperty.cpp
--- a/Engine/Plugins/VirtualProduction/RemoteControl/Source/RemoteControl/Private/RCVirtualProperty.cpp
+++ b/Engine/Plugins/VirtualProduction/RemoteControl/Source/RemoteControl/Private/RCVirtualProperty.cpp
@@ -859,6 +859,16 @@ bool URCVirtualPropertySelfContainer::UpdateValueWithProperty(const FProperty* I
 	return Result == EPropertyBagResult::Success;
 }
 
+bool URCVirtualPropertySelfContainer::UpdateValueWithProperty(const URCVirtualPropertyBase* InVirtualProperty)
+{
+	if (InVirtualProperty && InVirtualProperty->GetContainerPtr())
+	{
+		return Bag.SetValue(PropertyName, InVirtualProperty->GetProperty(), InVirtualProperty->GetContainerPtr()) == EPropertyBagResult::Success;	
+	}
+
+	return false;
+}
+
 void URCVirtualPropertySelfContainer::Reset()
 {
 	PropertyName = NAME_None;
diff --git a/Engine/Plugins/VirtualProduction/RemoteControl/Source/RemoteControl/Public/RCVirtualProperty.h b/Engine/Plugins/VirtualProduction/RemoteControl/Source/RemoteControl/Public/RCVirtualProperty.h
--- a/Engine/Plugins/VirtualProduction/RemoteControl/Source/RemoteControl/Public/RCVirtualProperty.h
+++ b/Engine/Plugins/VirtualProduction/RemoteControl/Source/RemoteControl/Public/RCVirtualProperty.h
@@ -416,6 +416,15 @@ public:
  	 */
 	bool UpdateValueWithProperty(const FProperty* InProperty, const void* InPropertyContainer);
 
+	/**
+  	* Update property bag value using URCVirtualPropertyBase passed.
+  	*
+  	* @param InVirtualProperty Virtual Property to update from
+  	*
+  	* @return true if value copied successfully
+  	*/
+	bool UpdateValueWithProperty(const URCVirtualPropertyBase* InVirtualProperty);
+
 	/** Resets the property bag instance to empty and remove Virtual Property data */
 	void Reset();
 
diff --git a/Engine/Plugins/VirtualProduction/RemoteControl/Source/RemoteControlLogic/Private/Action/RCPropertyAction.cpp b/Engine/Plugins/VirtualProduction/RemoteControl/Source/RemoteControlLogic/Private/Action/RCPropertyAction.cpp
--- a/Engine/Plugins/VirtualProduction/RemoteControl/Source/RemoteControlLogic/Private/Action/RCPropertyAction.cpp
+++ b/Engine/Plugins/VirtualProduction/RemoteControl/Source/RemoteControlLogic/Private/Action/RCPropertyAction.cpp
@@ -1,4 +1,4 @@
-﻿// Copyright Epic Games, Inc. All Rights Reserved.
+// Copyright Epic Games, Inc. All Rights Reserved.
 
 #include "Action/RCPropertyAction.h"
 
@@ -20,7 +20,7 @@ URCPropertyAction::URCPropertyAction()
 	PropertySelfContainer = CreateDefaultSubobject<URCVirtualPropertySelfContainer>(FName("VirtualPropertySelfContainer"));
 }
 
-URCBehaviour* URCAction::GetParentBehaviour()
+URCBehaviour* URCAction::GetParentBehaviour() const
 {
 	URCActionContainer* ActionContainer = Cast<URCActionContainer>(GetOuter());
 	if (ensure(ActionContainer))
diff --git a/Engine/Plugins/VirtualProduction/RemoteControl/Source/RemoteControlLogic/Private/Behaviour/Builtin/RangeMap/RCRangeMapBehaviour.cpp b/Engine/Plugins/VirtualProduction/RemoteControl/Source/RemoteControlLogic/Private/Behaviour/Builtin/RangeMap/RCRangeMapBehaviour.cpp
--- a/Engine/Plugins/VirtualProduction/RemoteControl/Source/RemoteControlLogic/Private/Behaviour/Builtin/RangeMap/RCRangeMapBehaviour.cpp
+++ b/Engine/Plugins/VirtualProduction/RemoteControl/Source/RemoteControlLogic/Private/Behaviour/Builtin/RangeMap/RCRangeMapBehaviour.cpp
@@ -1,4 +1,4 @@
-﻿// Copyright Epic Games, Inc. All Rights Reserved.
+// Copyright Epic Games, Inc. All Rights Reserved.
 
 #include "Behaviour/Builtin/RangeMap/RCRangeMapBehaviour.h"
 
@@ -133,6 +133,15 @@ bool URCRangeMapBehaviour::GetNearestActionByThreshold(TTuple<URCAction*, bool>&
 	return true;
 }
 
+bool FRCRangeMapInput::SetInputValue(double InValue) const
+{
+	if (InputProperty)
+	{
+		return InputProperty->SetValueDouble(InValue);
+	}
+	return false;
+}
+
 bool FRCRangeMapInput::GetInputValue(double& OutValue) const
 {
 	if (InputProperty)
diff --git a/Engine/Plugins/VirtualProduction/RemoteControl/Source/RemoteControlLogic/Public/Action/RCAction.h b/Engine/Plugins/VirtualProduction/RemoteControl/Source/RemoteControlLogic/Public/Action/RCAction.h
--- a/Engine/Plugins/VirtualProduction/RemoteControl/Source/RemoteControlLogic/Public/Action/RCAction.h
+++ b/Engine/Plugins/VirtualProduction/RemoteControl/Source/RemoteControlLogic/Public/Action/RCAction.h
@@ -1,4 +1,4 @@
-﻿// Copyright Epic Games, Inc. All Rights Reserved.
+// Copyright Epic Games, Inc. All Rights Reserved.
 
 #pragma once
 
@@ -25,7 +25,7 @@ public:
 	virtual void CopyTo(URCAction* Action) {}
 
 	/** Returns the parent Behaviour associated with this Action */
-	URCBehaviour* GetParentBehaviour();
+	URCBehaviour* GetParentBehaviour() const;
 
 public:
 	/** Exposed Property or Function field Id*/
diff --git a/Engine/Plugins/VirtualProduction/RemoteControl/Source/RemoteControlLogic/Public/Behaviour/Builtin/RangeMap/RCRangeMapBehaviour.h b/Engine/Plugins/VirtualProduction/RemoteControl/Source/RemoteControlLogic/Public/Behaviour/Builtin/RangeMap/RCRangeMapBehaviour.h
--- a/Engine/Plugins/VirtualProduction/RemoteControl/Source/RemoteControlLogic/Public/Behaviour/Builtin/RangeMap/RCRangeMapBehaviour.h
+++ b/Engine/Plugins/VirtualProduction/RemoteControl/Source/RemoteControlLogic/Public/Behaviour/Builtin/RangeMap/RCRangeMapBehaviour.h
@@ -1,4 +1,4 @@
-﻿// Copyright Epic Games, Inc. All Rights Reserved.
+// Copyright Epic Games, Inc. All Rights Reserved.
 
 #pragma once
 
@@ -30,6 +30,9 @@ struct REMOTECONTROLLOGIC_API FRCRangeMapInput
 
 	/** Returns the Step value from the virtual property, safely*/
 	bool GetInputValue(double& OutValue) const;
+
+	/** Set the Step value for the virtual property, return true if successful */
+	bool SetInputValue(double InValue) const;
 };
 
 /**
diff --git a/Engine/Plugins/VirtualProduction/RemoteControl/Source/RemoteControlUI/Private/UI/Action/Conditional/RCActionConditionalModel.cpp b/Engine/Plugins/VirtualProduction/RemoteControl/Source/RemoteControlUI/Private/UI/Action/Conditional/RCActionConditionalModel.cpp
--- a/Engine/Plugins/VirtualProduction/RemoteControl/Source/RemoteControlUI/Private/UI/Action/Conditional/RCActionConditionalModel.cpp
+++ b/Engine/Plugins/VirtualProduction/RemoteControl/Source/RemoteControlUI/Private/UI/Action/Conditional/RCActionConditionalModel.cpp
@@ -11,13 +11,16 @@
 #include "RemoteControlField.h"
 #include "RemoteControlPreset.h"
 #include "Styling/RemoteControlStyles.h"
+#include "UI/Action/SRCActionPanel.h"
 #include "UI/Action/SRCVirtualPropertyWidget.h"
 #include "UI/Behaviour/Builtin/Conditional/RCBehaviourConditionalModel.h"
 #include "UI/RCUIHelpers.h"
 #include "UI/RemoteControlPanelStyle.h"
 #include "UI/SRCPanelDragHandle.h"
 #include "UI/SRCPanelExposedEntity.h"
+#include "UI/SRemoteControlPanel.h"
 #include "Widgets/Layout/SBox.h"
+#include "Widgets/Text/STextBlock.h"
 #include "Widgets/Views/SHeaderRow.h"
 #include "Widgets/Views/SListView.h"
 #include "Widgets/Views/STableRow.h"
@@ -112,6 +115,7 @@ TSharedPtr<SHeaderRow> FRCActionConditionalModel::GetHeaderRow()
 FRCActionConditionalModel::FRCActionConditionalModel(URCAction* InAction, const TSharedPtr<class FRCBehaviourModel> InBehaviourItem, const TSharedPtr<SRemoteControlPanel> InRemoteControlPanel)
 	: FRCActionModel(InAction, InBehaviourItem, InRemoteControlPanel)
 {
+	SAssignNew(ConditionWidget, STextBlock);
 }
 
 TSharedPtr<FRCActionConditionalModel> FRCActionConditionalModel::GetModelByActionType(URCAction* InAction, const TSharedPtr<class FRCBehaviourModel> InBehaviourItem, const TSharedPtr<SRemoteControlPanel> InRemoteControlPanel)
@@ -132,6 +136,38 @@ TSharedPtr<FRCActionConditionalModel> FRCActionConditionalModel::GetModelByActio
 	return nullptr;
 }
 
+FText FRCActionConditionalModel::CreateConditionComparandText(const URCBehaviourConditional* InBehaviour, const FRCBehaviourCondition* InCondition)
+{
+	if (!InBehaviour || !InCondition)
+	{
+		return FText::GetEmpty();
+	}
+
+	const FText ConditionText = InBehaviour->GetConditionTypeAsText(InCondition->ConditionType);
+
+	if (InCondition->ConditionType == ERCBehaviourConditionType::Else)
+	{
+		return ConditionText; // Else doesn't need to display the comparand
+	}
+
+	if (!InCondition->Comparand)
+	{
+		return FText::GetEmpty();
+	}
+
+	const FText ComparandValueText = FText::FromString(InCondition->Comparand->GetDisplayValueAsString());
+
+	return FText::Format(INVTEXT("{0} {1}"), ConditionText, ComparandValueText);
+}
+
+void FRCActionConditionalModel::UpdateConditionWidget(const FText& InNewText) const
+{
+	if (ConditionWidget.IsValid())
+	{
+		ConditionWidget->SetText(InNewText);
+	}
+}
+
 TSharedRef<SWidget> FRCActionConditionalModel::GetConditionWidget()
 {
 	if (TSharedPtr<FRCBehaviourConditionalModel> BehaviourItem = StaticCastSharedPtr<FRCBehaviourConditionalModel>(BehaviourItemWeakPtr.Pin()))
@@ -143,7 +179,8 @@ TSharedRef<SWidget> FRCActionConditionalModel::GetConditionWidget()
 				if(ensure(Condition->Comparand))
 				{
 					return SAssignNew(EditableVirtualPropertyWidget, SRCVirtualPropertyWidget, Condition->Comparand)
-						.OnGenerateWidget(this, &FRCActionConditionalModel::OnGenerateConditionWidget);
+						.OnGenerateWidget(this, &FRCActionConditionalModel::OnGenerateConditionWidget)
+						.OnExitingEditMode(this, &FRCActionConditionalModel::OnExitingEditingMode);
 				}
 			}
 		}
@@ -159,33 +196,87 @@ TSharedRef<SWidget> FRCActionConditionalModel::OnGenerateConditionWidget(URCVirt
 		return SNullWidget::NullWidget;
 	}
 
-	if (TSharedPtr<FRCBehaviourConditionalModel> BehaviourItem = StaticCastSharedPtr<FRCBehaviourConditionalModel>(BehaviourItemWeakPtr.Pin()))
+	if (const TSharedPtr<FRCBehaviourConditionalModel> BehaviourItem = StaticCastSharedPtr<FRCBehaviourConditionalModel>(BehaviourItemWeakPtr.Pin()))
 	{
 		if (URCBehaviourConditional* Behaviour = Cast<URCBehaviourConditional>(BehaviourItem->GetBehaviour()))
 		{
-			if (FRCBehaviourCondition* Condition = Behaviour->Conditions.Find(GetAction()))
+			if (const FRCBehaviourCondition* Condition = Behaviour->Conditions.Find(GetAction()))
 			{
 				const FText ConditionText = Behaviour->GetConditionTypeAsText(Condition->ConditionType);
 
-				FText ConditionComparandText;
+				FText ConditionComparandText = CreateConditionComparandText(Behaviour, Condition);
+
+				UpdateConditionWidget(ConditionComparandText);
 
-				if (Condition->ConditionType == ERCBehaviourConditionType::Else)
+				if (ConditionWidget.IsValid())
 				{
-					ConditionComparandText = ConditionText; // Else doesn't need to display the comparand
+					return ConditionWidget.ToSharedRef();
 				}
-				else
-				{
-					const FText ComparandValueText = FText::FromName(*Condition->Comparand->GetDisplayValueAsString());
+			}
+		}
+	}
 
-					ConditionComparandText = FText::Format(FText::FromName("{0} {1}"), ConditionText, ComparandValueText);
-				}
+	return SNullWidget::NullWidget;
+}
 
-				return SNew(STextBlock).Text(ConditionComparandText);
+void FRCActionConditionalModel::OnExitingEditingMode()
+{
+	if (const TSharedPtr<FRCBehaviourConditionalModel> BehaviourItem = StaticCastSharedPtr<FRCBehaviourConditionalModel>(BehaviourItemWeakPtr.Pin()))
+	{
+		if (URCBehaviourConditional* Behaviour = Cast<URCBehaviourConditional>(BehaviourItem->GetBehaviour()))
+		{
+			if (const FRCBehaviourCondition* Condition = Behaviour->Conditions.Find(GetAction()))
+			{
+				const FText ConditionText = Behaviour->GetConditionTypeAsText(Condition->ConditionType);
+
+				const FText ConditionComparandText = CreateConditionComparandText(Behaviour, Condition);
+
+				UpdateSelectedConditionalActionModel(Condition, ConditionComparandText);
 			}
 		}
 	}
+}
+
+void FRCActionConditionalModel::UpdateSelectedConditionalActionModel(const FRCBehaviourCondition* InConditionToCopy, const FText& InNewConditionText) const
+{
+	if (!InConditionToCopy || !InConditionToCopy->Comparand)
+	{
+		return;
+	}
 
-	return SNullWidget::NullWidget;
+	// Update other selected action based on the edited one
+	if (const TSharedPtr<SRemoteControlPanel> RCPanel = GetRemoteControlPanel())
+    {
+    	if (const TSharedPtr<SRCActionPanel> LogicPanel = RCPanel->GetLogicActionPanel())
+    	{
+    		const TArray<TSharedPtr<FRCLogicModeBase>> SelectedLogicItems = LogicPanel->GetSelectedLogicItems();
+
+    		// If the selected items don't contain this one then skip it
+    		if (SelectedLogicItems.Contains(SharedThis(this)))
+    		{
+    			for (const TSharedPtr<FRCLogicModeBase>& SelectedLogicItem : SelectedLogicItems)
+    			{
+    				if (const TSharedPtr<FRCActionConditionalModel> SelectedConditionalActionModel = StaticCastSharedPtr<FRCActionConditionalModel>(SelectedLogicItem))
+    				{
+    					if (const TSharedPtr<FRCBehaviourModel> SelectedBehaviourModel = SelectedConditionalActionModel->GetParentBehaviour())
+    					{
+    						if (URCBehaviourConditional* SelectedBehaviour = Cast<URCBehaviourConditional>(SelectedBehaviourModel->GetBehaviour()))
+    						{
+    							if (const FRCBehaviourCondition* SelectedCondition = SelectedBehaviour->Conditions.Find(SelectedConditionalActionModel->GetAction()))
+    							{
+    								if (SelectedCondition->Comparand)
+    								{
+    									SelectedCondition->Comparand->UpdateValueWithProperty(InConditionToCopy->Comparand);
+    									SelectedConditionalActionModel->UpdateConditionWidget(InNewConditionText);
+    								}
+    							}
+    						}
+    					}
+    				}
+    			}
+    		}
+    	}
+    }
 }
 
 TSharedRef<ITableRow> FRCActionConditionalModel::OnGenerateWidgetForList(TSharedPtr<FRCActionConditionalModel> InItem, const TSharedRef<STableViewBase>& OwnerTable)
@@ -197,4 +288,4 @@ TSharedRef<ITableRow> FRCActionConditionalModel::OnGenerateWidgetForList(TShared
 		.Padding(FMargin(3.f));
 }
 
-#undef LOCTEXT_NAMESPACE
\ No newline at end of file
+#undef LOCTEXT_NAMESPACE
diff --git a/Engine/Plugins/VirtualProduction/RemoteControl/Source/RemoteControlUI/Private/UI/Action/Conditional/RCActionConditionalModel.h b/Engine/Plugins/VirtualProduction/RemoteControl/Source/RemoteControlUI/Private/UI/Action/Conditional/RCActionConditionalModel.h
--- a/Engine/Plugins/VirtualProduction/RemoteControl/Source/RemoteControlUI/Private/UI/Action/Conditional/RCActionConditionalModel.h
+++ b/Engine/Plugins/VirtualProduction/RemoteControl/Source/RemoteControlUI/Private/UI/Action/Conditional/RCActionConditionalModel.h
@@ -3,13 +3,16 @@
 #pragma once
 
 #include "Action/RCFunctionAction.h"
-#include "Action/RCPropertyIdAction.h"
 #include "Action/RCPropertyAction.h"
+#include "Action/RCPropertyIdAction.h"
 #include "UI/Action/RCActionModel.h"
 
+struct FRCBehaviourCondition;
 class SActionItemListRow;
 class SBox;
 class SRCVirtualPropertyWidget;
+class STextBlock;
+class URCBehaviourConditional;
 
 /* 
 * ~ FRCActionConditionalModel ~
@@ -41,7 +44,23 @@ public:
 	static TSharedPtr<FRCActionConditionalModel> GetModelByActionType(URCAction* InAction, const TSharedPtr<class FRCBehaviourModel> InBehaviourItem, const TSharedPtr<SRemoteControlPanel> InRemoteControlPanel);
 
 private:
+	/** Create the text to be assigned to the ConditionWidget */
+	FText CreateConditionComparandText(const URCBehaviourConditional* InBehaviour, const FRCBehaviourCondition* InCondition);
+
+	/** Update the condition widget text with the new text given */
+	void UpdateConditionWidget(const FText& InNewText) const;
+
+	/** Called when constructing the widget or after editing the condition value */
 	TSharedRef<SWidget> OnGenerateConditionWidget(class URCVirtualPropertySelfContainer* InComparand);
+
+	/** Called when exiting edit mode to update other selected actions */
+	void OnExitingEditingMode();
+
+	/** Called when updating the condition to update every other actions currently selected with the new condition */
+	void UpdateSelectedConditionalActionModel(const FRCBehaviourCondition* InConditionToCopy, const FText& InNewConditionText) const;
+
+	/** Holding the widget that display the condition as text */
+	TSharedPtr<STextBlock> ConditionWidget;
 };
 
 /*
diff --git a/Engine/Plugins/VirtualProduction/RemoteControl/Source/RemoteControlUI/Private/UI/Action/RangeMap/RCActionRangeMapModel.cpp b/Engine/Plugins/VirtualProduction/RemoteControl/Source/RemoteControlUI/Private/UI/Action/RangeMap/RCActionRangeMapModel.cpp
--- a/Engine/Plugins/VirtualProduction/RemoteControl/Source/RemoteControlUI/Private/UI/Action/RangeMap/RCActionRangeMapModel.cpp
+++ b/Engine/Plugins/VirtualProduction/RemoteControl/Source/RemoteControlUI/Private/UI/Action/RangeMap/RCActionRangeMapModel.cpp
@@ -11,13 +11,16 @@
 #include "RemoteControlField.h"
 #include "RemoteControlPreset.h"
 #include "Styling/RemoteControlStyles.h"
+#include "UI/Action/SRCActionPanel.h"
 #include "UI/Action/SRCVirtualPropertyWidget.h"
 #include "UI/Behaviour/Builtin/RangeMap/RCBehaviourRangeMapModel.h"
 #include "UI/RCUIHelpers.h"
 #include "UI/RemoteControlPanelStyle.h"
 #include "UI/SRCPanelDragHandle.h"
 #include "UI/SRCPanelExposedEntity.h"
+#include "UI/SRemoteControlPanel.h"
 #include "Widgets/Layout/SBox.h"
+#include "Widgets/Text/STextBlock.h"
 #include "Widgets/Views/SHeaderRow.h"
 #include "Widgets/Views/SListView.h"
 #include "Widgets/Views/STableRow.h"
@@ -112,6 +115,7 @@ TSharedPtr<SHeaderRow> FRCActionRangeMapModel::GetHeaderRow()
 FRCActionRangeMapModel::FRCActionRangeMapModel(URCAction* InAction, const TSharedPtr<FRCBehaviourModel> InBehaviourItem, const TSharedPtr<SRemoteControlPanel> InRemoteControlPanel)
 	: FRCActionModel(InAction, InBehaviourItem, InRemoteControlPanel)
 {
+	SAssignNew(ConditionWidget, STextBlock);
 }
 
 TSharedPtr<FRCActionRangeMapModel> FRCActionRangeMapModel::GetModelByActionType(URCAction* InAction, const TSharedPtr<class FRCBehaviourModel> InBehaviourItem, const TSharedPtr<SRemoteControlPanel> InRemoteControlPanel)
@@ -139,7 +143,8 @@ TSharedRef<SWidget> FRCActionRangeMapModel::GetInputWidget()
 				if (ensure(RangeMapInput->InputProperty))
 				{
 					return SAssignNew(EditableVirtualPropertyWidget, SRCVirtualPropertyWidget, RangeMapInput->InputProperty)
-						.OnGenerateWidget(this, &FRCActionRangeMapModel::OnGenerateInputWidget);
+						.OnGenerateWidget(this, &FRCActionRangeMapModel::OnGenerateInputWidget)
+						.OnExitingEditMode(this, &FRCActionRangeMapModel::OnExitingEditingMode);
 				}
 			}
 		}
@@ -155,7 +160,7 @@ TSharedRef<SWidget> FRCActionRangeMapModel::OnGenerateInputWidget(class URCVirtu
 		return SNullWidget::NullWidget;
 	}
 
-	if (TSharedPtr<FRCRangeMapBehaviourModel> BehaviourItem = StaticCastSharedPtr<FRCRangeMapBehaviourModel>(BehaviourItemWeakPtr.Pin()))
+	if (const TSharedPtr<FRCRangeMapBehaviourModel> BehaviourItem = StaticCastSharedPtr<FRCRangeMapBehaviourModel>(BehaviourItemWeakPtr.Pin()))
 	{
 		if (URCRangeMapBehaviour* Behaviour = Cast<URCRangeMapBehaviour>(BehaviourItem->GetBehaviour()))
 		{
@@ -166,10 +171,12 @@ TSharedRef<SWidget> FRCActionRangeMapModel::OnGenerateInputWidget(class URCVirtu
 				{
 					const FText StepAsText = FText::FromString(FString::SanitizeFloat(Value));
 
-					return SNew(SBox)
-						[
-							SNew(STextBlock).Text(StepAsText)
-						];
+					UpdateConditionWidget(StepAsText);
+
+					if (ConditionWidget.IsValid())
+					{
+						return ConditionWidget.ToSharedRef();
+					}
 				}
 			}
 		}
@@ -178,6 +185,26 @@ TSharedRef<SWidget> FRCActionRangeMapModel::OnGenerateInputWidget(class URCVirtu
 	return SNullWidget::NullWidget;
 }
 
+void FRCActionRangeMapModel::OnExitingEditingMode() const
+{
+	if (const TSharedPtr<FRCRangeMapBehaviourModel> BehaviourItem = StaticCastSharedPtr<FRCRangeMapBehaviourModel>(BehaviourItemWeakPtr.Pin()))
+	{
+		if (URCRangeMapBehaviour* Behaviour = Cast<URCRangeMapBehaviour>(BehaviourItem->GetBehaviour()))
+		{
+			if (const FRCRangeMapInput* RangeInput = Behaviour->RangeMapActionContainer.Find(GetAction()))
+			{
+				double Value;
+				if(RangeInput->GetInputValue(Value))
+				{
+					const FText StepAsText = FText::FromString(FString::SanitizeFloat(Value));
+
+					UpdateSelectedRangeMapActionModel(Value, StepAsText);
+				}
+			}
+		}
+	}
+}
+
 TSharedRef<ITableRow> FRCActionRangeMapModel::OnGenerateWidgetForList(TSharedPtr<FRCActionRangeMapModel> InItem, const TSharedRef<STableViewBase>& OwnerTable)
 {
 	typedef UE::RCActionPanelRangeMapList::SActionItemListRow ActionRowType;
@@ -187,4 +214,48 @@ TSharedRef<ITableRow> FRCActionRangeMapModel::OnGenerateWidgetForList(TSharedPtr
 		.Padding(3.f);
 }
 
-#undef LOCTEXT_NAMESPACE
\ No newline at end of file
+void FRCActionRangeMapModel::UpdateConditionWidget(const FText& InNewText) const
+{
+	if (ConditionWidget.IsValid())
+	{
+		ConditionWidget->SetText(InNewText);
+	}
+}
+
+void FRCActionRangeMapModel::UpdateSelectedRangeMapActionModel(double InNewValue, const FText& InNewConditionText) const
+{
+	// Update other selected action based on the edited one
+	if (const TSharedPtr<SRemoteControlPanel> RCPanel = GetRemoteControlPanel())
+	{
+		if (const TSharedPtr<SRCActionPanel> LogicPanel = RCPanel->GetLogicActionPanel())
+		{
+			const TArray<TSharedPtr<FRCLogicModeBase>> SelectedLogicItems = LogicPanel->GetSelectedLogicItems();
+
+			// If the selected items don't contain this one then skip it
+			if (SelectedLogicItems.Contains(SharedThis(this)))
+			{
+				for (const TSharedPtr<FRCLogicModeBase>& SelectedLogicItem : SelectedLogicItems)
+				{
+					if (const TSharedPtr<FRCActionRangeMapModel> SelectedRangeMapActionModel = StaticCastSharedPtr<FRCActionRangeMapModel>(SelectedLogicItem))
+					{
+						if (const TSharedPtr<FRCBehaviourModel> SelectedBehaviourModel = SelectedRangeMapActionModel->GetParentBehaviour())
+						{
+							if (URCRangeMapBehaviour* SelectedBehaviour = Cast<URCRangeMapBehaviour>(SelectedBehaviourModel->GetBehaviour()))
+							{
+								if (const FRCRangeMapInput* SelectedRangeInput = SelectedBehaviour->RangeMapActionContainer.Find(SelectedRangeMapActionModel->GetAction()))
+								{
+									if (SelectedRangeInput->SetInputValue(InNewValue))
+									{
+										SelectedRangeMapActionModel->UpdateConditionWidget(InNewConditionText);
+									}
+								}
+							}
+						}
+					}
+				}
+			}
+		}
+	}
+}
+
+#undef LOCTEXT_NAMESPACE
diff --git a/Engine/Plugins/VirtualProduction/RemoteControl/Source/RemoteControlUI/Private/UI/Action/RangeMap/RCActionRangeMapModel.h b/Engine/Plugins/VirtualProduction/RemoteControl/Source/RemoteControlUI/Private/UI/Action/RangeMap/RCActionRangeMapModel.h
--- a/Engine/Plugins/VirtualProduction/RemoteControl/Source/RemoteControlUI/Private/UI/Action/RangeMap/RCActionRangeMapModel.h
+++ b/Engine/Plugins/VirtualProduction/RemoteControl/Source/RemoteControlUI/Private/UI/Action/RangeMap/RCActionRangeMapModel.h
@@ -7,6 +7,8 @@
 #include "Action/RCPropertyAction.h"
 #include "UI/Action/RCActionModel.h"
 
+class STextBlock;
+
 /*
  * ~ FRCActionRangeMapModel ~
  *
@@ -35,7 +37,20 @@ public:
 	static TSharedPtr<FRCActionRangeMapModel> GetModelByActionType(URCAction* InAction, const TSharedPtr<class FRCBehaviourModel> InBehaviourItem, const TSharedPtr<SRemoteControlPanel> InRemoteControlPanel);
 
 private:
+	/** Update the condition widget text with the new text given */
+	void UpdateConditionWidget(const FText& InNewText) const;
+
+	/** Called when updating the condition to update every other actions currently selected with the new condition */
+	void UpdateSelectedRangeMapActionModel(double InNewValue, const FText& InNewConditionText) const;
+
+	/** Called when constructing the widget or after editing the condition value */
 	TSharedRef<SWidget> OnGenerateInputWidget(class URCVirtualPropertySelfContainer* InComparand) const;
+
+	/** Called when exiting edit mode to update other selected actions */
+	void OnExitingEditingMode() const;
+
+	/** Holding the widget that display the condition as text */
+	TSharedPtr<STextBlock> ConditionWidget;
 };
 
 /*
diff --git a/Engine/Plugins/VirtualProduction/RemoteControl/Source/RemoteControlUI/Private/UI/Action/SRCActionPanel.cpp b/Engine/Plugins/VirtualProduction/RemoteControl/Source/RemoteControlUI/Private/UI/Action/SRCActionPanel.cpp
--- a/Engine/Plugins/VirtualProduction/RemoteControl/Source/RemoteControlUI/Private/UI/Action/SRCActionPanel.cpp
+++ b/Engine/Plugins/VirtualProduction/RemoteControl/Source/RemoteControlUI/Private/UI/Action/SRCActionPanel.cpp
@@ -12,6 +12,7 @@
 #include "Action/RCPropertyAction.h"
 
 #include "Controller/RCController.h"
+#include "Behaviour/Builtin/Bind/RCBehaviourBind.h"
 #include "Behaviour/Builtin/Conditional/RCBehaviourConditional.h"
 #include "Behaviour/Builtin/RCBehaviourOnValueChangedNode.h"
 #include "Behaviour/RCBehaviour.h"
@@ -542,19 +543,19 @@ bool SRCActionPanel::IsListFocused() const
 	return ActionPanelList.IsValid() && ActionPanelList->IsListFocused();
 }
 
-void SRCActionPanel::DeleteSelectedPanelItem()
+void SRCActionPanel::DeleteSelectedPanelItems()
 {
-	ActionPanelList->DeleteSelectedPanelItem();
+	ActionPanelList->DeleteSelectedPanelItems();
 }
 
-TSharedPtr<FRCLogicModeBase> SRCActionPanel::GetSelectedLogicItem()
+TArray<TSharedPtr<FRCLogicModeBase>> SRCActionPanel::GetSelectedLogicItems() const
 {
-	if(ActionPanelList)
+	if (ActionPanelList)
 	{
-		return ActionPanelList->GetSelectedLogicItem();
+		return ActionPanelList->GetSelectedLogicItems();
 	}
 
-	return nullptr;
+	return {};
 }
 
 void SRCActionPanel::DuplicateAction(URCAction* InAction)
@@ -601,80 +602,165 @@ void SRCActionPanel::AddNewActionToList(URCAction* NewAction)
 	}
 }
 
-void SRCActionPanel::DuplicateSelectedPanelItem()
+void SRCActionPanel::DuplicateSelectedPanelItems()
 {
 	if (!ensure(ActionPanelList.IsValid()))
 	{
 		return;
 	}
 
-	if (const TSharedPtr<FRCActionModel> ActionItem = StaticCastSharedPtr<FRCActionModel>(ActionPanelList->GetSelectedLogicItem()))
+	for (const TSharedPtr<FRCLogicModeBase>& Action : GetSelectedLogicItems())
 	{
-		DuplicateAction(ActionItem->GetAction());
+		if (const TSharedPtr<FRCActionModel> ActionItem = StaticCastSharedPtr<FRCActionModel>(Action))
+		{
+			DuplicateAction(ActionItem->GetAction());
+		}
 	}
 }
 
-void SRCActionPanel::CopySelectedPanelItem()
+void SRCActionPanel::CopySelectedPanelItems()
 {
-	if (TSharedPtr<SRemoteControlPanel> RemoteControlPanel = GetRemoteControlPanel())
+	if (const TSharedPtr<SRemoteControlPanel>& RemoteControlPanel = GetRemoteControlPanel())
 	{
-		if (TSharedPtr<FRCActionModel> ActionItem = StaticCastSharedPtr<FRCActionModel>(ActionPanelList->GetSelectedLogicItem()))
+		TArray<UObject*> ItemsToCopy;
+		const TArray<TSharedPtr<FRCLogicModeBase>> LogicItems = GetSelectedLogicItems();
+		ItemsToCopy.Reserve(LogicItems.Num());
+
+		for (const TSharedPtr<FRCLogicModeBase>& LogicItem : LogicItems)
 		{
-			RemoteControlPanel->SetLogicClipboardItem(ActionItem->GetAction(), SharedThis(this));
+			if (const TSharedPtr<FRCActionModel> ActionItem = StaticCastSharedPtr<FRCActionModel>(LogicItem))
+			{
+				ItemsToCopy.Add(ActionItem->GetAction());
+			}
 		}
+
+		RemoteControlPanel->SetLogicClipboardItems(ItemsToCopy, SharedThis(this));
 	}
 }
 
-void SRCActionPanel::PasteItemFromClipboard()
+void SRCActionPanel::PasteItemsFromClipboard()
 {
-	if (TSharedPtr<SRemoteControlPanel> RemoteControlPanel = GetRemoteControlPanel())
+	if (const TSharedPtr<SRemoteControlPanel>& RemoteControlPanel = GetRemoteControlPanel())
 	{
 		if (RemoteControlPanel->LogicClipboardItemSource == SharedThis(this))
 		{
-			if(URCAction* Action = Cast<URCAction>(RemoteControlPanel->GetLogicClipboardItem()))
+			for (UObject* LogicClipboardItem : RemoteControlPanel->GetLogicClipboardItems())
+			{
+				if(URCAction* Action = Cast<URCAction>(LogicClipboardItem))
+				{
+					DuplicateAction(Action);
+				}
+			}
+		}
+	}
+}
+
+bool SRCActionPanel::CanPasteClipboardItems(const TArrayView<const TObjectPtr<UObject>> InLogicClipboardItems) const
+{
+	for (const UObject* LogicClipboardItem : InLogicClipboardItems)
+	{
+		const URCAction* LogicClipboardAction = Cast<URCAction>(LogicClipboardItem);
+		if (!LogicClipboardAction)
+		{
+			return false;
+		}
+
+		if (const URCBehaviour* BehaviourSource = LogicClipboardAction->GetParentBehaviour())
+		{
+			if (const TSharedPtr<FRCBehaviourModel>& BehaviourItemTarget = SelectedBehaviourItemWeakPtr.Pin())
+			{
+				if (const URCBehaviour* BehaviourTarget = BehaviourItemTarget->GetBehaviour())
+				{
+					// Copy-paste is allowed between compatible Behaviour types only
+					//
+					return BehaviourSource->GetClass() == BehaviourTarget->GetClass();
+				}
+			}
+		}
+	}
+
+	return false;
+}
+
+void SRCActionPanel::UpdateValue()
+{
+	for (const TSharedPtr<FRCLogicModeBase>& LogicItem : GetSelectedLogicItems())
+	{
+		if (const TSharedPtr<FRCActionModel>& ActionLogicItem = StaticCastSharedPtr<FRCActionModel>(LogicItem))
+		{
+			if (const URCPropertyAction* RCPropertyAction = Cast<URCPropertyAction>(ActionLogicItem->GetAction()))
 			{
-				DuplicateAction(Action);
+				RCPropertyAction->UpdateValueBasedOnRCProperty();
 			}
 		}
 	}
 }
 
-bool SRCActionPanel::CanPasteClipboardItem(UObject* InLogicClipboardItem)
+bool SRCActionPanel::CanUpdateValue() const
 {
-	URCAction* LogicClipboardAction = Cast<URCAction>(InLogicClipboardItem);
-	if (!LogicClipboardAction)
+	const TArray<TSharedPtr<FRCLogicModeBase>>& LogicItems = GetSelectedLogicItems();
+
+	if (LogicItems.IsEmpty())
 	{
 		return false;
 	}
 
-	if (URCBehaviour* BehaviourSource = LogicClipboardAction->GetParentBehaviour())
+	if (const TSharedPtr<FRCActionModel> ActionLogicItem = StaticCastSharedPtr<FRCActionModel>(LogicItems[0]))
 	{
-		if (TSharedPtr<FRCBehaviourModel> BehaviourItemTarget = SelectedBehaviourItemWeakPtr.Pin())
+		if (const TSharedPtr<FRCBehaviourModel> ParentBehaviour = ActionLogicItem->GetParentBehaviour())
 		{
-			if (URCBehaviour* BehaviourTarget = BehaviourItemTarget->GetBehaviour())
+			if (const URCBehaviour* Behaviour = ParentBehaviour->GetBehaviour())
 			{
-				// Copy-paste is allowed between compatible Behaviour types only
-				//
-				return BehaviourSource->GetClass() == BehaviourTarget->GetClass();
+				if (Behaviour->IsA<URCBehaviourBind>())
+				{
+					return false;
+				}
 			}
 		}
 	}
 
+	for (const TSharedPtr<FRCLogicModeBase>& ActionItem : LogicItems)
+	{
+		if (const TSharedPtr<FRCActionModel> ActionLogicItem = StaticCastSharedPtr<FRCActionModel>(ActionItem))
+		{
+			if (const URCAction* RCAction = ActionLogicItem->GetAction())
+			{
+				if (const URemoteControlPreset* RCPreset = GetPreset())
+				{
+					if (RCAction->IsA<URCPropertyAction>()	&&
+						RCPreset->GetExposedEntity(RCAction->ExposedFieldId).IsValid())
+					{
+						// if at least one of the selected actions can update, then enable it
+						return true;
+					}
+				}
+			}
+		}
+	}
 	return false;
 }
 
 FText SRCActionPanel::GetPasteItemMenuEntrySuffix()
 {
-	if (TSharedPtr<SRemoteControlPanel> RemoteControlPanel = GetRemoteControlPanel())
+	if (const TSharedPtr<SRemoteControlPanel> RemoteControlPanel = GetRemoteControlPanel())
 	{
 		// This function should only have been called if we were the source of the item copied.
 		if (ensure(RemoteControlPanel->LogicClipboardItemSource == SharedThis(this)))
 		{
-			if (URCAction* Action = Cast<URCAction>(RemoteControlPanel->GetLogicClipboardItem()))
+			TArray<UObject*> LogicClipboardItems = RemoteControlPanel->GetLogicClipboardItems();
+
+			if (LogicClipboardItems.Num() > 0)
 			{
-				if (URCBehaviour* Behaviour = Action->GetParentBehaviour())
+				if (const URCAction* Action = Cast<URCAction>(LogicClipboardItems[0]))
 				{
-					return FText::Format(FText::FromString("Action {0}"), Behaviour->GetDisplayName());
+					if (URCBehaviour* Behaviour = Action->GetParentBehaviour())
+					{
+						if (LogicClipboardItems.Num() > 1)
+						{
+							return FText::Format(LOCTEXT("ActionPanelPasteMenuMultiEntrySuffix", "Action {0} and {1} other(s)"), Behaviour->GetDisplayName(), (LogicClipboardItems.Num() - 1));
+						}
+						return FText::Format(LOCTEXT("ActionPanelPasteMenuEntrySuffix", "Action {0}"), Behaviour->GetDisplayName());
+					}
 				}
 			}
 		}
@@ -736,7 +822,7 @@ FReply SRCActionPanel::RequestDeleteSelectedItem()
 
 	if (UserResponse == EAppReturnType::Yes)
 	{
-		DeleteSelectedPanelItem();
+		DeleteSelectedPanelItems();
 	}
 
 	return FReply::Handled();
diff --git a/Engine/Plugins/VirtualProduction/RemoteControl/Source/RemoteControlUI/Private/UI/Action/SRCActionPanel.h b/Engine/Plugins/VirtualProduction/RemoteControl/Source/RemoteControlUI/Private/UI/Action/SRCActionPanel.h
--- a/Engine/Plugins/VirtualProduction/RemoteControl/Source/RemoteControlUI/Private/UI/Action/SRCActionPanel.h
+++ b/Engine/Plugins/VirtualProduction/RemoteControl/Source/RemoteControlUI/Private/UI/Action/SRCActionPanel.h
@@ -40,23 +40,29 @@ public:
 	/** Whether the Actions list widget currently has focus.*/
 	bool IsListFocused() const;
 
-	/** Delete Item UI command implementation for this panel */
-	virtual void DeleteSelectedPanelItem() override;
+	/** Delete Items UI command implementation for this panel */
+	virtual void DeleteSelectedPanelItems() override;
 
-	/** Returns the UI item currently selected by the user (if any)*/
-	virtual TSharedPtr<FRCLogicModeBase> GetSelectedLogicItem() override;
+	/** Returns the UI items currently selected by the user (if any). To be implemented per child panel */
+	virtual TArray<TSharedPtr<FRCLogicModeBase>> GetSelectedLogicItems() const override;
 
-	/** "Duplicate Item" UI command implementation for Action panel*/
-	virtual void DuplicateSelectedPanelItem() override;
+	/** "Duplicate Items" UI command implementation for Action panel*/
+	virtual void DuplicateSelectedPanelItems() override;
 
-	/** "Copy Item" UI command implementation for Action panel*/
-	virtual void CopySelectedPanelItem() override;
+	/** "Copy Items" UI command implementation for Action panel*/
+	virtual void CopySelectedPanelItems() override;
 
-	/** "Paste Item" UI command implementation for Action panel*/
-	virtual void PasteItemFromClipboard() override;
+	/** "Paste Items" UI command implementation for Action panel*/
+	virtual void PasteItemsFromClipboard() override;
 
-	/** Whether a given clipboard item can be successfully pasted into this panel */
-	virtual bool CanPasteClipboardItem(UObject* InLogicClipboardItem) override;
+	/** Whether clipboard items can be successfully pasted into this panel */
+	virtual bool CanPasteClipboardItems(const TArrayView<const TObjectPtr<UObject>> InLogicClipboardItems) const override;
+
+	/** "Update Value" UI command implementation for panels */
+	virtual void UpdateValue() override;
+
+	/** Whether this panel can call the UpdateValue command */
+	virtual bool CanUpdateValue() const override;
 
 	/** Provides an item suffix for the Paste context menu to provide users with useful context on the nature of the item being pasted */
 	virtual FText GetPasteItemMenuEntrySuffix() override;
diff --git a/Engine/Plugins/VirtualProduction/RemoteControl/Source/RemoteControlUI/Private/UI/Action/SRCActionPanelList.h b/Engine/Plugins/VirtualProduction/RemoteControl/Source/RemoteControlUI/Private/UI/Action/SRCActionPanelList.h
--- a/Engine/Plugins/VirtualProduction/RemoteControl/Source/RemoteControlUI/Private/UI/Action/SRCActionPanelList.h
+++ b/Engine/Plugins/VirtualProduction/RemoteControl/Source/RemoteControlUI/Private/UI/Action/SRCActionPanelList.h
@@ -64,6 +64,7 @@ public:
 
 		ListView = SNew(SListView<TSharedPtr<ActionType>>)
 			.ListItemsSource(&ActionItems)
+			.SelectionMode(ESelectionMode::Multi)
 			.HeaderRow(ActionType::GetHeaderRow())
 			.OnGenerateRow(this, &SRCActionPanelList::OnGenerateWidgetForList)
 			.OnSelectionChanged(this, &SRCActionPanelList::OnSelectionChanged)
@@ -130,37 +131,51 @@ public:
 		return ActionItems.Num();
 	}
 
-	/** The number of Controllers currently selected*/
+	/** The number of Controllers currently selected */
 	virtual int32 NumSelectedLogicItems() const override
 	{
 		return ListView->GetNumItemsSelected();
 	}
 
-	/** Whether the Actions List View currently has focus.*/
+	/** Whether the Actions List View currently has focus. */
 	virtual bool IsListFocused() const override
 	{
 		return ListView->HasAnyUserFocus().IsSet() || ContextMenuWidgetCached.IsValid();
 	}
 
-	/** Deletes currently selected items from the list view*/
-	virtual void DeleteSelectedPanelItem() override
+	/** Deletes currently selected items from the list view */
+	virtual void DeleteSelectedPanelItems() override
 	{
-		DeleteItemFromLogicPanel<ActionType>(ActionItems, ListView->GetSelectedItems());
+		DeleteItemsFromLogicPanel<ActionType>(ActionItems, ListView->GetSelectedItems());
 	}
 
-	/** Returns the UI item currently selected by the user (if any)*/
-	virtual TSharedPtr<FRCLogicModeBase> GetSelectedLogicItem() override
+	/** Returns the UI items currently selected by the user (if any). */
+	virtual TArray<TSharedPtr<FRCLogicModeBase>> GetSelectedLogicItems() override
 	{
-		return GetSelectedActionItem();
+		TArray<TSharedPtr<FRCLogicModeBase>> SelectedValidLogicItems;
+		if (ListView.IsValid())
+		{
+			TArray<TSharedPtr<ActionType>> AllSelectedLogicItems = ListView->GetSelectedItems();
+			SelectedValidLogicItems.Reserve(AllSelectedLogicItems.Num());
+
+			for (const TSharedPtr<ActionType>& LogicItem : AllSelectedLogicItems)
+			{
+				if (LogicItem.IsValid())
+				{
+					SelectedValidLogicItems.Add(LogicItem);
+				}
+			}
+		}
+		return SelectedValidLogicItems;
 	}
 
-	/** The currently selected Action item*/
+	/** The currently selected Action item */
 	TSharedPtr<FRCActionModel> GetSelectedActionItem()
 	{
 		return SelectedActionItem;
 	}
 
-	/** Fetches the parent Action panel*/
+	/** Fetches the parent Action panel */
 	TSharedPtr<SRCActionPanel> GetActionPanel()
 	{
 		return ActionPanelWeakPtr.Pin();
@@ -172,7 +187,7 @@ public:
 		return BehaviourItemWeakPtr.Pin();
 	}
 
-	/** Adds an Action by Remote Control Field Guid*/
+	/** Adds an Action by Remote Control Field Guid */
 	URCAction* AddAction(const FGuid& InRemoteControlFieldId)
 	{
 		if (const URemoteControlPreset* Preset = GetPreset())
@@ -236,7 +251,7 @@ public:
 
 private:
 
-	/** OnGenerateRow delegate for the Actions List View*/
+	/** OnGenerateRow delegate for the Actions List View */
 	TSharedRef<ITableRow> OnGenerateWidgetForList( TSharedPtr<ActionType> InItem, const TSharedRef<STableViewBase>& OwnerTable )
 	{
 		if (ensure(InItem))
@@ -250,7 +265,7 @@ private:
 			];
 	}
 
-	/** Responds to the selection of a newly created action. Resets UI state*/
+	/** Responds to the selection of a newly created action. Resets UI state */
 	void OnActionAdded(URCAction* InAction)
 	{
 		// Historical note: Previously we used to call Reset here after adding an Action
@@ -258,13 +273,13 @@ private:
 		// @todo: This improvement needs to be done for "Remove Action" as well.
 	}
 
-	/** Responds to the removal of all actions. Rests UI state*/
+	/** Responds to the removal of all actions. Rests UI state */
 	void OnEmptyActions()
 	{
 		Reset();
 	}
 
-	/** Refreshes the list from the latest state of the data model*/
+	/** Refreshes the list from the latest state of the data model */
 	virtual void Reset() override
 	{
 		ActionItems.Empty();
@@ -289,7 +304,7 @@ private:
 		ListView->RebuildList();
 	}
 
-	/** Handles broadcasting of a successful remove item operation.*/
+	/** Handles broadcasting of a successful remove item operation. */
 	virtual void BroadcastOnItemRemoved() override {}
 
 	/** Fetches the Remote Control preset associated with the parent panel */
@@ -303,7 +318,7 @@ private:
 		return nullptr;
 	}
 
-	/** Removes the given Action UI model item from the list of UI models*/
+	/** Removes the given Action UI model item from the list of UI models */
 	virtual int32 RemoveModel(const TSharedPtr<FRCLogicModeBase> InModel) override
 	{
 		if (const TSharedPtr<FRCBehaviourModel> BehaviourModel = BehaviourItemWeakPtr.Pin())
@@ -447,23 +462,23 @@ private:
 
 private:
 
-	/** The currently selected Action item*/
+	/** The currently selected Action item */
 	TSharedPtr<FRCActionModel> SelectedActionItem;
 
-	/** The parent Action Panel widget*/
+	/** The parent Action Panel widget */
 	TWeakPtr<SRCActionPanel> ActionPanelWeakPtr;
 
-	/** The Behaviour (UI model) associated with us*/
+	/** The Behaviour (UI model) associated with us */
 	TWeakPtr<FRCBehaviourModel> BehaviourItemWeakPtr;
 
 	/** List of Actions (UI model) active in this widget */
 	TArray<TSharedPtr<ActionType>> ActionItems;
 
-	/** List View widget for representing our Actions List*/
+	/** List View widget for representing our Actions List */
 	TSharedPtr<SListView<TSharedPtr<ActionType>>> ListView;
 
 	/** Panel Style reference. */
 	const FRCPanelStyle* RCPanelStyle;
 };
 
-#undef LOCTEXT_NAMESPACE
\ No newline at end of file
+#undef LOCTEXT_NAMESPACE
diff --git a/Engine/Plugins/VirtualProduction/RemoteControl/Source/RemoteControlUI/Private/UI/Action/SRCVirtualPropertyWidget.cpp b/Engine/Plugins/VirtualProduction/RemoteControl/Source/RemoteControlUI/Private/UI/Action/SRCVirtualPropertyWidget.cpp
--- a/Engine/Plugins/VirtualProduction/RemoteControl/Source/RemoteControlUI/Private/UI/Action/SRCVirtualPropertyWidget.cpp
+++ b/Engine/Plugins/VirtualProduction/RemoteControl/Source/RemoteControlUI/Private/UI/Action/SRCVirtualPropertyWidget.cpp
@@ -20,7 +20,8 @@ void SRCVirtualPropertyWidget::Construct(const FArguments& InArgs, URCVirtualPro
 		return;
 	}
 
-	this->OnGenerateWidget = InArgs._OnGenerateWidget;
+	OnGenerateWidget = InArgs._OnGenerateWidget;
+	OnExitingEditModeDelegate = InArgs._OnExitingEditMode;
 
 	VirtualPropertyWeakPtr = InVirtualProperty;
 
@@ -94,6 +95,9 @@ void SRCVirtualPropertyWidget::ExitEditMode()
 	{
 		TSharedRef<SWidget> VirtualPropertyDisplayWidget = OnGenerateWidget.Execute(VirtualPropertyWeakPtr.Get());
 
+		// Used to update other selected actions for Conditional and Range behaviour
+		OnExitingEditModeDelegate.ExecuteIfBound();
+
 		// Restore the readonly view of the virtual property widget
 		VirtualPropertyWidgetBox->SetContent(VirtualPropertyDisplayWidget);
 	}
diff --git a/Engine/Plugins/VirtualProduction/RemoteControl/Source/RemoteControlUI/Private/UI/Action/SRCVirtualPropertyWidget.h b/Engine/Plugins/VirtualProduction/RemoteControl/Source/RemoteControlUI/Private/UI/Action/SRCVirtualPropertyWidget.h
--- a/Engine/Plugins/VirtualProduction/RemoteControl/Source/RemoteControlUI/Private/UI/Action/SRCVirtualPropertyWidget.h
+++ b/Engine/Plugins/VirtualProduction/RemoteControl/Source/RemoteControlUI/Private/UI/Action/SRCVirtualPropertyWidget.h
@@ -12,6 +12,8 @@ class IPropertyRowGenerator;
 class SBox;
 class URCVirtualPropertySelfContainer;
 
+DECLARE_DELEGATE(FOnExitingEditMode)
+
 /*
 * ~ SRCVirtualPropertyWidget ~
 *
@@ -35,6 +37,8 @@ public:
 
 	SLATE_EVENT(FOnGenerateWidget, OnGenerateWidget)
 
+	SLATE_EVENT(FOnExitingEditMode, OnExitingEditMode)
+
 	SLATE_END_ARGS()
 
 	/** Constructs this widget with InArgs */
@@ -74,4 +78,8 @@ private:
 	TSharedPtr<IPropertyRowGenerator> PropertyRowGenerator;
 
 	/** Detail tree node used to create a generic value widget*/
-	TWeakPtr<IDetailTreeNode> DetailTreeNodeWeakPtr;};
+	TWeakPtr<IDetailTreeNode> DetailTreeNodeWeakPtr;
+
+	/** Delegate executed when exiting edit mode if bound */
+	FOnExitingEditMode OnExitingEditModeDelegate;
+};
diff --git a/Engine/Plugins/VirtualProduction/RemoteControl/Source/RemoteControlUI/Private/UI/BaseLogicUI/RCLogicModeBase.h b/Engine/Plugins/VirtualProduction/RemoteControl/Source/RemoteControlUI/Private/UI/BaseLogicUI/RCLogicModeBase.h
--- a/Engine/Plugins/VirtualProduction/RemoteControl/Source/RemoteControlUI/Private/UI/BaseLogicUI/RCLogicModeBase.h
+++ b/Engine/Plugins/VirtualProduction/RemoteControl/Source/RemoteControlUI/Private/UI/BaseLogicUI/RCLogicModeBase.h
@@ -1,4 +1,4 @@
-﻿// Copyright Epic Games, Inc. All Rights Reserved.
+// Copyright Epic Games, Inc. All Rights Reserved.
 
 #pragma once
 
@@ -32,7 +32,7 @@ public:
 	URemoteControlPreset* GetPreset() const;
 
 	/** Returns the Remote Control Preset object associated with us*/
-	TSharedPtr< SRemoteControlPanel> GetRemoteControlPanel() const;
+	TSharedPtr<SRemoteControlPanel> GetRemoteControlPanel() const;
 
 protected:
 	/** The parent Remote Control Panel widget*/
diff --git a/Engine/Plugins/VirtualProduction/RemoteControl/Source/RemoteControlUI/Private/UI/BaseLogicUI/SRCLogicPanelBase.h b/Engine/Plugins/VirtualProduction/RemoteControl/Source/RemoteControlUI/Private/UI/BaseLogicUI/SRCLogicPanelBase.h
--- a/Engine/Plugins/VirtualProduction/RemoteControl/Source/RemoteControlUI/Private/UI/BaseLogicUI/SRCLogicPanelBase.h
+++ b/Engine/Plugins/VirtualProduction/RemoteControl/Source/RemoteControlUI/Private/UI/BaseLogicUI/SRCLogicPanelBase.h
@@ -1,8 +1,9 @@
-﻿// Copyright Epic Games, Inc. All Rights Reserved.
+// Copyright Epic Games, Inc. All Rights Reserved.
 
 #pragma once
 
 #include "CoreMinimal.h"
+#include "SRCLogicPanelListBase.h"
 #include "Widgets/SCompoundWidget.h"
 
 class FRCLogicModeBase;
@@ -35,23 +36,29 @@ public:
 	* This is the main container holding all the Remote Control panels including Logic and Exposed Properties*/
 	virtual TSharedPtr<SRemoteControlPanel> GetRemoteControlPanel() const;
 
-	/** "Delete Item" UI command implementation for panels*/
-	virtual void DeleteSelectedPanelItem() = 0;
+	/** "Delete Items" UI command implementation for panels*/
+	virtual void DeleteSelectedPanelItems() = 0;
 
-	/** "Duplicate Item" UI command implementation for panels*/
-	virtual void DuplicateSelectedPanelItem() {}
+	/** "Duplicate Items" UI command implementation for action panels */
+	virtual void DuplicateSelectedPanelItems() {}
 
-	/** "Copy Item" UI command implementation for panels*/
-	virtual void CopySelectedPanelItem() {}
+	/** "Copy Items" UI command implementation for panels*/
+	virtual void CopySelectedPanelItems() {}
 
-	/** "Paste Item" UI command implementation for panels*/
-	virtual void PasteItemFromClipboard() {}
+	/** "Paste Items" UI command implementation for panels*/
+	virtual void PasteItemsFromClipboard() {}
 
-	/** Whether a given clipboard item can be successfully pasted into this panel */
-	virtual bool CanPasteClipboardItem(UObject* InLogicClipboardItem) { return true; }
+	/** Whether a given clipboard items can be successfully pasted into this panel */
+	virtual bool CanPasteClipboardItems(const TArrayView<const TObjectPtr<UObject>> InLogicClipboardItems) const { return true; }
 
-	/** Returns the UI item currently selected by the user (if any). To be implemented per child panel*/
-	virtual TSharedPtr<FRCLogicModeBase> GetSelectedLogicItem() = 0;
+	/** "Update Value" UI command implementation for panels */
+	virtual void UpdateValue() {}
+
+	/** Whether this panel can call the UpdateValue command */
+	virtual bool CanUpdateValue() const { return false; }
+
+	/** Returns the UI items currently selected by the user (if any). To be implemented per child panel*/
+	virtual TArray<TSharedPtr<FRCLogicModeBase>> GetSelectedLogicItems() const = 0;
 
 	/** Provides an item suffix for the Paste context menu to provide users with useful context on the nature of the item being pasted */
 	virtual FText GetPasteItemMenuEntrySuffix() { return FText::GetEmpty(); }
@@ -63,6 +70,9 @@ public:
 	virtual FReply RequestDeleteAllItems() = 0;
 
 protected:
+	/** Widget representing List of Controllers/Behaviours/Actions */
+	TSharedPtr<class SRCLogicPanelListBase> PanelList;
+
 	/** The parent Remote Control Panel widget*/
 	TWeakPtr<SRemoteControlPanel> PanelWeakPtr;
 };
diff --git a/Engine/Plugins/VirtualProduction/RemoteControl/Source/RemoteControlUI/Private/UI/BaseLogicUI/SRCLogicPanelListBase.h b/Engine/Plugins/VirtualProduction/RemoteControl/Source/RemoteControlUI/Private/UI/BaseLogicUI/SRCLogicPanelListBase.h
--- a/Engine/Plugins/VirtualProduction/RemoteControl/Source/RemoteControlUI/Private/UI/BaseLogicUI/SRCLogicPanelListBase.h
+++ b/Engine/Plugins/VirtualProduction/RemoteControl/Source/RemoteControlUI/Private/UI/BaseLogicUI/SRCLogicPanelListBase.h
@@ -1,4 +1,4 @@
-﻿// Copyright Epic Games, Inc. All Rights Reserved.
+// Copyright Epic Games, Inc. All Rights Reserved.
 
 #pragma once
 
@@ -42,16 +42,16 @@ public:
 	/** The number of items currently selected in a Logic list panel */
 	virtual int32 NumSelectedLogicItems() const = 0;
 
-	/** Whether the List View currently has focus.*/
+	/** Whether the List View currently has focus. */
 	virtual bool IsListFocused() const = 0;
 
-	/** Deletes currently selected items from the list view*/
-	virtual void DeleteSelectedPanelItem()  = 0;
+	/** Deletes currently selected items from the list view */
+	virtual void DeleteSelectedPanelItems()  = 0;
 
-	/** Returns the UI item currently selected by the user (if any). To be implemented per child panel*/
-	virtual TSharedPtr<FRCLogicModeBase> GetSelectedLogicItem() = 0;
+	/** Returns the UI items currently selected by the user (if any). To be implemented per child panel */
+	virtual TArray<TSharedPtr<FRCLogicModeBase>> GetSelectedLogicItems() = 0;
 
-	/** Provides a common entry point for adding Logic related data objects (UObjects) to their repsective panels*/
+	/** Provides a common entry point for adding Logic related data objects (UObjects) to their respective panels */
 	virtual void AddNewLogicItem(UObject* InLogicItem) {}
 
 	/** Builds the right-click context menu populated with generic actions (based on UI Commands) */
@@ -92,7 +92,7 @@ protected:
 	/** Helper function for handling common Delete Item functionality across all child panels (Actions/Behaviours/Controllers)
 	* Currently invoked from each Panel List child class with appropriate template class*/
 	template<class T>
-	void DeleteItemFromLogicPanel(TArray<TSharedPtr<T>>& ItemsSource, const TArray<TSharedPtr<T>>& SelectedItems)
+	void DeleteItemsFromLogicPanel(TArray<TSharedPtr<T>>& ItemsSource, const TArray<TSharedPtr<T>>& SelectedItems)
 	{
 		bool bIsDeleted = false;
 		for (const TSharedPtr<T> SelectedItem : SelectedItems)
diff --git a/Engine/Plugins/VirtualProduction/RemoteControl/Source/RemoteControlUI/Private/UI/Behaviour/SRCBehaviourPanel.cpp b/Engine/Plugins/VirtualProduction/RemoteControl/Source/RemoteControlUI/Private/UI/Behaviour/SRCBehaviourPanel.cpp
--- a/Engine/Plugins/VirtualProduction/RemoteControl/Source/RemoteControlUI/Private/UI/Behaviour/SRCBehaviourPanel.cpp
+++ b/Engine/Plugins/VirtualProduction/RemoteControl/Source/RemoteControlUI/Private/UI/Behaviour/SRCBehaviourPanel.cpp
@@ -1,4 +1,4 @@
-﻿// Copyright Epic Games, Inc. All Rights Reserved.
+// Copyright Epic Games, Inc. All Rights Reserved.
 
 #include "SRCBehaviourPanel.h"
 #include "SRCBehaviourPanelList.h"
@@ -219,9 +219,9 @@ bool SRCBehaviourPanel::IsListFocused() const
 	return BehaviourPanelList.IsValid() && BehaviourPanelList->IsListFocused();
 }
 
-void SRCBehaviourPanel::DeleteSelectedPanelItem()
+void SRCBehaviourPanel::DeleteSelectedPanelItems()
 {
-	BehaviourPanelList->DeleteSelectedPanelItem();
+	BehaviourPanelList->DeleteSelectedPanelItems();
 }
 
 void SRCBehaviourPanel::DuplicateBehaviour(URCBehaviour* InBehaviour)
@@ -234,74 +234,103 @@ void SRCBehaviourPanel::DuplicateBehaviour(URCBehaviour* InBehaviour)
 	}
 }
 
-void SRCBehaviourPanel::DuplicateSelectedPanelItem()
+void SRCBehaviourPanel::DuplicateSelectedPanelItems()
 {
 	FScopedTransaction Transaction(LOCTEXT("DuplicateBehaviour", "Duplicate Behaviour"));
 
-	if (TSharedPtr<FRCBehaviourModel> SelectedBehaviourItem = BehaviourPanelList->GetSelectedBehaviourItem())
+	for (const TSharedPtr<FRCLogicModeBase>& LogicItem : GetSelectedLogicItems())
 	{
-		DuplicateBehaviour(SelectedBehaviourItem->GetBehaviour());
+		if (const TSharedPtr<FRCBehaviourModel>& BehaviourItem = StaticCastSharedPtr<FRCBehaviourModel>(LogicItem))
+		{
+			DuplicateBehaviour(BehaviourItem->GetBehaviour());
+		}
 	}
 }
 
-void SRCBehaviourPanel::CopySelectedPanelItem()
+void SRCBehaviourPanel::CopySelectedPanelItems()
 {
-	if (TSharedPtr<SRemoteControlPanel> RemoteControlPanel = GetRemoteControlPanel())
+	if (const TSharedPtr<SRemoteControlPanel>& RemoteControlPanel = GetRemoteControlPanel())
 	{
-		if (TSharedPtr<FRCBehaviourModel> SelectedBehaviourItem = BehaviourPanelList->GetSelectedBehaviourItem())
+		TArray<UObject*> ItemsToCopy;
+		const TArray<TSharedPtr<FRCLogicModeBase>> LogicItems = GetSelectedLogicItems();
+		ItemsToCopy.Reserve(LogicItems.Num());
+
+		for (const TSharedPtr<FRCLogicModeBase>& LogicItem : LogicItems)
 		{
-			RemoteControlPanel->SetLogicClipboardItem(SelectedBehaviourItem->GetBehaviour(), SharedThis(this));
+			if (const TSharedPtr<FRCBehaviourModel>& SelectedBehaviourItem = StaticCastSharedPtr<FRCBehaviourModel>(LogicItem))
+			{
+				ItemsToCopy.Add(SelectedBehaviourItem->GetBehaviour());
+			}
 		}
+
+		RemoteControlPanel->SetLogicClipboardItems(ItemsToCopy, SharedThis(this));
 	}
 }
 
-void SRCBehaviourPanel::PasteItemFromClipboard()
+void SRCBehaviourPanel::PasteItemsFromClipboard()
 {
 	FScopedTransaction Transaction(LOCTEXT("PasteBehaviour", "Paste Behaviour"));
 
-	if (TSharedPtr<SRemoteControlPanel> RemoteControlPanel = GetRemoteControlPanel())
+	if (const TSharedPtr<SRemoteControlPanel>& RemoteControlPanel = GetRemoteControlPanel())
 	{
 		if (RemoteControlPanel->LogicClipboardItemSource == SharedThis(this))
 		{
-			if(URCBehaviour* Behaviour = Cast<URCBehaviour>(RemoteControlPanel->GetLogicClipboardItem()))
+			for (UObject* LogicItem : RemoteControlPanel->GetLogicClipboardItems())
 			{
-				DuplicateBehaviour(Behaviour);
+				if(URCBehaviour* Behaviour = Cast<URCBehaviour>(LogicItem))
+				{
+					DuplicateBehaviour(Behaviour);
+				}
 			}
 		}
 	}
 }
 
-bool SRCBehaviourPanel::CanPasteClipboardItem(UObject* InLogicClipboardItem)
+bool SRCBehaviourPanel::CanPasteClipboardItems(const TArrayView<const TObjectPtr<UObject>> InLogicClipboardItems) const
 {
-	URCBehaviour* LogicClipboardBehaviour = Cast<URCBehaviour>(InLogicClipboardItem);
-
-	URCController* ControllerSource = LogicClipboardBehaviour->ControllerWeakPtr.Get();
-	if (!ControllerSource)
+	for (const UObject* LogicClipboardItem : InLogicClipboardItems)
 	{
-		return false;
-	}
+		const URCBehaviour* LogicClipboardBehaviour = Cast<URCBehaviour>(LogicClipboardItem);
 
-	URCController* ControllerTarget = GetParentController();
-	if (!ControllerTarget)
-	{
-		return false;
+		const URCController* ControllerSource = LogicClipboardBehaviour->ControllerWeakPtr.Get();
+		if (!ControllerSource)
+		{
+			return false;
+		}
+
+		const URCController* ControllerTarget = GetParentController();
+		if (!ControllerTarget)
+		{
+			return false;
+		}
+
+		// Copy-Paste is only permitted between Controllers of the same type
+		//
+		return ControllerSource->GetValueType() == ControllerTarget->GetValueType();
 	}
 
-	// Copy-Paste is only permitted between Controllers of the same type
-	//
-	return ControllerSource->GetValueType() == ControllerTarget->GetValueType();
+	return false;
 }
 
 FText SRCBehaviourPanel::GetPasteItemMenuEntrySuffix()
 {
-	if (TSharedPtr<SRemoteControlPanel> RemoteControlPanel = GetRemoteControlPanel())
+	if (const TSharedPtr<SRemoteControlPanel>& RemoteControlPanel = GetRemoteControlPanel())
 	{
 		// This function should only have been called if we were the source of the item copied.
 		if (ensure(RemoteControlPanel->LogicClipboardItemSource == SharedThis(this)))
 		{
-			if (URCBehaviour* Behaviour = Cast<URCBehaviour>(RemoteControlPanel->GetLogicClipboardItem()))
+			TArray<UObject*> LogicClipboardItems = RemoteControlPanel->GetLogicClipboardItems();
+
+			if (LogicClipboardItems.Num() > 0)
 			{
-				return FText::Format(FText::FromString("Behaviour {0}"), Behaviour->GetDisplayName());
+				if (URCBehaviour* Behaviour = Cast<URCBehaviour>(LogicClipboardItems[0]))
+				{
+					if (LogicClipboardItems.Num() > 1)
+					{
+						return FText::Format(LOCTEXT("BehaviourPanelPasteMenuMultiEntrySuffix", "Behaviour {0} and {1} other(s)"), Behaviour->GetDisplayName(), (LogicClipboardItems.Num() - 1));
+					}
+					return FText::Format(LOCTEXT("BehaviourPanelPasteMenuEntrySuffix", "Behaviour {0}"), Behaviour->GetDisplayName());
+				}
 			}
 		}
 	}
@@ -309,19 +338,19 @@ FText SRCBehaviourPanel::GetPasteItemMenuEntrySuffix()
 	return FText::GetEmpty();
 }
 
-TSharedPtr<FRCLogicModeBase> SRCBehaviourPanel::GetSelectedLogicItem()
+TArray<TSharedPtr<FRCLogicModeBase>> SRCBehaviourPanel::GetSelectedLogicItems() const
 {
 	if (BehaviourPanelList)
 	{
-		return BehaviourPanelList->GetSelectedBehaviourItem();
+		return BehaviourPanelList->GetSelectedLogicItems();
 	}
 
-	return nullptr;
+	return {};
 }
 
-URCController* SRCBehaviourPanel::GetParentController()
+URCController* SRCBehaviourPanel::GetParentController() const
 {
-	if (TSharedPtr<FRCControllerModel> ControllerItem = SelectedControllerItemWeakPtr.Pin())
+	if (const TSharedPtr<FRCControllerModel> ControllerItem = SelectedControllerItemWeakPtr.Pin())
 	{
 		return Cast<URCController>(ControllerItem->GetVirtualProperty());
 	}
@@ -342,7 +371,7 @@ FReply SRCBehaviourPanel::RequestDeleteSelectedItem()
 
 	if (UserResponse == EAppReturnType::Yes)
 	{
-		DeleteSelectedPanelItem();
+		DeleteSelectedPanelItems();
 	}
 
 	return FReply::Handled();
diff --git a/Engine/Plugins/VirtualProduction/RemoteControl/Source/RemoteControlUI/Private/UI/Behaviour/SRCBehaviourPanel.h b/Engine/Plugins/VirtualProduction/RemoteControl/Source/RemoteControlUI/Private/UI/Behaviour/SRCBehaviourPanel.h
--- a/Engine/Plugins/VirtualProduction/RemoteControl/Source/RemoteControlUI/Private/UI/Behaviour/SRCBehaviourPanel.h
+++ b/Engine/Plugins/VirtualProduction/RemoteControl/Source/RemoteControlUI/Private/UI/Behaviour/SRCBehaviourPanel.h
@@ -1,4 +1,4 @@
-﻿// Copyright Epic Games, Inc. All Rights Reserved.
+// Copyright Epic Games, Inc. All Rights Reserved.
 
 #pragma once
 
@@ -33,29 +33,29 @@ public:
 	/** Whether the Behaviour list widget currently has focus. Used for Delete Item UI command */
 	bool IsListFocused() const;
 
-	/** Delete Item UI command implementation for this panel */
-	virtual void DeleteSelectedPanelItem() override;
+	/** Delete Items UI command implementation for this panel */
+	virtual void DeleteSelectedPanelItems() override;
 
-	/** "Duplicate Item" UI command implementation for Behaviour panel*/
-	virtual void DuplicateSelectedPanelItem() override;
+	/** "Duplicate Items" UI command implementation for Behaviour panel*/
+	virtual void DuplicateSelectedPanelItems() override;
 
-	/** "Copy Item" UI command implementation for Behaviour panel*/
-	virtual void CopySelectedPanelItem() override;
+	/** "Copy Items" UI command implementation for Behaviour panel*/
+	virtual void CopySelectedPanelItems() override;
 
-	/** "Paste Item" UI command implementation for Behaviour panel*/
-	virtual void PasteItemFromClipboard() override;
+	/** "Paste Items" UI command implementation for Behaviour panel*/
+	virtual void PasteItemsFromClipboard() override;
 
-	/** Whether a given clipboard item can be successfully pasted into this panel */
-	virtual bool CanPasteClipboardItem(UObject* InLogicClipboardItem) override;
+	/** Whether a given clipboard items can be successfully pasted into this panel */
+	virtual bool CanPasteClipboardItems(const TArrayView<const TObjectPtr<UObject>> InLogicClipboardItems) const override;
 
 	/** Provides an item suffix for the Paste context menu to provide users with useful context on the nature of the item being pasted */
 	virtual FText GetPasteItemMenuEntrySuffix() override;
 
-	/** Returns the UI item currently selected by the user (if any)*/
-	virtual TSharedPtr<FRCLogicModeBase> GetSelectedLogicItem() override;
+	/** Returns the UI items currently selected by the user (if any). To be implemented per child panel */
+	virtual TArray<TSharedPtr<FRCLogicModeBase>> GetSelectedLogicItems() const override;
 
 	/** Returns the parent Controller associated with this behaviour*/
-	URCController* GetParentController();
+	URCController* GetParentController() const;
 
 protected:
 
diff --git a/Engine/Plugins/VirtualProduction/RemoteControl/Source/RemoteControlUI/Private/UI/Behaviour/SRCBehaviourPanelList.cpp b/Engine/Plugins/VirtualProduction/RemoteControl/Source/RemoteControlUI/Private/UI/Behaviour/SRCBehaviourPanelList.cpp
--- a/Engine/Plugins/VirtualProduction/RemoteControl/Source/RemoteControlUI/Private/UI/Behaviour/SRCBehaviourPanelList.cpp
+++ b/Engine/Plugins/VirtualProduction/RemoteControl/Source/RemoteControlUI/Private/UI/Behaviour/SRCBehaviourPanelList.cpp
@@ -1,4 +1,4 @@
-﻿// Copyright Epic Games, Inc. All Rights Reserved.
+// Copyright Epic Games, Inc. All Rights Reserved.
 
 #include "SRCBehaviourPanelList.h"
 
@@ -101,33 +101,35 @@ void SRCBehaviourPanelList::AddNewLogicItem(UObject* InLogicItem)
 
 void SRCBehaviourPanelList::AddSpecialContextMenuOptions(FMenuBuilder& MenuBuilder)
 {
-	if (TSharedPtr<FRCBehaviourModel> BehaviourItem = GetSelectedBehaviourItem())
-	{
-		const bool bIsEnabled = BehaviourItem->IsBehaviourEnabled();
-
-		FUIAction Action(FExecuteAction::CreateSP(this, &SRCBehaviourPanelList::SetIsBehaviourEnabled, !bIsEnabled));
+	const FUIAction EnableAction(FExecuteAction::CreateSP(this, &SRCBehaviourPanelList::SetIsBehaviourEnabled, true));
+	const FUIAction DisableAction(FExecuteAction::CreateSP(this, &SRCBehaviourPanelList::SetIsBehaviourEnabled, false));
 
-		FText Label, Tooltip;
-		if (bIsEnabled)
-		{
-			Label = LOCTEXT("LabelDisableBehaviour", "Disable");
-			Tooltip = LOCTEXT("TooltipDisableBehaviour", "Disables the current behaviour. Actions will not be processed for this behaviour when the Controller value changes");
-		}
-		else
-		{
-			Label = LOCTEXT("LabelEnableBehaviour", "Enable");
-			Tooltip = LOCTEXT("TooltipEnableBehaviour", "Enables the current behaviour. Restores functioning of Actions associated with this behaviour.");
-		}
+	static const FText EnableLabel = LOCTEXT("LabelEnableBehaviour", "Enable");
+	static const FText EnableTooltip = LOCTEXT("TooltipEnableBehaviour", "Enables the current behaviours. Restores functioning of Actions associated with these behaviours.");
+	static const FText DisableLabel = LOCTEXT("LabelDisableBehaviour", "Disable");
+	static const FText DisableTooltip = LOCTEXT("TooltipDisableBehaviour", "Disables the current behaviours. Actions will not be processed for these behaviours when the Controller value changes");
 
-		MenuBuilder.AddMenuEntry(Label, Tooltip, FSlateIcon(), Action);
-	}
+	MenuBuilder.AddMenuEntry(EnableLabel, EnableTooltip, FSlateIcon(), EnableAction);
+	MenuBuilder.AddMenuEntry(DisableLabel, DisableTooltip, FSlateIcon(), DisableAction);
 }
 
 void SRCBehaviourPanelList::SetIsBehaviourEnabled(const bool bIsEnabled)
 {
-	if (TSharedPtr<SRemoteControlPanel> RemoteControlPanel = GetRemoteControlPanel())
+	// Disable all selected behaviour here
+	for (const TSharedPtr<FRCLogicModeBase>& LogicItem : GetSelectedLogicItems())
+	{
+		if (const TSharedPtr<FRCBehaviourModel>& BehaviourModel = StaticCastSharedPtr<FRCBehaviourModel>(LogicItem))
+		{
+			BehaviourModel->SetIsBehaviourEnabled(bIsEnabled);
+		}
+	}
+
+	// Disable only the ActionPanel of the currently selected behaviour.
+	// When selecting a different behaviour the ActionPanel will be re-constructed,
+	// it will get the current status of the behaviour (enabled or disabled) and will match it for the ActionPanel as well
+	if (const TSharedPtr<SRemoteControlPanel> RemoteControlPanel = GetRemoteControlPanel())
 	{
-		if (TSharedPtr<SRCActionPanel> ActionPanel = RemoteControlPanel->GetLogicActionPanel())
+		if (const TSharedPtr<SRCActionPanel> ActionPanel = RemoteControlPanel->GetLogicActionPanel())
 		{
 			ActionPanel->SetIsBehaviourEnabled(bIsEnabled);
 		}
@@ -278,9 +280,28 @@ bool SRCBehaviourPanelList::IsListFocused() const
 	return ListView->HasAnyUserFocus().IsSet() || ContextMenuWidgetCached.IsValid();;
 }
 
-void SRCBehaviourPanelList::DeleteSelectedPanelItem()
+void SRCBehaviourPanelList::DeleteSelectedPanelItems()
+{
+	DeleteItemsFromLogicPanel<FRCBehaviourModel>(BehaviourItems, ListView->GetSelectedItems());
+}
+
+TArray<TSharedPtr<FRCLogicModeBase>> SRCBehaviourPanelList::GetSelectedLogicItems()
 {
-	DeleteItemFromLogicPanel<FRCBehaviourModel>(BehaviourItems, ListView->GetSelectedItems());
+	TArray<TSharedPtr<FRCLogicModeBase>> SelectedValidLogicItems;
+	if (ListView.IsValid())
+	{
+		TArray<TSharedPtr<FRCBehaviourModel>> AllSelectedLogicItems = ListView->GetSelectedItems();
+		SelectedValidLogicItems.Reserve(AllSelectedLogicItems.Num());
+
+		for (const TSharedPtr<FRCBehaviourModel>& LogicItem : AllSelectedLogicItems)
+		{
+			if (LogicItem.IsValid())
+			{
+				SelectedValidLogicItems.Add(LogicItem);
+			}
+		}
+	}
+	return SelectedValidLogicItems;
 }
 
 void SRCBehaviourPanelList::RequestRefresh()
diff --git a/Engine/Plugins/VirtualProduction/RemoteControl/Source/RemoteControlUI/Private/UI/Behaviour/SRCBehaviourPanelList.h b/Engine/Plugins/VirtualProduction/RemoteControl/Source/RemoteControlUI/Private/UI/Behaviour/SRCBehaviourPanelList.h
--- a/Engine/Plugins/VirtualProduction/RemoteControl/Source/RemoteControlUI/Private/UI/Behaviour/SRCBehaviourPanelList.h
+++ b/Engine/Plugins/VirtualProduction/RemoteControl/Source/RemoteControlUI/Private/UI/Behaviour/SRCBehaviourPanelList.h
@@ -1,4 +1,4 @@
-﻿// Copyright Epic Games, Inc. All Rights Reserved.
+// Copyright Epic Games, Inc. All Rights Reserved.
 
 #pragma once
 
@@ -40,20 +40,17 @@ public:
 	/** Returns number of items in the list. */
 	virtual int32 Num() const override;
 
-	/** The number of Controllers currently selected*/
+	/** The number of Controllers currently selected */
 	virtual int32 NumSelectedLogicItems() const override;
 
-	/** Whether the Behaviours List View currently has focus.*/
+	/** Whether the Behaviours List View currently has focus. */
 	virtual bool IsListFocused() const override;
 
-	/** Deletes currently selected items from the list view*/
-	virtual void DeleteSelectedPanelItem() override;
+	/** Deletes currently selected items from the list view */
+	virtual void DeleteSelectedPanelItems() override;
 
-	/** Returns the UI item currently selected by the user (if any)*/
-	virtual TSharedPtr<FRCLogicModeBase> GetSelectedLogicItem() override
-	{
-		return GetSelectedBehaviourItem();
-	}
+	/** Returns the UI items currently selected by the user (if any). */
+	virtual TArray<TSharedPtr<FRCLogicModeBase>> GetSelectedLogicItems() override;
 
 	TSharedPtr<FRCBehaviourModel> GetSelectedBehaviourItem()
 	{
@@ -69,7 +66,7 @@ public:
 
 private:
 
-	/** Enables or Disables the currently selected behaviour*/
+	/** Enables or Disables the currently selected behaviour */
 	void SetIsBehaviourEnabled(const bool bIsEnabled);
 
 	void AddBehaviourToList(URCBehaviour* InBehaviour);
@@ -86,16 +83,16 @@ private:
 	/** Responds to the removal of all Behaviours. Rests UI state */
 	void OnEmptyBehaviours();
 
-	/** Refreshes the list from the latest state of the model*/
+	/** Refreshes the list from the latest state of the model */
 	virtual void Reset() override;
 
-	/** Handles broadcasting of a successful remove item operation.*/
+	/** Handles broadcasting of a successful remove item operation. */
 	virtual void BroadcastOnItemRemoved() override;
 
 	/** Fetches the Remote Control preset associated with the parent panel */
 	virtual URemoteControlPreset* GetPreset() override;
 
-	/** Removes the given Behaviour UI model item from the list of UI models*/
+	/** Removes the given Behaviour UI model item from the list of UI models */
 	virtual int32 RemoveModel(const TSharedPtr<FRCLogicModeBase> InModel) override;
 
 	void OnBehaviourListModified();
@@ -114,7 +111,7 @@ private:
 	/** List of Behaviours (UI models) active in this widget */
 	TArray<TSharedPtr<FRCBehaviourModel>> BehaviourItems;
 
-	/** List View widget for representing our Behaivours List*/
+	/** List View widget for representing our Behaviours List */
 	TSharedPtr<SListView<TSharedPtr<FRCBehaviourModel>>> ListView;
 
 	/** Panel Style reference. */
diff --git a/Engine/Plugins/VirtualProduction/RemoteControl/Source/RemoteControlUI/Private/UI/Controller/SRCControllerPanel.cpp b/Engine/Plugins/VirtualProduction/RemoteControl/Source/RemoteControlUI/Private/UI/Controller/SRCControllerPanel.cpp
--- a/Engine/Plugins/VirtualProduction/RemoteControl/Source/RemoteControlUI/Private/UI/Controller/SRCControllerPanel.cpp
+++ b/Engine/Plugins/VirtualProduction/RemoteControl/Source/RemoteControlUI/Private/UI/Controller/SRCControllerPanel.cpp
@@ -1,4 +1,4 @@
-﻿// Copyright Epic Games, Inc. All Rights Reserved.
+// Copyright Epic Games, Inc. All Rights Reserved.
 
 #include "SRCControllerPanel.h"
 
@@ -127,19 +127,19 @@ bool SRCControllerPanel::IsListFocused() const
 	return ControllerPanelList->IsListFocused();
 }
 
-void SRCControllerPanel::DeleteSelectedPanelItem()
+void SRCControllerPanel::DeleteSelectedPanelItems()
 {
 	if (ControllerPanelList)
 	{
-		ControllerPanelList->DeleteSelectedPanelItem();
+		ControllerPanelList->DeleteSelectedPanelItems();
 	}
 }
 
-void SRCControllerPanel::DuplicateSelectedPanelItem()
+void SRCControllerPanel::DuplicateSelectedPanelItems()
 {
-	if (ControllerPanelList)
+	for (const TSharedPtr<FRCLogicModeBase>& LogicItem : GetSelectedLogicItems())
 	{
-		if(TSharedPtr<FRCControllerModel> ControllerItem = ControllerPanelList->GetSelectedControllerItem())
+		if (const TSharedPtr<FRCControllerModel>& ControllerItem = StaticCastSharedPtr<FRCControllerModel>(LogicItem))
 		{
 			if (URCController* Controller = Cast<URCController>(ControllerItem->GetVirtualProperty()))
 			{
@@ -149,29 +149,41 @@ void SRCControllerPanel::DuplicateSelectedPanelItem()
 	}
 }
 
-void SRCControllerPanel::CopySelectedPanelItem()
+void SRCControllerPanel::CopySelectedPanelItems()
 {
-	if (TSharedPtr<SRemoteControlPanel> RemoteControlPanel = GetRemoteControlPanel())
+	if (const TSharedPtr<SRemoteControlPanel>& RemoteControlPanel = GetRemoteControlPanel())
 	{
-		if (TSharedPtr<FRCControllerModel> ControllerItem = ControllerPanelList->GetSelectedControllerItem())
+		TArray<UObject*> ItemsToCopy;
+		const TArray<TSharedPtr<FRCLogicModeBase>> LogicItems = GetSelectedLogicItems();
+		ItemsToCopy.Reserve(LogicItems.Num());
+
+		for (const TSharedPtr<FRCLogicModeBase>& LogicItem : LogicItems)
 		{
-			if (URCController* Controller = Cast<URCController>(ControllerItem->GetVirtualProperty()))
+			if (const TSharedPtr<FRCControllerModel>& ControllerItem = StaticCastSharedPtr<FRCControllerModel>(LogicItem))
 			{
-				RemoteControlPanel->SetLogicClipboardItem(Controller, SharedThis(this));
+				if (URCController* Controller = Cast<URCController>(ControllerItem->GetVirtualProperty()))
+				{
+					ItemsToCopy.Add(Controller);
+				}
 			}
 		}
+
+		RemoteControlPanel->SetLogicClipboardItems(ItemsToCopy, SharedThis(this));
 	}
 }
 
-void SRCControllerPanel::PasteItemFromClipboard()
+void SRCControllerPanel::PasteItemsFromClipboard()
 {
-	if (TSharedPtr<SRemoteControlPanel> RemoteControlPanel = GetRemoteControlPanel())
+	if (const TSharedPtr<SRemoteControlPanel>& RemoteControlPanel = GetRemoteControlPanel())
 	{
-		if(RemoteControlPanel->LogicClipboardItemSource == SharedThis(this))
+		if (RemoteControlPanel->LogicClipboardItemSource == SharedThis(this))
 		{
-			if (URCController* Controller = Cast<URCController>(RemoteControlPanel->GetLogicClipboardItem()))
+			for (UObject* LogicClipboardItem : RemoteControlPanel->GetLogicClipboardItems())
 			{
-				DuplicateController(Controller);
+				if (URCController* Controller = Cast<URCController>(LogicClipboardItem))
+				{
+					DuplicateController(Controller);
+				}
 			}
 		}
 	}
@@ -179,14 +191,23 @@ void SRCControllerPanel::PasteItemFromClipboard()
 
 FText SRCControllerPanel::GetPasteItemMenuEntrySuffix()
 {
-	if (TSharedPtr<SRemoteControlPanel> RemoteControlPanel = GetRemoteControlPanel())
+	if (const TSharedPtr<SRemoteControlPanel> RemoteControlPanel = GetRemoteControlPanel())
 	{
 		// This function should only have been called if we were the source of the item copied.
 		if (ensure(RemoteControlPanel->LogicClipboardItemSource == SharedThis(this)))
 		{
-			if (URCController* Controller = Cast<URCController>(RemoteControlPanel->GetLogicClipboardItem()))
+			TArray<UObject*> LogicClipboardItems = RemoteControlPanel->GetLogicClipboardItems();
+
+			if (LogicClipboardItems.Num() > 0)
 			{
-				return FText::Format(FText::FromString("Controller {0}"), FText::FromName(Controller->DisplayName));
+				if (const URCController* Controller = Cast<URCController>(LogicClipboardItems[0]))
+				{
+					if (LogicClipboardItems.Num() > 1)
+					{
+						return FText::Format(LOCTEXT("ControllerPanelPasteMenuMultiEntrySuffix", "Controller {0} and {1} other(s)"), FText::FromName(Controller->DisplayName), (LogicClipboardItems.Num() - 1));
+					}
+					return FText::Format(LOCTEXT("ControllerPanelPasteMenuMultiEntrySuffix", "Controller {0}"), FText::FromName(Controller->DisplayName));
+				}
 			}
 		}
 	}
@@ -194,14 +215,14 @@ FText SRCControllerPanel::GetPasteItemMenuEntrySuffix()
 	return FText::GetEmpty();
 }
 
-TSharedPtr<FRCLogicModeBase> SRCControllerPanel::GetSelectedLogicItem()
+TArray<TSharedPtr<FRCLogicModeBase>> SRCControllerPanel::GetSelectedLogicItems() const
 {
 	if (ControllerPanelList)
 	{
-		return ControllerPanelList->GetSelectedControllerItem();
+		return ControllerPanelList->GetSelectedLogicItems();
 	}
 
-	return nullptr;
+	return {};
 }
 
 void SRCControllerPanel::DuplicateController(URCController* InController)
@@ -235,7 +256,7 @@ FReply SRCControllerPanel::RequestDeleteSelectedItem()
 
 	if (UserResponse == EAppReturnType::Yes)
 	{
-		DeleteSelectedPanelItem();
+		DeleteSelectedPanelItems();
 	}
 
 	return FReply::Handled();
diff --git a/Engine/Plugins/VirtualProduction/RemoteControl/Source/RemoteControlUI/Private/UI/Controller/SRCControllerPanel.h b/Engine/Plugins/VirtualProduction/RemoteControl/Source/RemoteControlUI/Private/UI/Controller/SRCControllerPanel.h
--- a/Engine/Plugins/VirtualProduction/RemoteControl/Source/RemoteControlUI/Private/UI/Controller/SRCControllerPanel.h
+++ b/Engine/Plugins/VirtualProduction/RemoteControl/Source/RemoteControlUI/Private/UI/Controller/SRCControllerPanel.h
@@ -1,4 +1,4 @@
-﻿// Copyright Epic Games, Inc. All Rights Reserved.
+// Copyright Epic Games, Inc. All Rights Reserved.
 
 #pragma once
 
@@ -58,23 +58,23 @@ public:
 	/** Whether the Controller list widget currently has focus. Used for Delete Item UI command */
 	bool IsListFocused() const;
 
-	/** Delete Item UI command implementation for this panel */
-	virtual void DeleteSelectedPanelItem() override;
+	/** Delete Items UI command implementation for this panel */
+	virtual void DeleteSelectedPanelItems() override;
 
-	/** "Duplicate Item" UI command implementation for Controller panel*/
-	virtual void DuplicateSelectedPanelItem() override;
+	/** "Duplicate Items" UI command implementation for Controller panel*/
+	virtual void DuplicateSelectedPanelItems() override;
 
-	/** "Copy Item" UI command implementation for Controller panel*/
-	virtual void CopySelectedPanelItem() override;
+	/** "Copy Items" UI command implementation for Controller panel*/
+	virtual void CopySelectedPanelItems() override;
 
-	/** "Paste Item" UI command implementation for Controller panel*/
-	virtual void PasteItemFromClipboard() override;
+	/** "Paste Items" UI command implementation for Controller panel*/
+	virtual void PasteItemsFromClipboard() override;
 
 	/** Provides an item suffix for the Paste context menu to provide users with useful context on the nature of the item being pasted */
 	virtual FText GetPasteItemMenuEntrySuffix() override;
 
-	/** Returns the UI item currently selected by the user (if any)*/
-	virtual TSharedPtr<FRCLogicModeBase> GetSelectedLogicItem() override;
+	/** Returns the UI items currently selected by the user (if any). To be implemented per child panel*/
+	virtual TArray<TSharedPtr<FRCLogicModeBase>> GetSelectedLogicItems() const override;
 
 	void EnterRenameMode();
 
diff --git a/Engine/Plugins/VirtualProduction/RemoteControl/Source/RemoteControlUI/Private/UI/Controller/SRCControllerPanelList.cpp b/Engine/Plugins/VirtualProduction/RemoteControl/Source/RemoteControlUI/Private/UI/Controller/SRCControllerPanelList.cpp
--- a/Engine/Plugins/VirtualProduction/RemoteControl/Source/RemoteControlUI/Private/UI/Controller/SRCControllerPanelList.cpp
+++ b/Engine/Plugins/VirtualProduction/RemoteControl/Source/RemoteControlUI/Private/UI/Controller/SRCControllerPanelList.cpp
@@ -1,4 +1,4 @@
-﻿// Copyright Epic Games, Inc. All Rights Reserved.
+// Copyright Epic Games, Inc. All Rights Reserved.
 
 #include "SRCControllerPanelList.h"
 
@@ -732,9 +732,19 @@ bool SRCControllerPanelList::IsListFocused() const
 	return ListView->HasAnyUserFocus().IsSet() || ContextMenuWidgetCached.IsValid();
 }
 
-void SRCControllerPanelList::DeleteSelectedPanelItem()
+void SRCControllerPanelList::DeleteSelectedPanelItems()
 {
-	DeleteItemFromLogicPanel<FRCControllerModel>(ControllerItems, ListView->GetSelectedItems());
+	DeleteItemsFromLogicPanel<FRCControllerModel>(ControllerItems, ListView->GetSelectedItems());
+}
+
+TArray<TSharedPtr<FRCLogicModeBase>> SRCControllerPanelList::GetSelectedLogicItems()
+{
+	// Controllers don't support multi selection
+	if (const TSharedPtr<FRCLogicModeBase> SelectedControllerItemPtr = SelectedControllerItemWeakPtr.Pin())
+	{
+		return { SelectedControllerItemPtr };
+	}
+	return TArray<TSharedPtr<FRCLogicModeBase>>();
 }
 
 void SRCControllerPanelList::NotifyPreChange(FEditPropertyChain* PropertyAboutToChange)
diff --git a/Engine/Plugins/VirtualProduction/RemoteControl/Source/RemoteControlUI/Private/UI/Controller/SRCControllerPanelList.h b/Engine/Plugins/VirtualProduction/RemoteControl/Source/RemoteControlUI/Private/UI/Controller/SRCControllerPanelList.h
--- a/Engine/Plugins/VirtualProduction/RemoteControl/Source/RemoteControlUI/Private/UI/Controller/SRCControllerPanelList.h
+++ b/Engine/Plugins/VirtualProduction/RemoteControl/Source/RemoteControlUI/Private/UI/Controller/SRCControllerPanelList.h
@@ -1,4 +1,4 @@
-﻿// Copyright Epic Games, Inc. All Rights Reserved.
+// Copyright Epic Games, Inc. All Rights Reserved.
 
 #pragma once
 #include "DragAndDrop/DecoratedDragDropOp.h"
@@ -79,20 +79,17 @@ public:
 	/** Returns number of items in the list. */
 	virtual int32 Num() const override;
 
-	/** The number of Controllers currently selected*/
+	/** The number of Controllers currently selected */
 	virtual int32 NumSelectedLogicItems() const override;
 
-	/** Whether the Controllers List View currently has focus.*/
+	/** Whether the Controllers List View currently has focus. */
 	virtual bool IsListFocused() const override;
 
-	/** Deletes currently selected items from the list view*/
-	virtual void DeleteSelectedPanelItem() override;
+	/** Deletes currently selected items from the list view */
+	virtual void DeleteSelectedPanelItems() override;
 
-	/** Returns the UI item currently selected by the user (if any)*/
-	virtual TSharedPtr<FRCLogicModeBase> GetSelectedLogicItem() override
-	{
-		return GetSelectedControllerItem();
-	}
+	/** Returns the UI items currently selected by the user (if any). */
+	virtual TArray<TSharedPtr<FRCLogicModeBase>> GetSelectedLogicItems() override;
 
 	/** FNotifyHook Interface Begin */
 	virtual void NotifyPreChange(FEditPropertyChain* PropertyAboutToChange) override;
@@ -105,20 +102,20 @@ public:
 		return ControllerItems.Num();
 	}
 
-	/** Finds a Controller UI model by unique Id*/
+	/** Finds a Controller UI model by unique Id */
 	TSharedPtr<FRCControllerModel> FindControllerItemById(const FGuid& InId) const;
 
 	/** Given an item to move and an anchor row this function moves the item to the position of the anchor
 	* and pushes all other rows below */
 	void ReorderControllerItem(TSharedRef<FRCControllerModel> ItemToMove, TSharedRef<FRCControllerModel> AnchorItem);
 
-	/** Returns the currently selected Controller UI Item*/
+	/** Returns the currently selected Controller UI Item */
 	TSharedPtr<FRCControllerModel> GetSelectedControllerItem() const
 	{
 		return SelectedControllerItemWeakPtr.Pin();
 	}
 
-	/** Requests the panel to refresh its contents from the latest list of Controllers*/
+	/** Requests the panel to refresh its contents from the latest list of Controllers */
 	void RequestRefresh()
 	{
 		Reset();
@@ -127,16 +124,16 @@ public:
 	/** Drag-Drop validation delegate for the Controllers Panel List */
 	bool OnAllowDrop(TSharedPtr<FDragDropOperation> DragDropOperation);
 
-	/** Drag-Drop action delegate for the Controllers Panel List*/
+	/** Drag-Drop action delegate for the Controllers Panel List */
 	FReply OnControllerListViewDragDrop(TSharedPtr<FDragDropOperation> DragDropOperation);
 
 	/** Fetches the Remote Control preset associated with the parent panel */
 	virtual URemoteControlPreset* GetPreset() override;
 
-	/** Creates a Bind Behaviour for the given Controller and binds the given remote control property to it*/
+	/** Creates a Bind Behaviour for the given Controller and binds the given remote control property to it */
 	void CreateBindBehaviourAndAssignTo(URCController* Controller, TSharedRef<const FRemoteControlProperty> InRemoteControlProperty, const bool bExecuteBind);
 
-	/** Whether the user's cursor is directly hovered over the List View*/
+	/** Whether the user's cursor is directly hovered over the List View */
 	bool IsListViewHovered();
 
 	/** Enable/disable MultiController Mode */
@@ -206,16 +203,16 @@ private:
 	/** List of Controllers (UI model) active in this widget */
 	TArray<TSharedPtr<FRCControllerModel>> ControllerItems;
 
-	/** List View widget for representing our Controllers List*/
+	/** List View widget for representing our Controllers List */
 	TSharedPtr<SListView<TSharedPtr<FRCControllerModel>>> ListView;
 	
-	/** Refreshes the list from the latest state of the model*/
+	/** Refreshes the list from the latest state of the model */
 	virtual void Reset() override;
 
-	/** Handles broadcasting of a successful remove item operation.*/
+	/** Handles broadcasting of a successful remove item operation. */
 	virtual void BroadcastOnItemRemoved() override;
 
-	/** Removes the given Controller UI model item from the list of UI models*/
+	/** Removes the given Controller UI model item from the list of UI models */
 	virtual int32 RemoveModel(const TSharedPtr<FRCLogicModeBase> InModel) override;
 
 	/** Set the visibility of the List header the Value Type Column used for MultiControllers */
@@ -227,7 +224,7 @@ private:
 	/** Add a custom column */
 	void AddColumn(const FName& InColumnName);
 	
-	/** Checks whether the provided Entity can be used to create a Controller*/
+	/** Checks whether the provided Entity can be used to create a Controller */
 	bool IsEntitySupported(FGuid ExposedEntityId);
 
 	/** Panel Style reference. */
@@ -244,4 +241,3 @@ private:
 	/** List of custom columns names */
 	TArray<FName> CustomColumns;
 };
-
diff --git a/Engine/Plugins/VirtualProduction/RemoteControl/Source/RemoteControlUI/Private/UI/SRemoteControlPanel.cpp b/Engine/Plugins/VirtualProduction/RemoteControl/Source/RemoteControlUI/Private/UI/SRemoteControlPanel.cpp
--- a/Engine/Plugins/VirtualProduction/RemoteControl/Source/RemoteControlUI/Private/UI/SRemoteControlPanel.cpp
+++ b/Engine/Plugins/VirtualProduction/RemoteControl/Source/RemoteControlUI/Private/UI/SRemoteControlPanel.cpp
@@ -100,6 +100,11 @@ TSharedRef<SBox> SRemoteControlPanel::CreateNoneSelectedWidget()
 		];
 }
 
+void SRemoteControlPanel::AddReferencedObjects(FReferenceCollector& Collector)
+{
+	Collector.AddReferencedObjects(LogicClipboardItems);
+}
+
 namespace RemoteControlPanelUtils
 {
 	bool IsExposableActor(AActor* Actor)
@@ -2466,9 +2471,9 @@ bool SRemoteControlPanel::CanDeleteEntity() const
 		return false;
 	}
 
-	if (TSharedPtr<SRCLogicPanelBase> ActiveLogicPanel = GetActiveLogicPanel())
+	if (const TSharedPtr<SRCLogicPanelBase> ActiveLogicPanel = GetActiveLogicPanel())
 	{
-		return ActiveLogicPanel->GetSelectedLogicItem() != nullptr; // User has focus on a logic panel
+		return !ActiveLogicPanel->GetSelectedLogicItems().IsEmpty(); // User has focus on a logic panel
 	}
 
 	if (LastSelectedEntity.IsValid() && Preset.IsValid())
@@ -2488,7 +2493,8 @@ void SRemoteControlPanel::RenameEntity_Execute() const
 		return;
 	}
 
-	if (LastSelectedEntity->GetRCType() == SRCPanelTreeNode::FieldChild) // Field Child does not contain entity ID, that is why it should not be processed
+	if (LastSelectedEntity->GetRCType() == SRCPanelTreeNode::FieldChild ||
+		LastSelectedEntity->GetRCType() == SRCPanelTreeNode::FieldGroup) // Field Child/Group does not contain entity ID, that is why it should not be processed
 	{
 		return;
 	}
@@ -2517,17 +2523,17 @@ bool SRemoteControlPanel::CanRenameEntity() const
 	return false;
 }
 
-void SRemoteControlPanel::SetLogicClipboardItem(UObject* InItem, TSharedPtr<SRCLogicPanelBase> InSourcePanel)
+void SRemoteControlPanel::SetLogicClipboardItems(const TArray<UObject*>& InItems, const TSharedPtr<SRCLogicPanelBase>& InSourcePanel)
 {
-	LogicClipboardItem = InItem;
+	LogicClipboardItems = InItems;
 	LogicClipboardItemSource = InSourcePanel;
 }
 
 void SRemoteControlPanel::CopyItem_Execute()
 {
-	if (TSharedPtr<SRCLogicPanelBase> ActiveLogicPanel = GetActiveLogicPanel())
+	if (const TSharedPtr<SRCLogicPanelBase> ActiveLogicPanel = GetActiveLogicPanel())
 	{
-		ActiveLogicPanel->CopySelectedPanelItem();
+		ActiveLogicPanel->CopySelectedPanelItems();
 	}
 }
 
@@ -2538,9 +2544,9 @@ bool SRemoteControlPanel::CanCopyItem() const
 		return false;
 	}
 
-	if (TSharedPtr<SRCLogicPanelBase> ActiveLogicPanel = GetActiveLogicPanel())
+	if (const TSharedPtr<SRCLogicPanelBase> ActiveLogicPanel = GetActiveLogicPanel())
 	{
-		return ActiveLogicPanel->GetSelectedLogicItem().IsValid();
+		return !ActiveLogicPanel->GetSelectedLogicItems().IsEmpty();
 	}
 
 	return false;
@@ -2548,9 +2554,9 @@ bool SRemoteControlPanel::CanCopyItem() const
 
 void SRemoteControlPanel::PasteItem_Execute()
 {
-	if (TSharedPtr<SRCLogicPanelBase> ActiveLogicPanel = GetActiveLogicPanel())
+	if (const TSharedPtr<SRCLogicPanelBase> ActiveLogicPanel = GetActiveLogicPanel())
 	{
-		ActiveLogicPanel->PasteItemFromClipboard();
+		ActiveLogicPanel->PasteItemsFromClipboard();
 	}
 }
 
@@ -2561,14 +2567,14 @@ bool SRemoteControlPanel::CanPasteItem() const
 		return false;
 	}
 
-	if(LogicClipboardItem)
+	if(!LogicClipboardItems.IsEmpty())
 	{
-		if (TSharedPtr<SRCLogicPanelBase> ActiveLogicPanel = GetActiveLogicPanel())
+		if (const TSharedPtr<SRCLogicPanelBase>& ActiveLogicPanel = GetActiveLogicPanel())
 		{
 			// Currently we only support pasting items between panels of exactly the same type.
 			if (LogicClipboardItemSource == ActiveLogicPanel)
 			{
-				return ActiveLogicPanel->CanPasteClipboardItem(LogicClipboardItem);
+				return ActiveLogicPanel->CanPasteClipboardItems(LogicClipboardItems);
 			}
 		}
 	}
@@ -2578,9 +2584,9 @@ bool SRemoteControlPanel::CanPasteItem() const
 
 void SRemoteControlPanel::DuplicateItem_Execute()
 {
-	if (TSharedPtr<SRCLogicPanelBase> ActiveLogicPanel = GetActiveLogicPanel())
+	if (const TSharedPtr<SRCLogicPanelBase>& ActiveLogicPanel = GetActiveLogicPanel())
 	{
-		ActiveLogicPanel->DuplicateSelectedPanelItem();
+		ActiveLogicPanel->DuplicateSelectedPanelItems();
 
 		return;
 	}
@@ -2593,9 +2599,9 @@ bool SRemoteControlPanel::CanDuplicateItem() const
 		return false;
 	}
 
-	if (TSharedPtr<SRCLogicPanelBase> ActiveLogicPanel = GetActiveLogicPanel())
+	if (const TSharedPtr<SRCLogicPanelBase>& ActiveLogicPanel = GetActiveLogicPanel())
 	{
-		return ActiveLogicPanel->GetSelectedLogicItem().IsValid();
+		return !ActiveLogicPanel->GetSelectedLogicItems().IsEmpty();
 	}
 
 	return false;
@@ -2603,15 +2609,9 @@ bool SRemoteControlPanel::CanDuplicateItem() const
 
 void SRemoteControlPanel::UpdateValue_Execute()
 {
-	if (const TSharedPtr<SRCLogicPanelBase> ActiveLogicPanel = GetActiveLogicPanel())
+	if (const TSharedPtr<SRCLogicPanelBase>& ActiveLogicPanel = GetActiveLogicPanel())
 	{
-		if (const TSharedPtr<FRCActionModel> RCActionLogicItem = StaticCastSharedPtr<FRCActionModel>(ActiveLogicPanel->GetSelectedLogicItem()))
-		{
-			if (const URCPropertyAction* RCPropertyAction = Cast<URCPropertyAction>(RCActionLogicItem->GetAction()))
-			{
-				RCPropertyAction->UpdateValueBasedOnRCProperty();
-			}
-		}
+		ActiveLogicPanel->UpdateValue();
 	}
 }
 
@@ -2624,22 +2624,7 @@ bool SRemoteControlPanel::CanUpdateValue() const
 	
 	if (const TSharedPtr<SRCLogicPanelBase> ActiveLogicPanel = GetActiveLogicPanel())
 	{
-		if (const TSharedPtr<FRCActionModel> RCActionLogicItem = StaticCastSharedPtr<FRCActionModel>(ActiveLogicPanel->GetSelectedLogicItem()))
-		{
-			if (const URemoteControlPreset* RCPreset = GetPreset())
-			{
-				if (const URCAction* RCAction = RCActionLogicItem->GetAction())
-				{
-					if (const TSharedPtr<FRCBehaviourModel> ParentBehaviour = RCActionLogicItem->GetParentBehaviour())
-					{
-						if (const URCBehaviour* Behaviour = ParentBehaviour->GetBehaviour())
-						{
-							return RCAction->IsA<URCPropertyAction>() && !Behaviour->IsA<URCBehaviourBind>() && RCPreset->GetExposedEntity(RCAction->ExposedFieldId).IsValid();
-						}
-					}
-				}
-			}
-		}
+		return ActiveLogicPanel->CanUpdateValue();
 	}
 
 	return false;
diff --git a/Engine/Plugins/VirtualProduction/RemoteControl/Source/RemoteControlUI/Private/UI/SRemoteControlPanel.h b/Engine/Plugins/VirtualProduction/RemoteControl/Source/RemoteControlUI/Private/UI/SRemoteControlPanel.h
--- a/Engine/Plugins/VirtualProduction/RemoteControl/Source/RemoteControlUI/Private/UI/SRemoteControlPanel.h
+++ b/Engine/Plugins/VirtualProduction/RemoteControl/Source/RemoteControlUI/Private/UI/SRemoteControlPanel.h
@@ -54,7 +54,7 @@ DECLARE_DELEGATE_TwoParams(FOnLiveModeChange, TSharedPtr<SRemoteControlPanel> /*
  * UI representation of a remote control preset.
  * Allows a user to expose/unexpose properties and functions from actors and blueprint libraries.
  */
-class SRemoteControlPanel : public SCompoundWidget
+class SRemoteControlPanel : public SCompoundWidget, public FGCObject
 {
 	SLATE_BEGIN_ARGS(SRemoteControlPanel) {}
 		SLATE_EVENT(FOnLiveModeChange, OnLiveModeChange)
@@ -153,14 +153,17 @@ public:
 	}
 
 	/** For Copy UI command - Sets the logic clipboard item and source */
-	void SetLogicClipboardItem(UObject* InItem, TSharedPtr<SRCLogicPanelBase> SourcePanel);
+	void SetLogicClipboardItems(const TArray<UObject*>& InItems, const TSharedPtr<SRCLogicPanelBase>& SourcePanel);
 
 	/** Fetches the last UI item copied to Logic clipboard by the user */
-	UObject* GetLogicClipboardItem()
+	TArray<UObject*> GetLogicClipboardItems()
 	{
-		return LogicClipboardItem;
+		return LogicClipboardItems;
 	}
 
+	// FGCObject interface
+	virtual void AddReferencedObjects(FReferenceCollector& Collector) override;
+
 private:
 
 	/** Returns a Helper widget for entity details view and protocol details view. */
@@ -426,13 +429,15 @@ private:
 	/** Action panel UI widget for Remote Control Logic*/
 	TSharedPtr<class SRCActionPanel> ActionPanel;
 
-	/** LogicClipboardItem - Holds the latest item copied from a Logic panel
+	/** LogicClipboardItems - Holds the items copied from a Logic panel
 	*
 	* Note: We track UObjects (Data Model) here rather than the UI Models as the latter are swept away the moment the user navigates to a different Controller.
 	* For example if the user copies an action from a behaviour in a given Controller but then navigates to another Controller, we can no longer rely on the previous UI objects
 	* as they would have been discarded in favor of a new data set for the actively selected Controller */
-	UPROPERTY(Transient)
-	TObjectPtr<UObject> LogicClipboardItem;
+	TArray<TObjectPtr<UObject>> LogicClipboardItems;
+
+	// FGCObject interface
+	virtual FString GetReferencerName() const override { return "RemoteControlPanel"; }
 
 	/** Keeps track of whether materials were compiled from the current frame. Used to limit the number of UI refresh to once per frame. */
 	bool bMaterialsCompiledThisFrame = false;
