[RemoteControl]
Avoid Add Action button and drag and drop to add properties with warnings.
Fix behaviour conditions to either:
- All allow the property
- All don't allow the property

diff --git a/Engine/Plugins/VirtualProduction/RemoteControl/Source/RemoteControl/Private/RemoteControlField.cpp b/Engine/Plugins/VirtualProduction/RemoteControl/Source/RemoteControl/Private/RemoteControlField.cpp
--- a/Engine/Plugins/VirtualProduction/RemoteControl/Source/RemoteControl/Private/RemoteControlField.cpp
+++ b/Engine/Plugins/VirtualProduction/RemoteControl/Source/RemoteControl/Private/RemoteControlField.cpp
@@ -6,6 +6,7 @@
 #include "GameFramework/Actor.h"
 #include "IRemoteControlModule.h"
 #include "Math/NumericLimits.h"
+#include "Misc/App.h"
 #include "RemoteControlObjectVersion.h"
 #include "RemoteControlFieldPath.h"
 #include "RemoteControlBinding.h"
@@ -334,6 +335,15 @@ bool FRemoteControlProperty::IsEditableInEditor(FString* OutError) const
 	return IRemoteControlModule::Get().PropertySupportsRawModification(GetProperty(), GetBoundObject(), true, OutError);
 }
 
+bool FRemoteControlProperty::IsEditable(FString* OutError) const
+{
+	if (FApp::IsGame())
+	{
+		return IsEditableInPackaged(OutError); 
+	}
+	return IsEditableInEditor(OutError);
+}
+
 bool FRemoteControlProperty::Serialize(FArchive& Ar)
 {
 	Ar.UsingCustomVersion(FRemoteControlObjectVersion::GUID);
diff --git a/Engine/Plugins/VirtualProduction/RemoteControl/Source/RemoteControl/Public/RemoteControlField.h b/Engine/Plugins/VirtualProduction/RemoteControl/Source/RemoteControl/Public/RemoteControlField.h
--- a/Engine/Plugins/VirtualProduction/RemoteControl/Source/RemoteControl/Public/RemoteControlField.h
+++ b/Engine/Plugins/VirtualProduction/RemoteControl/Source/RemoteControl/Public/RemoteControlField.h
@@ -220,6 +220,9 @@ public:
 	/** Returns whether the property is editable in the Editor. */
 	bool IsEditableInEditor(FString* OutError = nullptr) const;
 
+	/** Returns whether the property is editable, will check for editor or packaged automatically */
+	bool IsEditable(FString* OutError = nullptr) const;
+
 	bool Serialize(FArchive& Ar);
 	void PostSerialize(const FArchive& Ar);
 	
diff --git a/Engine/Plugins/VirtualProduction/RemoteControl/Source/RemoteControlLogic/Private/Behaviour/Builtin/Bind/RCBehaviourBind.cpp b/Engine/Plugins/VirtualProduction/RemoteControl/Source/RemoteControlLogic/Private/Behaviour/Builtin/Bind/RCBehaviourBind.cpp
--- a/Engine/Plugins/VirtualProduction/RemoteControl/Source/RemoteControlLogic/Private/Behaviour/Builtin/Bind/RCBehaviourBind.cpp
+++ b/Engine/Plugins/VirtualProduction/RemoteControl/Source/RemoteControlLogic/Private/Behaviour/Builtin/Bind/RCBehaviourBind.cpp
@@ -71,12 +71,21 @@ bool URCBehaviourBind::CanHaveActionForField(const TSharedPtr<FRemoteControlFiel
 		return false; // already exists!
 	}
 
+	// Only property type will be allowed so we check before to avoid a cast
+	if (InRemoteControlField->FieldType != EExposedFieldType::Property)
+	{
+		return false;
+	}
+
 	// Advanced checks (by Controller type and Target type)
-	if (TSharedPtr<FRemoteControlProperty> RemoteControlEntityAsProperty = StaticCastSharedPtr<FRemoteControlProperty>(InRemoteControlField))
+	if (const TSharedPtr<FRemoteControlProperty> RCProperty = StaticCastSharedPtr<FRemoteControlProperty>(InRemoteControlField))
 	{
 		if (URCController* Controller = ControllerWeakPtr.Get())
 		{
-			return URCBehaviourBind::CanHaveActionForField(Controller, RemoteControlEntityAsProperty.ToSharedRef(), bAllowNumericInputAsStrings);
+			if (URCBehaviourBind::CanHaveActionForField(Controller, RCProperty.ToSharedRef(), bAllowNumericInputAsStrings))
+			{
+				return RCProperty->IsEditable();
+			}
 		}
 	}
 
diff --git a/Engine/Plugins/VirtualProduction/RemoteControl/Source/RemoteControlLogic/Private/Behaviour/Builtin/Conditional/RCBehaviourConditional.cpp b/Engine/Plugins/VirtualProduction/RemoteControl/Source/RemoteControlLogic/Private/Behaviour/Builtin/Conditional/RCBehaviourConditional.cpp
--- a/Engine/Plugins/VirtualProduction/RemoteControl/Source/RemoteControlLogic/Private/Behaviour/Builtin/Conditional/RCBehaviourConditional.cpp
+++ b/Engine/Plugins/VirtualProduction/RemoteControl/Source/RemoteControlLogic/Private/Behaviour/Builtin/Conditional/RCBehaviourConditional.cpp
@@ -185,6 +185,19 @@ URCAction* URCBehaviourConditional::AddConditionalAction(const TSharedRef<const
 	return ActionContainer->AddAction(UniquenessTest, InRemoteControlField);
 }
 
+bool URCBehaviourConditional::CanHaveActionForField(const TSharedPtr<FRemoteControlField> InRemoteControlField) const
+{
+	if (InRemoteControlField->FieldType == EExposedFieldType::Property)
+	{
+		if (const TSharedPtr<FRemoteControlProperty>& RCProperty = StaticCastSharedPtr<FRemoteControlProperty>(InRemoteControlField))
+		{
+			return RCProperty->IsEditable();
+		}
+	}
+
+	return true;
+}
+
 FText URCBehaviourConditional::GetConditionTypeAsText(ERCBehaviourConditionType ConditionType) const
 {
 	FText ConditionDisplayText;
@@ -228,4 +241,4 @@ FText URCBehaviourConditional::GetConditionTypeAsText(ERCBehaviourConditionType
 	}
 
 	return ConditionDisplayText;
-}
\ No newline at end of file
+}
diff --git a/Engine/Plugins/VirtualProduction/RemoteControl/Source/RemoteControlLogic/Private/Behaviour/Builtin/RangeMap/RCRangeMapBehaviour.cpp b/Engine/Plugins/VirtualProduction/RemoteControl/Source/RemoteControlLogic/Private/Behaviour/Builtin/RangeMap/RCRangeMapBehaviour.cpp
--- a/Engine/Plugins/VirtualProduction/RemoteControl/Source/RemoteControlLogic/Private/Behaviour/Builtin/RangeMap/RCRangeMapBehaviour.cpp
+++ b/Engine/Plugins/VirtualProduction/RemoteControl/Source/RemoteControlLogic/Private/Behaviour/Builtin/RangeMap/RCRangeMapBehaviour.cpp
@@ -484,8 +484,15 @@ TMap<double, URCAction*> URCRangeMapBehaviour::GetNonLerpActions()
 
 bool URCRangeMapBehaviour::CanHaveActionForField(const TSharedPtr<FRemoteControlField> InRemoteControlField) const
 {
-	// We can select all of the actions, and the check whether its actually added is done elsewhere.
-	return true;
+	if (InRemoteControlField->FieldType == EExposedFieldType::Property)
+	{
+		if (const TSharedPtr<FRemoteControlProperty>& RCProperty = StaticCastSharedPtr<FRemoteControlProperty>(InRemoteControlField))
+		{
+			return RCProperty->IsEditable();
+		}
+	}
+
+	return false;
 }
 
 void URCRangeMapBehaviour::ApplyLerpOnStruct(const FRCRangeMapInput* MinRangeInput, const FRCRangeMapInput* MaxRangeInput, const double& InputAlpha, const FGuid& FieldId)
diff --git a/Engine/Plugins/VirtualProduction/RemoteControl/Source/RemoteControlLogic/Private/Behaviour/RCBehaviour.cpp b/Engine/Plugins/VirtualProduction/RemoteControl/Source/RemoteControlLogic/Private/Behaviour/RCBehaviour.cpp
--- a/Engine/Plugins/VirtualProduction/RemoteControl/Source/RemoteControlLogic/Private/Behaviour/RCBehaviour.cpp
+++ b/Engine/Plugins/VirtualProduction/RemoteControl/Source/RemoteControlLogic/Private/Behaviour/RCBehaviour.cpp
@@ -82,6 +82,14 @@ bool URCBehaviour::CanHaveActionForField(const TSharedPtr<FRemoteControlField> I
 		return false; // already exists!
 	}
 
+	if (InRemoteControlField->FieldType == EExposedFieldType::Property)
+	{
+		if (const TSharedPtr<FRemoteControlProperty>& RCProperty = StaticCastSharedPtr<FRemoteControlProperty>(InRemoteControlField))
+		{
+			return RCProperty->IsEditable();
+		}
+	}
+
 	return true;
 }
 
diff --git a/Engine/Plugins/VirtualProduction/RemoteControl/Source/RemoteControlLogic/Public/Behaviour/Builtin/Conditional/RCBehaviourConditional.h b/Engine/Plugins/VirtualProduction/RemoteControl/Source/RemoteControlLogic/Public/Behaviour/Builtin/Conditional/RCBehaviourConditional.h
--- a/Engine/Plugins/VirtualProduction/RemoteControl/Source/RemoteControlLogic/Public/Behaviour/Builtin/Conditional/RCBehaviourConditional.h
+++ b/Engine/Plugins/VirtualProduction/RemoteControl/Source/RemoteControlLogic/Public/Behaviour/Builtin/Conditional/RCBehaviourConditional.h
@@ -82,10 +82,7 @@ public:
 	URCAction* AddConditionalAction(const TSharedRef<const FRemoteControlField> InRemoteControlField, const ERCBehaviourConditionType InConditionType, const TObjectPtr<URCVirtualPropertySelfContainer> InComparand);
 
 	/** Whether we can create an action pertaining to a given remote control field for the current behaviour */
-	virtual bool CanHaveActionForField(const TSharedPtr<FRemoteControlField> InRemoteControlField) const override
-	{
-		return true; // Conditional Behaviour can support multiple Action rows for a single Remote Control Field
-	}
+	virtual bool CanHaveActionForField(const TSharedPtr<FRemoteControlField> InRemoteControlField) const override;
 
 	/** User-friendly text representation of a condition enum. Used to display comparator info in the Actions table */
 	FText GetConditionTypeAsText(ERCBehaviourConditionType ConditionType) const;
diff --git a/Engine/Plugins/VirtualProduction/RemoteControl/Source/RemoteControlUI/Private/UI/Controller/SRCControllerPanelList.cpp b/Engine/Plugins/VirtualProduction/RemoteControl/Source/RemoteControlUI/Private/UI/Controller/SRCControllerPanelList.cpp
--- a/Engine/Plugins/VirtualProduction/RemoteControl/Source/RemoteControlUI/Private/UI/Controller/SRCControllerPanelList.cpp
+++ b/Engine/Plugins/VirtualProduction/RemoteControl/Source/RemoteControlUI/Private/UI/Controller/SRCControllerPanelList.cpp
@@ -823,6 +823,12 @@ bool SRCControllerPanelList::IsEntitySupported(const FGuid ExposedEntityId)
 	{
 		if (const TSharedPtr<const FRemoteControlProperty>& RemoteControlProperty = Preset->GetExposedEntity<FRemoteControlProperty>(ExposedEntityId).Pin())
 		{
+			if (!RemoteControlProperty->IsEditable())
+			{
+				// Property with error(s)
+				return false;
+			}
+
 			if (RemoteControlProperty->FieldType == EExposedFieldType::Property)
 			{
 				const FProperty* Property = RemoteControlProperty->GetProperty();
