[Remote Control] Adding Renew of Entity and Controller Guids when duplicating an embedded RemoteControlPreset.

Requirement:
The requirement of this comes from the "combo template" feature of Avalanche. As all the controller and entity values of different sub-level's embedded presets are combined in the same page, guid collisions must be avoided.

The combo template creation will fail if it detects guid collisions between the guids.

Feature:

When the parent asset (either level or bp) is duplicated, the embedded preset will renew the entity and controller's guids automatically.

If users want to create combo templates of existing assets that have guid collisions, the fix is to duplicate them.

This has been restricted to embedded presets only.

Implementation:

Entity Ids are used as index in several classes of the overall data structure (identified by inspection).

The UpdateEntityIds function has been added to propagate the re-indexing throughout the rc data structures.

New delegates have been added if propagation beyond the rc data structures are required.

#rnx

diff --git a/Engine/Plugins/VirtualProduction/RemoteControl/Source/RemoteControl/Private/RCVirtualPropertyContainer.cpp b/Engine/Plugins/VirtualProduction/RemoteControl/Source/RemoteControl/Private/RCVirtualPropertyContainer.cpp
--- a/Engine/Plugins/VirtualProduction/RemoteControl/Source/RemoteControl/Private/RCVirtualPropertyContainer.cpp
+++ b/Engine/Plugins/VirtualProduction/RemoteControl/Source/RemoteControl/Private/RCVirtualPropertyContainer.cpp
@@ -291,6 +291,17 @@ FName URCVirtualPropertyContainerBase::GenerateUniquePropertyName(const FName& I
 	return *FinalName;
 }
 
+void URCVirtualPropertyContainerBase::UpdateEntityIds(const TMap<FGuid, FGuid>& InEntityIdMap)
+{
+	for (const TObjectPtr<URCVirtualPropertyBase>& VirtualProperty : VirtualProperties)
+	{
+		if (VirtualProperty)
+		{
+			VirtualProperty->UpdateEntityIds(InEntityIdMap);
+		}
+	}
+}
+
 #if WITH_EDITOR
 void URCVirtualPropertyContainerBase::PostEditUndo()
 {
diff --git a/Engine/Plugins/VirtualProduction/RemoteControl/Source/RemoteControl/Private/RemoteControlExposeRegistry.cpp b/Engine/Plugins/VirtualProduction/RemoteControl/Source/RemoteControl/Private/RemoteControlExposeRegistry.cpp
--- a/Engine/Plugins/VirtualProduction/RemoteControl/Source/RemoteControl/Private/RemoteControlExposeRegistry.cpp
+++ b/Engine/Plugins/VirtualProduction/RemoteControl/Source/RemoteControl/Private/RemoteControlExposeRegistry.cpp
@@ -168,6 +168,21 @@ void URemoteControlExposeRegistry::PostDuplicate(bool bDuplicateForPIE)
 	CacheLabels();
 }
 
+void URemoteControlExposeRegistry::Rehash()
+{
+	TSet<FRCEntityWrapper> RehashedEntities;
+	RehashedEntities.Reserve(ExposedEntities.Num());
+	
+	for (FRCEntityWrapper& Wrapper : ExposedEntities)
+	{
+		RehashedEntities.Add(FRCEntityWrapper(Wrapper));
+	}
+	
+	ExposedEntities = MoveTemp(RehashedEntities);
+
+	CacheLabels();
+}
+
 TSharedPtr<FRemoteControlEntity> URemoteControlExposeRegistry::GetEntity(const FGuid& EntityId)
 {
 	/** Get a raw pointer to an entity using its id. */
diff --git a/Engine/Plugins/VirtualProduction/RemoteControl/Source/RemoteControl/Private/RemoteControlExposeRegistry.h b/Engine/Plugins/VirtualProduction/RemoteControl/Source/RemoteControl/Private/RemoteControlExposeRegistry.h
--- a/Engine/Plugins/VirtualProduction/RemoteControl/Source/RemoteControl/Private/RemoteControlExposeRegistry.h
+++ b/Engine/Plugins/VirtualProduction/RemoteControl/Source/RemoteControl/Private/RemoteControlExposeRegistry.h
@@ -199,6 +199,11 @@ public:
 
 	virtual void PostDuplicate(bool bDuplicateForPIE) override;
 
+	/**
+	 * @brief Rehash the entities after ids have been modified.
+	 */
+	void Rehash();
+	
 private:
 	/** Get a raw pointer to an entity using its id. */
 	TSharedPtr<FRemoteControlEntity> GetEntity(const FGuid& EntityId);
diff --git a/Engine/Plugins/VirtualProduction/RemoteControl/Source/RemoteControl/Private/RemoteControlPreset.cpp b/Engine/Plugins/VirtualProduction/RemoteControl/Source/RemoteControl/Private/RemoteControlPreset.cpp
--- a/Engine/Plugins/VirtualProduction/RemoteControl/Source/RemoteControl/Private/RemoteControlPreset.cpp
+++ b/Engine/Plugins/VirtualProduction/RemoteControl/Source/RemoteControl/Private/RemoteControlPreset.cpp
@@ -141,6 +141,31 @@ namespace
 
 		return CurrentConfiguration;
 	}
+
+	template<typename InValueType>
+	void RehashMap(TMap<FGuid, InValueType>& InOutMapToRehash, const TMap<FGuid, FGuid>& InGuidMap)
+	{
+		TMap<FGuid, InValueType> RehashedMap;
+		RehashedMap.Reserve(InOutMapToRehash.Num());
+	
+		for (TPair<FGuid, InValueType>& Entry : InOutMapToRehash)
+		{
+			RehashedMap.Add(InGuidMap.Contains(Entry.Key) ? InGuidMap[Entry.Key] : Entry.Key, MoveTemp(Entry.Value));
+		}
+
+		InOutMapToRehash = MoveTemp(RehashedMap);
+	}
+
+	void RehashSet(TSet<FGuid>& InOutSetToRehash, const TMap<FGuid, FGuid>& InGuidMap)
+	{
+		TSet<FGuid> RehashedSet;
+		RehashedSet.Reserve(InOutSetToRehash.Num());
+		for (const FGuid& Key : InOutSetToRehash)
+		{
+			RehashedSet.Add(InGuidMap.Contains(Key) ? InGuidMap[Key] : Key);
+		}
+		InOutSetToRehash = MoveTemp(RehashedSet);
+	}
 }
 
 FRemoteControlPresetExposeArgs::FRemoteControlPresetExposeArgs()
@@ -603,6 +628,20 @@ URemoteControlPreset* FRemoteControlPresetLayout::GetOwner()
 	return Owner.Get();
 }
 
+void FRemoteControlPresetLayout::UpdateEntityIds(const TMap<FGuid, FGuid>& InEntityIdMap)
+{
+	for (FRemoteControlPresetGroup& Group : Groups)
+	{
+		for (FGuid& FieldId : Group.AccessFields())
+		{
+			if (InEntityIdMap.Contains(FieldId))
+			{
+				FieldId = InEntityIdMap[FieldId];
+			}
+		}
+	}
+}
+
 FRemoteControlPresetGroup& FRemoteControlPresetLayout::CreateGroupInternal(FName GroupName, FGuid GroupId)
 {
 	PRAGMA_DISABLE_DEPRECATION_WARNINGS
@@ -672,6 +711,12 @@ void URemoteControlPreset::PostDuplicate(bool bDuplicateForPIE)
 	if (!bDuplicateForPIE)
 	{
 		PresetId = FGuid::NewGuid();
+
+		if (IsEmbeddedPreset())
+		{
+			RenewEntityIds();
+			RenewControllerIds();
+		}
 	}
 }
 
@@ -1718,6 +1763,84 @@ void URemoteControlPreset::RebindAllEntitiesUnderSameActor(const FGuid& EntityId
 	}
 }
 
+void URemoteControlPreset::RenewEntityIds()
+{
+	Modify();
+
+	TArray<TSharedPtr<FRemoteControlEntity>> ExposedEntities = Registry->GetExposedEntities();
+
+	// Keep track of the guids for delegate broadcast.
+	TMap<FGuid, FGuid> EntityIdMap;	// old -> new
+	EntityIdMap.Reserve(ExposedEntities.Num());
+	
+	for (const TSharedPtr<FRemoteControlEntity>& Entity : ExposedEntities)
+	{
+		if (Entity)
+		{
+			const FGuid OldEntityId = Entity->Id;
+			Entity->Id = FGuid::NewGuid();
+			EntityIdMap.Add(OldEntityId, Entity->Id);
+		}
+	}
+
+	// Rehash the registries
+	Registry->Rehash();
+	PropertyIdRegistry->UpdateEntityIds(EntityIdMap);
+	Layout.UpdateEntityIds(EntityIdMap);
+	
+	// Update NameToGuidMap.
+	for (TPair<FName, FGuid>& NameToGuid : NameToGuidMap)
+	{
+		if (const FGuid* FoundNewId = EntityIdMap.Find(NameToGuid.Value))
+		{
+			NameToGuid.Value = *FoundNewId;
+		}
+	}
+
+	RehashMap(FieldCache, EntityIdMap);
+	RehashMap(PropertyWatchers, EntityIdMap);
+	RehashMap(PreObjectsModifiedCache, EntityIdMap);
+	RehashMap(PreObjectsModifiedActorCache, EntityIdMap);
+	RehashMap(PreMaterialModifiedCache, EntityIdMap);
+	RehashSet(PerFrameModifiedProperties, EntityIdMap);
+	RehashSet(PerFrameUpdatedEntities, EntityIdMap);
+
+	if (ControllerContainer)
+	{
+		ControllerContainer->UpdateEntityIds(EntityIdMap);
+	}
+
+	// Fire the delegates once the registry has been rehashed to ensure proper lookup.
+	OnPropertyIdsRenewed().Broadcast(this, EntityIdMap);
+}
+
+void URemoteControlPreset::RenewControllerIds()
+{
+	if (!ControllerContainer)
+	{
+		return;
+	}
+	
+#if WITH_EDITOR
+	ControllerContainer->Modify();
+#endif
+
+	// Keep track of the guids for delegate broadcast.
+	TMap<FGuid, FGuid> ControllerIdMap;	// old -> new
+	ControllerIdMap.Reserve(ControllerContainer->VirtualProperties.Num());
+	
+	for (const TObjectPtr<URCVirtualPropertyBase>& Controller : ControllerContainer->VirtualProperties)
+	{
+		const FGuid OldControllerId = Controller->Id;
+		Controller->Id = FGuid::NewGuid();
+		ControllerIdMap.Add(OldControllerId, Controller->Id);
+	}
+
+	// It doesn't seem like controller Ids are used as hashing key.
+	
+	OnControllerIdsRenewed().Broadcast(this, ControllerIdMap);
+}
+
 void URemoteControlPreset::NotifyExposedPropertyChanged(FName PropertyLabel)
 {
 	if (TSharedPtr<FRemoteControlProperty> ExposedProperty = GetExposedEntity<FRemoteControlProperty>(GetExposedEntityId(PropertyLabel)).Pin())
diff --git a/Engine/Plugins/VirtualProduction/RemoteControl/Source/RemoteControl/Private/RemoteControlPropertyIdRegistry.cpp b/Engine/Plugins/VirtualProduction/RemoteControl/Source/RemoteControl/Private/RemoteControlPropertyIdRegistry.cpp
--- a/Engine/Plugins/VirtualProduction/RemoteControl/Source/RemoteControl/Private/RemoteControlPropertyIdRegistry.cpp
+++ b/Engine/Plugins/VirtualProduction/RemoteControl/Source/RemoteControl/Private/RemoteControlPropertyIdRegistry.cpp
@@ -620,6 +620,32 @@ TSet<FGuid> URemoteControlPropertyIdRegistry::GetEntityIdsList()
 	return OutIds;
 }
 
+void URemoteControlPropertyIdRegistry::UpdateEntityIds(const TMap<FGuid, FGuid>& InEntityIdMap)
+{
+	bool bNeedsRehash = false;
+	for (FRCPropertyIdWrapper& PropertyId : IdentifiedFields)
+	{
+		if (const FGuid* FoundNewId = InEntityIdMap.Find(PropertyId.GetEntityId()))
+		{
+			PropertyId.EntityId = *FoundNewId;
+			bNeedsRehash = true;
+		}
+	}
+
+	if (bNeedsRehash)
+	{
+		TSet<FRCPropertyIdWrapper> RehashedIdentifiedFields;
+		RehashedIdentifiedFields.Reserve(IdentifiedFields.Num());
+	
+		for (FRCPropertyIdWrapper& Wrapper : IdentifiedFields)
+		{
+			RehashedIdentifiedFields.Add(FRCPropertyIdWrapper(Wrapper));
+		}
+	
+		IdentifiedFields = MoveTemp(RehashedIdentifiedFields);
+	}
+}
+
 bool URemoteControlPropertyIdRegistry::CopyNonUClassOwnerProperty(const FRemoteControlPropertyIdArgs& InArgs, TSharedPtr<FRemoteControlProperty> TargetRCProperty, FProperty* InProperty, UObject* BoundObject)
 {
 	bool bCheckByteEnumComparison = false;
diff --git a/Engine/Plugins/VirtualProduction/RemoteControl/Source/RemoteControl/Public/RCVirtualProperty.h b/Engine/Plugins/VirtualProduction/RemoteControl/Source/RemoteControl/Public/RCVirtualProperty.h
--- a/Engine/Plugins/VirtualProduction/RemoteControl/Source/RemoteControl/Public/RCVirtualProperty.h
+++ b/Engine/Plugins/VirtualProduction/RemoteControl/Source/RemoteControl/Public/RCVirtualProperty.h
@@ -54,6 +54,12 @@ public:
 	/** Initialization routine. Called after the parent container has setup data for this property */
 	virtual void Init() {}
 
+	/**
+	 * @brief Called internally when entity Ids are renewed.
+	 * @param InEntityIdMap Map of old Id to new Id.
+	 */
+	virtual void UpdateEntityIds(const TMap<FGuid, FGuid>& InEntityIdMap) {}
+
 	/** Returns const FProperty for this RC virtual property */
 	virtual const FProperty* GetProperty() const;
 
diff --git a/Engine/Plugins/VirtualProduction/RemoteControl/Source/RemoteControl/Public/RCVirtualPropertyContainer.h b/Engine/Plugins/VirtualProduction/RemoteControl/Source/RemoteControl/Public/RCVirtualPropertyContainer.h
--- a/Engine/Plugins/VirtualProduction/RemoteControl/Source/RemoteControl/Public/RCVirtualPropertyContainer.h
+++ b/Engine/Plugins/VirtualProduction/RemoteControl/Source/RemoteControl/Public/RCVirtualPropertyContainer.h
@@ -132,6 +132,12 @@ public:
 	*/
 	static FName GenerateUniquePropertyName(const FName& InPropertyName, const URCVirtualPropertyContainerBase* InContainer);
 
+	/**
+	 * @brief Called internally when entity Ids are renewed.
+	 * @param InEntityIdMap Map of old Id to new Id.
+	 */
+	virtual void UpdateEntityIds(const TMap<FGuid, FGuid>& InEntityIdMap);
+
 #if WITH_EDITOR
 	/** Called after applying a transaction to the object. Used to broadcast Undo related container changes to UI. */
 	virtual void PostEditUndo();
diff --git a/Engine/Plugins/VirtualProduction/RemoteControl/Source/RemoteControl/Public/RemoteControlPreset.h b/Engine/Plugins/VirtualProduction/RemoteControl/Source/RemoteControl/Public/RemoteControlPreset.h
--- a/Engine/Plugins/VirtualProduction/RemoteControl/Source/RemoteControl/Public/RemoteControlPreset.h
+++ b/Engine/Plugins/VirtualProduction/RemoteControl/Source/RemoteControl/Public/RemoteControlPreset.h
@@ -286,6 +286,12 @@ struct REMOTECONTROL_API FRemoteControlPresetLayout
 	/** Get the preset that owns this layout. */
 	URemoteControlPreset* GetOwner();
 
+	/**
+	 * @brief Called internally when entity Ids are renewed.
+	 * @param InEntityIdMap Map of old Id to new Id.
+	 */
+	void UpdateEntityIds(const TMap<FGuid, FGuid>& InEntityIdMap);
+	
 	// Layout operation delegates
 	DECLARE_MULTICAST_DELEGATE_OneParam(FOnGroupAdded, const FRemoteControlPresetGroup& /*NewGroup*/);
 	FOnGroupAdded& OnGroupAdded() { return OnGroupAddedDelegate; }
@@ -585,6 +591,16 @@ public:
 	 */
 	void RebindAllEntitiesUnderSameActor(const FGuid& EntityId, AActor* NewActor, bool bUseRebindingContext = true);
 
+	/**
+	 * Renews all exposed entity guids. Necessary when duplicating.
+	 */
+	void RenewEntityIds();
+
+	/**
+	 * Renews all controller guids. Necessary when duplicating.
+	 */
+	void RenewControllerIds();
+	
 	DECLARE_MULTICAST_DELEGATE_TwoParams(FOnPresetEntityEvent, URemoteControlPreset* /*Preset*/, const FGuid& /*EntityId*/);
 	FOnPresetEntityEvent& OnEntityExposed() { return OnEntityExposedDelegate; }
 	FOnPresetEntityEvent& OnEntityUnexposed() { return OnEntityUnexposedDelegate; }
@@ -611,7 +627,12 @@ public:
 	DECLARE_MULTICAST_DELEGATE_TwoParams(FOnPresetPropertyUnexposed, URemoteControlPreset* /*Preset*/, FName /*ExposedLabel*/);
 	UE_DEPRECATED(4.27, "This delegate is deprecated, use OnEntityUnexposed instead.")
 	FOnPresetPropertyUnexposed& OnPropertyUnexposed() { return OnPropertyUnexposedDelegate; }
-
+	
+	using FGuidToGuidMap = TMap<FGuid, FGuid>;	// Workaround for delegate macro.
+	
+	DECLARE_MULTICAST_DELEGATE_TwoParams(FOnPropertyIdsRenewed, URemoteControlPreset* /*Preset*/, const FGuidToGuidMap& /*OldToNewEntityIds*/);
+	FOnPropertyIdsRenewed& OnPropertyIdsRenewed() { return OnPropertyIdsRenewedDelegate;}
+	
 	DECLARE_MULTICAST_DELEGATE_ThreeParams(FOnPresetFieldRenamed, URemoteControlPreset* /*Preset*/, FName /*OldExposedLabel*/, FName /**NewExposedLabel*/);
 	FOnPresetFieldRenamed& OnFieldRenamed() { return OnPresetFieldRenamed; }
 
@@ -632,6 +653,9 @@ public:
 
 	DECLARE_MULTICAST_DELEGATE_ThreeParams(FOnControllerRenamed, URemoteControlPreset* /*Preset*/, const FName /*OldLabel*/, const FName /*NewLabel*/);
 	FOnControllerRenamed& OnControllerRenamed() { return OnControllerRenamedDelegate; }
+
+	DECLARE_MULTICAST_DELEGATE_TwoParams(FOnControllerIdsRenewed, URemoteControlPreset* /*Preset*/, const FGuidToGuidMap& /*OldToNewControllerIds*/);
+	FOnControllerIdsRenewed& OnControllerIdsRenewed() { return OnControllerIdsRenewedDelegate; }
 	
 	DECLARE_MULTICAST_DELEGATE_TwoParams(FOnControllerModified, URemoteControlPreset* /*Preset*/, const TSet<FGuid>& /*ModifiedControllerIds*/);
 	FOnControllerModified& OnControllerModified() { return OnControllerModifiedDelegate; }
@@ -895,6 +919,8 @@ private:
 	FOnPresetPropertyExposed OnPropertyExposedDelegate;
 	/** Delegate triggered when a property has been unexposed. */
 	FOnPresetPropertyUnexposed OnPropertyUnexposedDelegate;
+	/** Delegate triggered when the property guids are renewed. */
+	FOnPropertyIdsRenewed OnPropertyIdsRenewedDelegate;
 	/** Delegate triggered when a field has been renamed. */
 	FOnPresetFieldRenamed OnPresetFieldRenamed;
 	/** Delegate triggered when the preset's metadata has been modified. */
@@ -909,6 +935,8 @@ private:
 	FOnControllerRemoved OnControllerRemovedDelegate;
 	/** Delegate triggered when a Controller is renamed */
 	FOnControllerRenamed OnControllerRenamedDelegate;
+	/** Delegate triggered when a Controller guids are renewed. */
+	FOnControllerIdsRenewed OnControllerIdsRenewedDelegate;
 	/** Delegate triggered when a Controller has been changed */
 	FOnControllerModified OnControllerModifiedDelegate;
 
diff --git a/Engine/Plugins/VirtualProduction/RemoteControl/Source/RemoteControl/Public/RemoteControlPropertyIdRegistry.h b/Engine/Plugins/VirtualProduction/RemoteControl/Source/RemoteControl/Public/RemoteControlPropertyIdRegistry.h
--- a/Engine/Plugins/VirtualProduction/RemoteControl/Source/RemoteControl/Public/RemoteControlPropertyIdRegistry.h
+++ b/Engine/Plugins/VirtualProduction/RemoteControl/Source/RemoteControl/Public/RemoteControlPropertyIdRegistry.h
@@ -101,6 +101,8 @@ private:
 	/** Holds the type to be created. */
 	UPROPERTY()
 	UClass* ClassToCreate = nullptr;
+
+	friend class URemoteControlPropertyIdRegistry;
 };
 
 template<> struct TStructOpsTypeTraits<FRCPropertyIdWrapper> : public TStructOpsTypeTraitsBase2<FRCPropertyIdWrapper>
@@ -193,6 +195,12 @@ public:
 	FOnPropertyIdActionNeedsRefresh& OnPropertyIdActionNeedsRefresh() { return OnPropertyIdActionNeedsRefreshDelegate; };
 #endif
 
+	/**
+	 * @brief Called internally when entity Ids are renewed.
+	 * @param InEntityIdMap Map of old Id to new Id.
+	 */
+	void UpdateEntityIds(const TMap<FGuid, FGuid>& InEntityIdMap);
+
 private:
 	/** Call the recursive function TryCopyNonUClassOwnerProperty for the last try on copying the value to the property
 	 *	@param InArgs Argument used to expose a property.
diff --git a/Engine/Plugins/VirtualProduction/RemoteControl/Source/RemoteControlLogic/Private/Action/RCActionContainer.cpp b/Engine/Plugins/VirtualProduction/RemoteControl/Source/RemoteControlLogic/Private/Action/RCActionContainer.cpp
--- a/Engine/Plugins/VirtualProduction/RemoteControl/Source/RemoteControlLogic/Private/Action/RCActionContainer.cpp
+++ b/Engine/Plugins/VirtualProduction/RemoteControl/Source/RemoteControlLogic/Private/Action/RCActionContainer.cpp
@@ -99,6 +99,28 @@ URCBehaviour* URCActionContainer::GetParentBehaviour()
 	return Cast<URCBehaviour>(GetOuter());
 }
 
+void URCActionContainer::ForEachAction(TFunctionRef<void(URCAction*)> InActionFunction, bool bInRecursive)
+{
+	for (URCAction* Action : GetActions())
+	{
+		if (Action)
+		{
+			InActionFunction(Action);
+		}
+	}
+
+	if (bInRecursive)
+	{
+		for (URCActionContainer* ChildActionContainer : ActionContainers)
+		{
+			if (ChildActionContainer)
+			{
+				ChildActionContainer->ForEachAction(InActionFunction, bInRecursive);
+			}
+		}
+	}
+}
+
 TArray<const URCPropertyAction*> URCActionContainer::GetPropertyActions() const
 {
 	TArray<const URCPropertyAction*> PropertyActions;
diff --git a/Engine/Plugins/VirtualProduction/RemoteControl/Source/RemoteControlLogic/Private/Action/RCPropertyAction.cpp b/Engine/Plugins/VirtualProduction/RemoteControl/Source/RemoteControlLogic/Private/Action/RCPropertyAction.cpp
--- a/Engine/Plugins/VirtualProduction/RemoteControl/Source/RemoteControlLogic/Private/Action/RCPropertyAction.cpp
+++ b/Engine/Plugins/VirtualProduction/RemoteControl/Source/RemoteControlLogic/Private/Action/RCPropertyAction.cpp
@@ -31,6 +31,14 @@ URCBehaviour* URCAction::GetParentBehaviour() const
 	return nullptr;
 }
 
+void URCAction::UpdateEntityIds(const TMap<FGuid, FGuid>& InEntityIdMap)
+{
+	if (const FGuid* FoundId = InEntityIdMap.Find(ExposedFieldId))
+	{
+		ExposedFieldId = *FoundId;
+	}
+}
+
 void URCPropertyAction::Execute() const
 {
 	if (!PresetWeakPtr.IsValid())
@@ -73,6 +81,16 @@ void URCPropertyAction::Execute() const
 	Super::Execute();
 }
 
+void URCPropertyAction::UpdateEntityIds(const TMap<FGuid, FGuid>& InEntityIdMap)
+{
+	if (PropertySelfContainer)
+	{
+		PropertySelfContainer->UpdateEntityIds(InEntityIdMap);
+	}
+	
+	Super::UpdateEntityIds(InEntityIdMap);	
+}
+
 FProperty* URCPropertyAction::GetProperty() const
 {
 	if (ensure(PropertySelfContainer))
diff --git a/Engine/Plugins/VirtualProduction/RemoteControl/Source/RemoteControlLogic/Private/Action/RCPropertyIdAction.cpp b/Engine/Plugins/VirtualProduction/RemoteControl/Source/RemoteControlLogic/Private/Action/RCPropertyIdAction.cpp
--- a/Engine/Plugins/VirtualProduction/RemoteControl/Source/RemoteControlLogic/Private/Action/RCPropertyIdAction.cpp
+++ b/Engine/Plugins/VirtualProduction/RemoteControl/Source/RemoteControlLogic/Private/Action/RCPropertyIdAction.cpp
@@ -90,6 +90,26 @@ void URCPropertyIdAction::Execute() const
 	Super::Execute();
 }
 
+void URCPropertyIdAction::UpdateEntityIds(const TMap<FGuid, FGuid>& InEntityIdMap)
+{
+	for (const TPair<FName, TObjectPtr<URCVirtualPropertySelfContainer>>& PropertyContainerEntry : PropertySelfContainer)
+	{
+		if (PropertyContainerEntry.Value)
+		{
+			PropertyContainerEntry.Value->UpdateEntityIds(InEntityIdMap);
+		}
+	}
+	for (const TPair<FName, TObjectPtr<URCVirtualPropertySelfContainer>>& PropertyContainerEntry : CachedPropertySelfContainer)
+	{
+		if (PropertyContainerEntry.Value)
+		{
+			PropertyContainerEntry.Value->UpdateEntityIds(InEntityIdMap);
+		}
+	}
+	
+	Super::UpdateEntityIds(InEntityIdMap);
+}
+
 void URCPropertyIdAction::PostLoad()
 {
 	UObject::PostLoad();
diff --git a/Engine/Plugins/VirtualProduction/RemoteControl/Source/RemoteControlLogic/Private/Behaviour/Builtin/Path/RCSetAssetByPathBehaviour.cpp b/Engine/Plugins/VirtualProduction/RemoteControl/Source/RemoteControlLogic/Private/Behaviour/Builtin/Path/RCSetAssetByPathBehaviour.cpp
--- a/Engine/Plugins/VirtualProduction/RemoteControl/Source/RemoteControlLogic/Private/Behaviour/Builtin/Path/RCSetAssetByPathBehaviour.cpp
+++ b/Engine/Plugins/VirtualProduction/RemoteControl/Source/RemoteControlLogic/Private/Behaviour/Builtin/Path/RCSetAssetByPathBehaviour.cpp
@@ -40,6 +40,16 @@ void URCSetAssetByPathBehaviour::Initialize()
 	Super::Initialize();
 }
 
+void URCSetAssetByPathBehaviour::UpdateEntityIds(const TMap<FGuid, FGuid>& InEntityIdMap)
+{
+	if (const FGuid* FoundId = InEntityIdMap.Find(TargetEntityId))
+	{
+		TargetEntityId = *FoundId;
+	}
+
+	Super::UpdateEntityIds(InEntityIdMap);
+}
+
 bool URCSetAssetByPathBehaviour::SetAssetByPath(const FString& AssetPath, const FString& DefaultString)
 {
 	const URCController* Controller = ControllerWeakPtr.Get();
diff --git a/Engine/Plugins/VirtualProduction/RemoteControl/Source/RemoteControlLogic/Private/Behaviour/RCBehaviour.cpp b/Engine/Plugins/VirtualProduction/RemoteControl/Source/RemoteControlLogic/Private/Behaviour/RCBehaviour.cpp
--- a/Engine/Plugins/VirtualProduction/RemoteControl/Source/RemoteControlLogic/Private/Behaviour/RCBehaviour.cpp
+++ b/Engine/Plugins/VirtualProduction/RemoteControl/Source/RemoteControlLogic/Private/Behaviour/RCBehaviour.cpp
@@ -163,4 +163,15 @@ const FText& URCBehaviour::GetBehaviorDescription()
 	return GetDefault<URCBehaviourNode>(BehaviourNodeClass)->BehaviorDescription;
 }
 
+void URCBehaviour::UpdateEntityIds(const TMap<FGuid, FGuid>& InEntityIdMap)
+{
+	if (ActionContainer)
+	{
+		ActionContainer->ForEachAction([&InEntityIdMap](URCAction* InAction)
+		{
+			InAction->UpdateEntityIds(InEntityIdMap);
+		}, /*bInRecursive*/ true);
+	}
+}
+
 #endif
\ No newline at end of file
diff --git a/Engine/Plugins/VirtualProduction/RemoteControl/Source/RemoteControlLogic/Private/Controller/RCController.cpp b/Engine/Plugins/VirtualProduction/RemoteControl/Source/RemoteControlLogic/Private/Controller/RCController.cpp
--- a/Engine/Plugins/VirtualProduction/RemoteControl/Source/RemoteControlLogic/Private/Controller/RCController.cpp
+++ b/Engine/Plugins/VirtualProduction/RemoteControl/Source/RemoteControlLogic/Private/Controller/RCController.cpp
@@ -10,6 +10,19 @@
 
 #define LOCTEXT_NAMESPACE "RCController"
 
+void URCController::UpdateEntityIds(const TMap<FGuid, FGuid>& InEntityIdMap)
+{
+	for (URCBehaviour* Behaviour : Behaviours)
+	{
+		if (Behaviour)
+		{
+			Behaviour->UpdateEntityIds(InEntityIdMap);
+		}
+	}
+	
+	Super::UpdateEntityIds(InEntityIdMap);
+}
+
 #if WITH_EDITOR
 void URCController::PostEditUndo()
 {
diff --git a/Engine/Plugins/VirtualProduction/RemoteControl/Source/RemoteControlLogic/Private/Controller/RCControllerContainer.cpp b/Engine/Plugins/VirtualProduction/RemoteControl/Source/RemoteControlLogic/Private/Controller/RCControllerContainer.cpp
--- a/Engine/Plugins/VirtualProduction/RemoteControl/Source/RemoteControlLogic/Private/Controller/RCControllerContainer.cpp
+++ b/Engine/Plugins/VirtualProduction/RemoteControl/Source/RemoteControlLogic/Private/Controller/RCControllerContainer.cpp
@@ -2,10 +2,28 @@
 
 #include "Controller/RCControllerContainer.h"
 
+#include "Action/RCAction.h"
+#include "Action/RCActionContainer.h"
 #include "Controller/RCController.h"
 #include "RCVirtualProperty.h"
 #include "Templates/SubclassOf.h"
 
+void URCControllerContainer::UpdateEntityIds(const TMap<FGuid, FGuid>& InEntityIdMap)
+{
+	for (URCActionContainer* ActionContainer : SharedActionContainers)
+	{
+		if (ActionContainer)
+		{
+			ActionContainer->ForEachAction([&InEntityIdMap](URCAction* InAction)
+			{
+				InAction->UpdateEntityIds(InEntityIdMap);
+			}, /*bInRecursive*/ true);
+		}
+	}
+
+	Super::UpdateEntityIds(InEntityIdMap);
+}
+
 URCVirtualPropertyInContainer* URCControllerContainer::AddProperty(const FName& InPropertyName, TSubclassOf<URCVirtualPropertyInContainer> InPropertyClass, const EPropertyBagPropertyType InValueType, UObject* InValueTypeObject /*= nullptr*/, TArray<FPropertyBagPropertyDescMetaData> MetaData /*= TArray<FPropertyBagPropertyDescMetaData>()*/)
 {
 	// Vector Controllers
diff --git a/Engine/Plugins/VirtualProduction/RemoteControl/Source/RemoteControlLogic/Public/Action/RCAction.h b/Engine/Plugins/VirtualProduction/RemoteControl/Source/RemoteControlLogic/Public/Action/RCAction.h
--- a/Engine/Plugins/VirtualProduction/RemoteControl/Source/RemoteControlLogic/Public/Action/RCAction.h
+++ b/Engine/Plugins/VirtualProduction/RemoteControl/Source/RemoteControlLogic/Public/Action/RCAction.h
@@ -27,6 +27,12 @@ public:
 	/** Returns the parent Behaviour associated with this Action */
 	URCBehaviour* GetParentBehaviour() const;
 
+	/**
+	 * @brief Called internally when entity Ids are renewed.
+	 * @param InEntityIdMap Map of old Id to new Id.
+	 */
+	virtual void UpdateEntityIds(const TMap<FGuid, FGuid>& InEntityIdMap);
+
 public:
 	/** Exposed Property or Function field Id*/
 	UPROPERTY()
diff --git a/Engine/Plugins/VirtualProduction/RemoteControl/Source/RemoteControlLogic/Public/Action/RCActionContainer.h b/Engine/Plugins/VirtualProduction/RemoteControl/Source/RemoteControlLogic/Public/Action/RCActionContainer.h
--- a/Engine/Plugins/VirtualProduction/RemoteControl/Source/RemoteControlLogic/Public/Action/RCActionContainer.h
+++ b/Engine/Plugins/VirtualProduction/RemoteControl/Source/RemoteControlLogic/Public/Action/RCActionContainer.h
@@ -95,6 +95,9 @@ public:
 	/** Derive the parent Behaviour holding this Action Container */
 	URCBehaviour* GetParentBehaviour();
 
+	/** Call the given function on all actions of the container. */
+	void ForEachAction(TFunctionRef<void(URCAction*)> InActionFunction, bool bInRecursive);
+
 private:
 	/** Add remote control property action  */
 	URCPropertyAction* AddPropertyAction(const TSharedRef<const FRemoteControlProperty> InRemoteControlProperty);
diff --git a/Engine/Plugins/VirtualProduction/RemoteControl/Source/RemoteControlLogic/Public/Action/RCPropertyAction.h b/Engine/Plugins/VirtualProduction/RemoteControl/Source/RemoteControlLogic/Public/Action/RCPropertyAction.h
--- a/Engine/Plugins/VirtualProduction/RemoteControl/Source/RemoteControlLogic/Public/Action/RCPropertyAction.h
+++ b/Engine/Plugins/VirtualProduction/RemoteControl/Source/RemoteControlLogic/Public/Action/RCPropertyAction.h
@@ -22,6 +22,7 @@ public:
 	
 	//~ Begin URCAction interface
 	virtual void Execute() const override;
+	virtual void UpdateEntityIds(const TMap<FGuid, FGuid>& InEntityIdMap) override;
 	//~ End URCAction interface
 
 	/* Returns the FProperty underlying this Action */
diff --git a/Engine/Plugins/VirtualProduction/RemoteControl/Source/RemoteControlLogic/Public/Action/RCPropertyIdAction.h b/Engine/Plugins/VirtualProduction/RemoteControl/Source/RemoteControlLogic/Public/Action/RCPropertyIdAction.h
--- a/Engine/Plugins/VirtualProduction/RemoteControl/Source/RemoteControlLogic/Public/Action/RCPropertyIdAction.h
+++ b/Engine/Plugins/VirtualProduction/RemoteControl/Source/RemoteControlLogic/Public/Action/RCPropertyIdAction.h
@@ -20,10 +20,11 @@ class REMOTECONTROLLOGIC_API URCPropertyIdAction : public URCAction
 public:
 	//~ BEGIN : URCAction Interface
 	virtual void Execute() const override;
-	virtual void PostLoad() override;
+	virtual void UpdateEntityIds(const TMap<FGuid, FGuid>& InEntityIdMap) override;
 	//~ END : URCAction Interface
 	
 	//~ BEGIN : UObject Interface
+	virtual void PostLoad() override;
 #if WITH_EDITOR
 	virtual void PostEditChangeChainProperty(struct FPropertyChangedChainEvent& PropertyChangedEvent) override;
 #endif // WITH_EDITOR
diff --git a/Engine/Plugins/VirtualProduction/RemoteControl/Source/RemoteControlLogic/Public/Behaviour/Builtin/Path/RCSetAssetByPathBehaviour.h b/Engine/Plugins/VirtualProduction/RemoteControl/Source/RemoteControlLogic/Public/Behaviour/Builtin/Path/RCSetAssetByPathBehaviour.h
--- a/Engine/Plugins/VirtualProduction/RemoteControl/Source/RemoteControlLogic/Public/Behaviour/Builtin/Path/RCSetAssetByPathBehaviour.h
+++ b/Engine/Plugins/VirtualProduction/RemoteControl/Source/RemoteControlLogic/Public/Behaviour/Builtin/Path/RCSetAssetByPathBehaviour.h
@@ -74,9 +74,10 @@ public:
 
 	//~ Begin URCBehaviour interface
 	virtual void Initialize() override;
+	virtual void UpdateEntityIds(const TMap<FGuid, FGuid>& InEntityIdMap) override;
 	//~ End URCBehaviour interface
 
-	/** Given an Input Path, sets the Target Exposed Propterty to the Asset. */
+	/** Given an Input Path, sets the Target Exposed Property to the Asset. */
 	bool SetAssetByPath(const FString& AssetPath, const FString& DefaultString);
 
 	/** List of Supported Assets SetAssetByPath Logic Behaviour can set and change. */
@@ -88,10 +89,10 @@ public:
 	/** Sets the target entity upon which the asset will be set upon. */
 	void SetTargetEntity(const TSharedPtr<const FRemoteControlEntity>& InEntity);
 
-	/** Returns a Pointer to the current Target Entityy */
+	/** Returns a Pointer to the current Target Entity */
 	TWeakPtr<const FRemoteControlEntity> GetTargetEntity() const;
 
-	/** Auxialiary Function to apply to update the Target Texture */
+	/** Auxiliary Function to apply to update the Target Texture */
 	void UpdateTargetEntity();
 
 	/** Called whenever a change has occured in the Slate */
diff --git a/Engine/Plugins/VirtualProduction/RemoteControl/Source/RemoteControlLogic/Public/Behaviour/RCBehaviour.h b/Engine/Plugins/VirtualProduction/RemoteControl/Source/RemoteControlLogic/Public/Behaviour/RCBehaviour.h
--- a/Engine/Plugins/VirtualProduction/RemoteControl/Source/RemoteControlLogic/Public/Behaviour/RCBehaviour.h
+++ b/Engine/Plugins/VirtualProduction/RemoteControl/Source/RemoteControlLogic/Public/Behaviour/RCBehaviour.h
@@ -85,6 +85,12 @@ public:
 	const FText& GetBehaviorDescription();
 #endif
 
+	/**
+	 * @brief Called internally when entity Ids are renewed.
+	 * @param InEntityIdMap Map of old Id to new Id.
+	 */
+	virtual void UpdateEntityIds(const TMap<FGuid, FGuid>& InEntityIdMap);
+
 protected:
 	/**
 	 * It created the node if it called first time
diff --git a/Engine/Plugins/VirtualProduction/RemoteControl/Source/RemoteControlLogic/Public/Controller/RCController.h b/Engine/Plugins/VirtualProduction/RemoteControl/Source/RemoteControlLogic/Public/Controller/RCController.h
--- a/Engine/Plugins/VirtualProduction/RemoteControl/Source/RemoteControlLogic/Public/Controller/RCController.h
+++ b/Engine/Plugins/VirtualProduction/RemoteControl/Source/RemoteControlLogic/Public/Controller/RCController.h
@@ -20,6 +20,10 @@ class REMOTECONTROLLOGIC_API URCController : public URCVirtualPropertyInContaine
 
 public:
 
+	//~ Begin URCVirtualPropertyBase
+	virtual void UpdateEntityIds(const TMap<FGuid, FGuid>& InEntityIdMap) override;
+	//~ End URCVirtualPropertyBase
+	
 #if WITH_EDITOR
 	/** Called after applying a transaction to the object. Used to broadcast Undo related container changes to UI Used to broadcast Undo related container changes to UI */
 	virtual void PostEditUndo();
diff --git a/Engine/Plugins/VirtualProduction/RemoteControl/Source/RemoteControlLogic/Public/Controller/RCControllerContainer.h b/Engine/Plugins/VirtualProduction/RemoteControl/Source/RemoteControlLogic/Public/Controller/RCControllerContainer.h
--- a/Engine/Plugins/VirtualProduction/RemoteControl/Source/RemoteControlLogic/Public/Controller/RCControllerContainer.h
+++ b/Engine/Plugins/VirtualProduction/RemoteControl/Source/RemoteControlLogic/Public/Controller/RCControllerContainer.h
@@ -17,6 +17,10 @@ class REMOTECONTROLLOGIC_API URCControllerContainer : public URCVirtualPropertyC
 	GENERATED_BODY()
 
 public:
+	//~ Begin URCVirtualPropertyContainerBase
+	virtual void UpdateEntityIds(const TMap<FGuid, FGuid>& InEntityIdMap) override;
+	//~ End URCVirtualPropertyContainerBase
+	
 	/** Adds a new Controller to this container */
 	virtual URCVirtualPropertyInContainer* AddProperty(const FName& InPropertyName, TSubclassOf<URCVirtualPropertyInContainer> InPropertyClass, const EPropertyBagPropertyType InValueType, UObject* InValueTypeObject = nullptr, TArray<FPropertyBagPropertyDescMetaData> MetaData = TArray<FPropertyBagPropertyDescMetaData>()) override;
 	
