[RemoteControl]
Fix issue 412245: Warnings in Build Engine Localization and Localize UnrealEditor Main

diff --git a/Engine/Plugins/VirtualProduction/RemoteControl/Source/RemoteControlUI/Private/UI/Controller/SRCControllerPanel.cpp b/Engine/Plugins/VirtualProduction/RemoteControl/Source/RemoteControlUI/Private/UI/Controller/SRCControllerPanel.cpp
--- a/Engine/Plugins/VirtualProduction/RemoteControl/Source/RemoteControlUI/Private/UI/Controller/SRCControllerPanel.cpp
+++ b/Engine/Plugins/VirtualProduction/RemoteControl/Source/RemoteControlUI/Private/UI/Controller/SRCControllerPanel.cpp
@@ -206,7 +206,7 @@ FText SRCControllerPanel::GetPasteItemMenuEntrySuffix()
 					{
 						return FText::Format(LOCTEXT("ControllerPanelPasteMenuMultiEntrySuffix", "Controller {0} and {1} other(s)"), FText::FromName(Controller->DisplayName), (LogicClipboardItems.Num() - 1));
 					}
-					return FText::Format(LOCTEXT("ControllerPanelPasteMenuMultiEntrySuffix", "Controller {0}"), FText::FromName(Controller->DisplayName));
+					return FText::Format(LOCTEXT("ControllerPanelPasteMenuEntrySuffix", "Controller {0}"), FText::FromName(Controller->DisplayName));
 				}
 			}
 		}
