[RemoteControl]
Adding controller description and made controller name unique

diff --git a/Engine/Plugins/VirtualProduction/RemoteControl/Source/RemoteControl/Private/RCVirtualPropertyContainer.cpp b/Engine/Plugins/VirtualProduction/RemoteControl/Source/RemoteControl/Private/RCVirtualPropertyContainer.cpp
--- a/Engine/Plugins/VirtualProduction/RemoteControl/Source/RemoteControl/Private/RCVirtualPropertyContainer.cpp
+++ b/Engine/Plugins/VirtualProduction/RemoteControl/Source/RemoteControl/Private/RCVirtualPropertyContainer.cpp
@@ -44,10 +44,12 @@ URCVirtualPropertyInContainer* URCVirtualPropertyContainerBase::AddProperty(cons
 	// Create Property in Container
 	URCVirtualPropertyInContainer* VirtualPropertyInContainer = NewObject<URCVirtualPropertyInContainer>(this, InPropertyClass.Get(), NAME_None, RF_Transactional);
 	VirtualPropertyInContainer->PropertyName = PropertyName;
+	VirtualPropertyInContainer->DisplayName = PropertyName;
 	VirtualPropertyInContainer->PresetWeakPtr = PresetWeakPtr;
 	VirtualPropertyInContainer->ContainerWeakPtr = this;
 	VirtualPropertyInContainer->Id = FGuid::NewGuid();
 
+	ControllerLabelToIdCache.Add(PropertyName, VirtualPropertyInContainer->Id);
 	AddVirtualProperty(VirtualPropertyInContainer);
 
 	return VirtualPropertyInContainer;
@@ -66,10 +68,12 @@ URCVirtualPropertyInContainer* URCVirtualPropertyContainerBase::DuplicatePropert
 	
 	URCVirtualPropertyInContainer* VirtualPropertyInContainer = NewObject<URCVirtualPropertyInContainer>(this, InPropertyClass.Get());
 	VirtualPropertyInContainer->PropertyName = InPropertyName;
+	VirtualPropertyInContainer->DisplayName = GenerateUniqueDisplayName(VirtualPropertyInContainer->DisplayName, this);;
 	VirtualPropertyInContainer->PresetWeakPtr = PresetWeakPtr;
 	VirtualPropertyInContainer->ContainerWeakPtr = this;
 	VirtualPropertyInContainer->Id = FGuid::NewGuid();
 
+	ControllerLabelToIdCache.Add(InPropertyName, VirtualPropertyInContainer->Id);
 	AddVirtualProperty(VirtualPropertyInContainer);
 
 	return VirtualPropertyInContainer;
@@ -101,6 +105,7 @@ URCVirtualPropertyInContainer* URCVirtualPropertyContainerBase::DuplicateVirtual
 	if (URCVirtualPropertyInContainer* NewVirtualProperty = DuplicateObject<URCVirtualPropertyInContainer>(InVirtualProperty, InVirtualProperty->GetOuter()))
 	{
 		NewVirtualProperty->PropertyName = GenerateUniquePropertyName(InVirtualProperty->PropertyName, this);
+		NewVirtualProperty->DisplayName = GenerateUniqueDisplayName(InVirtualProperty->DisplayName, this);
 		NewVirtualProperty->Id = FGuid::NewGuid();
 
 		// Sync Property Bag
@@ -116,6 +121,9 @@ URCVirtualPropertyInContainer* URCVirtualPropertyContainerBase::DuplicateVirtual
 		// Sync Virtual Properties List
 		AddVirtualProperty(NewVirtualProperty);
 
+		// Sync Cache
+		ControllerLabelToIdCache.Add(NewVirtualProperty->DisplayName, NewVirtualProperty->Id);
+
 		return NewVirtualProperty;
 	}
 
@@ -125,8 +133,9 @@ URCVirtualPropertyInContainer* URCVirtualPropertyContainerBase::DuplicateVirtual
 bool URCVirtualPropertyContainerBase::RemoveProperty(const FName& InPropertyName)
 {
 	Bag.RemovePropertyByName(InPropertyName);
+	ControllerLabelToIdCache.Remove(InPropertyName);
 
-	for (auto PropertiesIt = VirtualProperties.CreateIterator(); PropertiesIt; ++PropertiesIt)
+	for (TSet<TObjectPtr<URCVirtualPropertyBase>>::TIterator PropertiesIt = VirtualProperties.CreateIterator(); PropertiesIt; ++PropertiesIt)
 	{
 		if (const URCVirtualPropertyBase* VirtualProperty = *PropertiesIt)
 		{
@@ -144,7 +153,7 @@ bool URCVirtualPropertyContainerBase::RemoveProperty(const FName& InPropertyName
 void URCVirtualPropertyContainerBase::Reset()
 {
 	VirtualProperties.Empty();
-
+	ControllerLabelToIdCache.Reset();
 	Bag.Reset();
 }
 
@@ -181,6 +190,14 @@ URCVirtualPropertyBase* URCVirtualPropertyContainerBase::GetVirtualProperty(cons
 
 URCVirtualPropertyBase* URCVirtualPropertyContainerBase::GetVirtualPropertyByDisplayName(const FName InDisplayName) const
 {
+	if (const FGuid* ControllerId = ControllerLabelToIdCache.Find(InDisplayName))
+	{
+		if (URCVirtualPropertyBase* Controller = GetVirtualProperty(*ControllerId))
+		{
+			return Controller;
+		}
+	}
+
 	for (URCVirtualPropertyBase* VirtualProperty : VirtualProperties)
 	{
 		if (VirtualProperty->DisplayName == InDisplayName)
@@ -250,6 +267,22 @@ TSharedPtr<FStructOnScope> URCVirtualPropertyContainerBase::CreateStructOnScope(
 	return MakeShared<FStructOnScope>(Bag.GetPropertyBagStruct(), Bag.GetMutableValue().GetMemory());
 }
 
+FName URCVirtualPropertyContainerBase::SetControllerDisplayName(FGuid InGuid, FName InNewName)
+{
+	if (URCVirtualPropertyBase* Controller = GetVirtualProperty(InGuid))
+	{
+		Controller->Modify();
+
+		ControllerLabelToIdCache.Remove(Controller->DisplayName);
+		Controller->DisplayName = GenerateUniqueDisplayName(InNewName, this);
+		ControllerLabelToIdCache.Add(Controller->DisplayName, Controller->Id);
+
+		return Controller->DisplayName;
+	}
+
+	return NAME_None;
+}
+
 FName URCVirtualPropertyContainerBase::GenerateUniquePropertyName(const FName& InPropertyName, const EPropertyBagPropertyType InValueType, UObject* InValueTypeObject, const URCVirtualPropertyContainerBase* InContainer)
 {
 	FName BaseName = InPropertyName;
@@ -291,6 +324,36 @@ FName URCVirtualPropertyContainerBase::GenerateUniquePropertyName(const FName& I
 	return *FinalName;
 }
 
+FName URCVirtualPropertyContainerBase::GenerateUniqueDisplayName(const FName& InPropertyName, const URCVirtualPropertyContainerBase* InContainer)
+{
+	auto GetFinalName = [](const FName& InPrefix, const int32 InIndex = 0)
+	{
+		FName FinalName = InPrefix;
+
+		if (InIndex > 0)
+		{
+			FinalName = FName(FinalName.ToString() + TEXT("_") + FString::FromInt(InIndex));
+		}
+
+		return FinalName;
+	};
+
+	int32 Index = 0;
+	const FName InitialName = InPropertyName;
+	FName FinalName = InitialName;
+
+	// Recursively search for an available name by incrementing suffix till we find one.
+	bool bControllerExist = InContainer->ControllerLabelToIdCache.Contains(FinalName);
+	while (bControllerExist)
+	{
+		++Index;
+		FinalName = GetFinalName(InitialName, Index);
+		bControllerExist = InContainer->ControllerLabelToIdCache.Contains(FinalName);
+	}
+
+	return FinalName;
+}
+
 void URCVirtualPropertyContainerBase::UpdateEntityIds(const TMap<FGuid, FGuid>& InEntityIdMap)
 {
 	for (const TObjectPtr<URCVirtualPropertyBase>& VirtualProperty : VirtualProperties)
@@ -302,6 +365,43 @@ void URCVirtualPropertyContainerBase::UpdateEntityIds(const TMap<FGuid, FGuid>&
 	}
 }
 
+void URCVirtualPropertyContainerBase::CacheControllersLabels()
+{
+	ControllerLabelToIdCache.Reset();
+	for (const URCVirtualPropertyBase* Controller : VirtualProperties)
+	{
+		if (Controller)
+		{
+			if (!ControllerLabelToIdCache.Contains(Controller->DisplayName))
+			{
+				ControllerLabelToIdCache.Add(Controller->DisplayName, Controller->Id);
+			}
+		}
+	}
+}
+
+void URCVirtualPropertyContainerBase::FixAndCacheControllersLabels()
+{
+	ControllerLabelToIdCache.Reset();
+	for (URCVirtualPropertyBase* Controller : VirtualProperties)
+	{
+		if (Controller)
+		{
+			if (!ControllerLabelToIdCache.Contains(Controller->DisplayName))
+			{
+				ControllerLabelToIdCache.Add(Controller->DisplayName, Controller->Id);
+			}
+			// Cases for older presets where you could have the same DisplayName, and if it is already saved then we need to change it now before caching it
+			else
+			{
+				const FName NewControllerName = GenerateUniqueDisplayName(Controller->DisplayName, this);
+				Controller->DisplayName = NewControllerName;
+				ControllerLabelToIdCache.Add(NewControllerName, Controller->Id);
+			}
+		}
+	}
+}
+
 #if WITH_EDITOR
 void URCVirtualPropertyContainerBase::PostEditUndo()
 {
diff --git a/Engine/Plugins/VirtualProduction/RemoteControl/Source/RemoteControl/Private/RemoteControlPreset.cpp b/Engine/Plugins/VirtualProduction/RemoteControl/Source/RemoteControl/Private/RemoteControlPreset.cpp
--- a/Engine/Plugins/VirtualProduction/RemoteControl/Source/RemoteControl/Private/RemoteControlPreset.cpp
+++ b/Engine/Plugins/VirtualProduction/RemoteControl/Source/RemoteControl/Private/RemoteControlPreset.cpp
@@ -66,7 +66,7 @@ namespace
 	}
 
 	/** Create a unique name. */
-	FName MakeUniqueName(FName InBase, TFunctionRef<bool(FName)> NamePoolContains)
+	FName MakeUniqueName(FName InBase, TFunctionRef<bool(FName)> NamePoolContains, const FString& InSeparatorLeft = TEXT(" ("), const FString& InSeparatorRight = TEXT(")"))
 	{
 		// Try using the field name itself
 		if (!NamePoolContains(InBase))
@@ -77,7 +77,7 @@ namespace
 		// Then try the field name with a suffix
 		for (uint32 Index = 1; Index < 1000; ++Index)
 		{
-			const FName Candidate = FName(*FString::Printf(TEXT("%s (%d)"), *InBase.ToString(), Index));
+			const FName Candidate = FName(*FString::Printf(TEXT("%s%s%d%s"), *InBase.ToString(), *InSeparatorLeft, Index, *InSeparatorRight));
 			if (!NamePoolContains(Candidate))
 			{
 				return Candidate;
@@ -692,6 +692,8 @@ void URemoteControlPreset::PostLoad()
 
 	CacheFieldLayoutData();
 
+	FixAndCacheControllersLabels();
+
 	InitializeEntitiesMetadata();
 
 	RegisterEntityDelegates();
@@ -1735,6 +1737,14 @@ void URemoteControlPreset::CacheLayoutData()
 	CacheFieldLayoutData();
 }
 
+void URemoteControlPreset::CacheControllersLabels() const
+{
+	if (ControllerContainer)
+	{
+		ControllerContainer->CacheControllersLabels();
+	}
+}
+
 TArray<UObject*> URemoteControlPreset::ResolvedBoundObjects(FName FieldLabel)
 {
 	TArray<UObject*> Objects;
@@ -1836,11 +1846,21 @@ void URemoteControlPreset::RenewControllerIds()
 		ControllerIdMap.Add(OldControllerId, Controller->Id);
 	}
 
-	// It doesn't seem like controller Ids are used as hashing key.
-	
+	// Re-cache after the Ids are updated
+	CacheControllersLabels();
+
 	OnControllerIdsRenewed().Broadcast(this, ControllerIdMap);
 }
 
+FName URemoteControlPreset::SetControllerDisplayName(FGuid InControllerGuid, const FName& InNewName) const
+{
+	if (ControllerContainer)
+	{
+		return ControllerContainer->SetControllerDisplayName(InControllerGuid, InNewName);
+	}
+	return NAME_None;
+}
+
 void URemoteControlPreset::NotifyExposedPropertyChanged(FName PropertyLabel)
 {
 	if (TSharedPtr<FRemoteControlProperty> ExposedProperty = GetExposedEntity<FRemoteControlProperty>(GetExposedEntityId(PropertyLabel)).Pin())
@@ -2347,6 +2367,14 @@ void URemoteControlPreset::OnObjectTransacted(UObject* InObject, const FTransact
 	}
 }
 
+void URemoteControlPreset::FixAndCacheControllersLabels() const
+{
+	if (ControllerContainer)
+	{
+		ControllerContainer->FixAndCacheControllersLabels();
+	}
+}
+
 void URemoteControlPreset::RegisterDelegates()
 {
 	UnregisterDelegates();
diff --git a/Engine/Plugins/VirtualProduction/RemoteControl/Source/RemoteControl/Public/RCVirtualProperty.h b/Engine/Plugins/VirtualProduction/RemoteControl/Source/RemoteControl/Public/RCVirtualProperty.h
--- a/Engine/Plugins/VirtualProduction/RemoteControl/Source/RemoteControl/Public/RCVirtualProperty.h
+++ b/Engine/Plugins/VirtualProduction/RemoteControl/Source/RemoteControl/Public/RCVirtualProperty.h
@@ -330,10 +330,14 @@ public:
 	UPROPERTY()
 	TWeakObjectPtr<URemoteControlPreset> PresetWeakPtr;
 
-	/** User friendly name of the Controller*/
+	/** Unique user friendly name of the controller, used as identifier in some API. */
 	UPROPERTY()
 	FName DisplayName;
 
+	/** Controller description */
+	UPROPERTY()
+	FText Description;
+
 	/** User configurable Display Index for this Virtual Property (as Logic Controller) when represented as a row in the RC Logic Controllers list  */
 	UPROPERTY()
 	int32 DisplayIndex;
diff --git a/Engine/Plugins/VirtualProduction/RemoteControl/Source/RemoteControl/Public/RCVirtualPropertyContainer.h b/Engine/Plugins/VirtualProduction/RemoteControl/Source/RemoteControl/Public/RCVirtualPropertyContainer.h
--- a/Engine/Plugins/VirtualProduction/RemoteControl/Source/RemoteControl/Public/RCVirtualPropertyContainer.h
+++ b/Engine/Plugins/VirtualProduction/RemoteControl/Source/RemoteControl/Public/RCVirtualPropertyContainer.h
@@ -122,6 +122,11 @@ public:
 	 */
 	TSharedPtr<FStructOnScope> CreateStructOnScope();
 
+	/**
+	 * Rename the controller with the given InGuid with the InNewName passed as candidate and return the new name
+	 */
+	FName SetControllerDisplayName(FGuid InGuid, FName InNewName);
+
 	/**
 	 * Generates unique name for the property for specified property container
 	 */
@@ -132,12 +137,23 @@ public:
 	*/
 	static FName GenerateUniquePropertyName(const FName& InPropertyName, const URCVirtualPropertyContainerBase* InContainer);
 
+	/**
+	 * Generates unique display name for the controllers
+	 */
+	static FName GenerateUniqueDisplayName(const FName& InPropertyName, const URCVirtualPropertyContainerBase* InContainer);
+
 	/**
 	 * @brief Called internally when entity Ids are renewed.
 	 * @param InEntityIdMap Map of old Id to new Id.
 	 */
 	virtual void UpdateEntityIds(const TMap<FGuid, FGuid>& InEntityIdMap);
 
+	/** Cache this preset's controllers labels. */
+	void CacheControllersLabels();
+
+	/** Fix controllers labels for older presets. */
+	void FixAndCacheControllersLabels();
+
 #if WITH_EDITOR
 	/** Called after applying a transaction to the object. Used to broadcast Undo related container changes to UI. */
 	virtual void PostEditUndo();
@@ -172,6 +188,11 @@ public:
 
 	/** Delegate that notifies changes to the virtual property container*/
 	FOnVirtualPropertyContainerModified OnVirtualPropertyContainerModifiedDelegate;
+
+private:
+	/** Map of Controller Name to GUID. */
+	UPROPERTY(Transient)
+	TMap<FName, FGuid> ControllerLabelToIdCache;
 };
 
 
diff --git a/Engine/Plugins/VirtualProduction/RemoteControl/Source/RemoteControl/Public/RemoteControlPreset.h b/Engine/Plugins/VirtualProduction/RemoteControl/Source/RemoteControl/Public/RemoteControlPreset.h
--- a/Engine/Plugins/VirtualProduction/RemoteControl/Source/RemoteControl/Public/RemoteControlPreset.h
+++ b/Engine/Plugins/VirtualProduction/RemoteControl/Source/RemoteControl/Public/RemoteControlPreset.h
@@ -577,6 +577,9 @@ public:
 	/** Cache this preset's layout data. */
 	void CacheLayoutData();
 
+	/** Cache this preset's controllers labels. */
+	void CacheControllersLabels() const;
+
 	/** Resolves exposed property/function bounded objects */
 	UE_DEPRECATED(4.27, "ResolvedBoundObjects is deprecated, you can now resolve bound objects using FRemoteControlEntity.")
 	TArray<UObject*> ResolvedBoundObjects(FName FieldLabel);
@@ -600,7 +603,15 @@ public:
 	 * Renews all controller guids. Necessary when duplicating.
 	 */
 	void RenewControllerIds();
-	
+
+	/**
+	 * Rename the given controller with the new name if the name is unique
+	 * @param InControllerGuid Id of the controller to rename
+	 * @param InNewName New name for the controller
+	 * @return New name of the controller
+	 */
+	FName SetControllerDisplayName(FGuid InControllerGuid, const FName& InNewName) const;
+
 	DECLARE_MULTICAST_DELEGATE_TwoParams(FOnPresetEntityEvent, URemoteControlPreset* /*Preset*/, const FGuid& /*EntityId*/);
 	FOnPresetEntityEvent& OnEntityExposed() { return OnEntityExposedDelegate; }
 	FOnPresetEntityEvent& OnEntityUnexposed() { return OnEntityUnexposedDelegate; }
@@ -787,6 +798,9 @@ private:
 	void OnPreObjectPropertyChanged(UObject* Object, const class FEditPropertyChain& PropertyChain);
 	void OnObjectTransacted(UObject* InObject, const FTransactionObjectEvent& InTransactionEvent);
 
+	/** Fix controllers labels for older presets. */
+	void FixAndCacheControllersLabels() const;
+
 #if WITH_EDITOR	
 	//~ Handle events that can incur bindings to be modified.
 	void OnActorDeleted(AActor* Actor);
diff --git a/Engine/Plugins/VirtualProduction/RemoteControl/Source/RemoteControlUI/Private/UI/Controller/RCControllerModel.cpp b/Engine/Plugins/VirtualProduction/RemoteControl/Source/RemoteControlUI/Private/UI/Controller/RCControllerModel.cpp
--- a/Engine/Plugins/VirtualProduction/RemoteControl/Source/RemoteControlUI/Private/UI/Controller/RCControllerModel.cpp
+++ b/Engine/Plugins/VirtualProduction/RemoteControl/Source/RemoteControlUI/Private/UI/Controller/RCControllerModel.cpp
@@ -1,4 +1,4 @@
-﻿// Copyright Epic Games, Inc. All Rights Reserved.
+// Copyright Epic Games, Inc. All Rights Reserved.
 
 #include "RCControllerModel.h"
 #include "Controller/RCCustomControllerUtilities.h"
@@ -6,8 +6,10 @@
 #include "IDetailTreeNode.h"
 #include "RCVirtualProperty.h"
 #include "RemoteControlPreset.h"
+#include "ScopedTransaction.h"
 #include "TypeTranslator/RCTypeTranslator.h"
 #include "UI/SRemoteControlPanel.h"
+#include "Widgets/Input/SEditableTextBox.h"
 #include "Widgets/Input/STextComboBox.h"
 #include "Widgets/Text/SInlineEditableTextBlock.h"
 
@@ -24,11 +26,16 @@ FRCControllerModel::FRCControllerModel(URCVirtualPropertyBase* InVirtualProperty
 		{
 			InVirtualProperty->DisplayName = InVirtualProperty->PropertyName;
 		}
-		
-		SAssignNew(ControllerNameTextBox, SInlineEditableTextBlock)
+
+		SAssignNew(ControllerNameTextBox, SEditableTextBox)
 			.Text(FText::FromName(InVirtualProperty->DisplayName))
 			.OnTextCommitted_Raw(this, &FRCControllerModel::OnControllerNameCommitted);
 
+		SAssignNew(ControllerDescriptionTextBox, SInlineEditableTextBlock)
+			.Text(InVirtualProperty->Description)
+			.MultiLine(true)
+			.OnTextCommitted_Raw(this, &FRCControllerModel::OnControllerDescriptionCommitted);
+
 		SAssignNew(ControllerFieldIdTextBox, SInlineEditableTextBlock)
 			.Text(FText::FromName(InVirtualProperty->FieldId))
 			.OnTextCommitted_Raw(this, &FRCControllerModel::OnControllerFieldIdCommitted);
@@ -102,6 +109,14 @@ TSharedRef<SWidget> FRCControllerModel::GetNameWidget() const
 		];
 }
 
+TSharedRef<SWidget> FRCControllerModel::GetDescriptionWidget() const
+{
+	return SNew(SBox).Padding(10.f, 2.f)
+		[
+			ControllerDescriptionTextBox.ToSharedRef()
+		];
+}
+
 TSharedRef<SWidget> FRCControllerModel::GetFieldIdWidget() const
 {
 	return SNew(SBox).Padding(10.f, 2.f)
@@ -182,16 +197,48 @@ void FRCControllerModel::UpdateSelectedBehaviourModel(TSharedPtr<FRCBehaviourMod
 	SelectedBehaviourModelWeakPtr = InModel;
 }
 
+void FRCControllerModel::PostUndo(bool bSuccess)
+{
+	if (bSuccess)
+	{
+		if (const URCVirtualPropertyBase* Controller = GetVirtualProperty())
+		{
+			if (URemoteControlPreset* Preset = GetPreset())
+			{
+				// Cache controllers label again here since during Undo/Redo the cache map won't be changed
+				Preset->CacheControllersLabels();
+			}
+			ControllerNameTextBox->SetText(FText::FromName(Controller->DisplayName));
+			ControllerDescriptionTextBox->SetText(Controller->Description);
+		}
+	}
+}
+
 void FRCControllerModel::OnControllerNameCommitted(const FText& InNewControllerName, ETextCommit::Type InCommitInfo)
+{
+	if (URemoteControlPreset* Preset = GetPreset())
+	{
+		if (const URCVirtualPropertyBase* Controller = GetVirtualProperty())
+		{
+			FScopedTransaction Transaction(LOCTEXT("RenameController", "Rename Controller"));
+			const FName OldName = Controller->DisplayName;
+			const FName AssignedLabel = Preset->SetControllerDisplayName(Controller->Id, FName(FText::TrimPrecedingAndTrailing(InNewControllerName).ToString()));
+			ControllerNameTextBox->SetText(FText::FromName(AssignedLabel));
+			Preset->OnControllerRenamed().Broadcast(Preset, OldName, AssignedLabel);
+		}
+	}
+}
+
+void FRCControllerModel::OnControllerDescriptionCommitted(const FText& InNewControllerDescription, ETextCommit::Type InCommitInfo)
 {
 	if (URemoteControlPreset* Preset = GetPreset())
 	{
 		if (URCVirtualPropertyBase* Controller = GetVirtualProperty())
 		{
-			const FName& OldName = Controller->DisplayName;
-			Controller->DisplayName = *InNewControllerName.ToString();
-			ControllerNameTextBox->SetText(InNewControllerName);
-			Preset->OnControllerRenamed().Broadcast(Preset, OldName, *InNewControllerName.ToString());
+			FScopedTransaction Transaction(LOCTEXT("ChangedControllerDescription", "Update controller description"));
+			Controller->Modify();
+			Controller->Description = InNewControllerDescription;
+			ControllerDescriptionTextBox->SetText(InNewControllerDescription);
 		}
 	}
 }
@@ -276,9 +323,9 @@ void FRCControllerModel::InitControlledTypes()
 	}
 }
 
-void FRCControllerModel::EnterRenameMode()
+void FRCControllerModel::EnterDescriptionEditingMode()
 {	
-	ControllerNameTextBox->EnterEditingMode();
+	ControllerDescriptionTextBox->EnterEditingMode();
 }
 
 FName FRCControllerModel::GetControllerDisplayName()
@@ -291,4 +338,14 @@ FName FRCControllerModel::GetControllerDisplayName()
 	return NAME_None;
 }
 
+FText FRCControllerModel::GetControllerDescription()
+{
+	if (URCVirtualPropertyBase* Controller = GetVirtualProperty())
+	{
+		return Controller->Description;
+	}
+
+	return FText::GetEmpty();
+}
+
 #undef LOCTEXT_NAMESPACE
diff --git a/Engine/Plugins/VirtualProduction/RemoteControl/Source/RemoteControlUI/Private/UI/Controller/RCControllerModel.h b/Engine/Plugins/VirtualProduction/RemoteControl/Source/RemoteControlUI/Private/UI/Controller/RCControllerModel.h
--- a/Engine/Plugins/VirtualProduction/RemoteControl/Source/RemoteControlUI/Private/UI/Controller/RCControllerModel.h
+++ b/Engine/Plugins/VirtualProduction/RemoteControl/Source/RemoteControlUI/Private/UI/Controller/RCControllerModel.h
@@ -1,7 +1,8 @@
-﻿// Copyright Epic Games, Inc. All Rights Reserved.
+// Copyright Epic Games, Inc. All Rights Reserved.
 
 #pragma once
 
+#include "EditorUndoClient.h"
 #include "UI/BaseLogicUI/RCLogicModeBase.h"
 #include "UObject/WeakFieldPtr.h"
 #include "Widgets/DeclarativeSyntaxSupport.h"
@@ -10,6 +11,7 @@
 enum class EPropertyBagPropertyType : uint8;
 class FRCBehaviourModel;
 class IDetailTreeNode;
+class SEditableTextBox;
 class SInlineEditableTextBlock;
 class SRemoteControlPanel;
 class URCBehaviour;
@@ -25,7 +27,7 @@ class URCVirtualPropertyBase;
 * Contains a row widget with Controller Name, Field Id and a Value Widget.
 * If this is setup to be a MultiController, a Value Type selection Widget will be available as well.
 */
-class FRCControllerModel : public FRCLogicModeBase
+class FRCControllerModel : public FRCLogicModeBase, public FSelfRegisteringEditorUndoClient
 {
 public:
 	DECLARE_MULTICAST_DELEGATE_TwoParams(FOnValueTypeChanged, URCVirtualPropertyBase* /* InController */, EPropertyBagPropertyType /* InValueType */);
@@ -42,6 +44,9 @@ public:
 	/** The widget showing the Name of this Controller */
 	TSharedRef<SWidget> GetNameWidget() const;
 
+	/** The widget showing the Description of this Controller */
+	TSharedRef<SWidget> GetDescriptionWidget() const;
+
 	/** The widget showing the Field Id of this Controller */
 	TSharedRef<SWidget> GetFieldIdWidget() const;
 
@@ -70,11 +75,14 @@ public:
 	void UpdateSelectedBehaviourModel(TSharedPtr<FRCBehaviourModel> InModel);
 
 	/** Allows users to enter text into the Controller Name Box */
-	void EnterRenameMode();
+	void EnterDescriptionEditingMode();
 
 	/** User-friendly Name of the underlying Controller */
 	FName GetControllerDisplayName();
 
+	/** Description of the underlying Controller */
+	FText GetControllerDescription();
+
 	/**Fetches the unique Id for this UI item */
 	FGuid GetId() const { return Id; }
 
@@ -85,9 +93,17 @@ public:
 	FOnValueChanged OnValueChanged;
 
 private:
+	//~ Begin FEditorUndoClient Interface
+	virtual void PostUndo(bool bSuccess) override;
+	virtual void PostRedo(bool bSuccess) override { PostUndo(bSuccess); };
+	// End of FEditorUndoClient
+
 	/** Text commit event for Controller Name text box */
 	void OnControllerNameCommitted(const FText& InNewControllerName, ETextCommit::Type InCommitInfo);
 
+	/** Text commit event for Controller description text box */
+	void OnControllerDescriptionCommitted(const FText& InNewControllerDescription, ETextCommit::Type InCommitInfo);
+
 	/** Text commit event for Controller FieldId text box */
 	void OnControllerFieldIdCommitted(const FText& InNewControllerFieldId, ETextCommit::Type InCommitInfo);
 
@@ -110,7 +126,10 @@ private:
 	TWeakPtr<FRCBehaviourModel>  SelectedBehaviourModelWeakPtr;
 
 	/** Controller name - editable text box */
-	TSharedPtr<SInlineEditableTextBlock> ControllerNameTextBox;
+	TSharedPtr<SEditableTextBox> ControllerNameTextBox;
+
+	/** Controller name - editable text box */
+	TSharedPtr<SInlineEditableTextBlock> ControllerDescriptionTextBox;
 
 	/** Controller Field Id - editable text box */
 	TSharedPtr<SInlineEditableTextBlock> ControllerFieldIdTextBox;
diff --git a/Engine/Plugins/VirtualProduction/RemoteControl/Source/RemoteControlUI/Private/UI/Controller/SRCControllerPanelList.cpp b/Engine/Plugins/VirtualProduction/RemoteControl/Source/RemoteControlUI/Private/UI/Controller/SRCControllerPanelList.cpp
--- a/Engine/Plugins/VirtualProduction/RemoteControl/Source/RemoteControlUI/Private/UI/Controller/SRCControllerPanelList.cpp
+++ b/Engine/Plugins/VirtualProduction/RemoteControl/Source/RemoteControlUI/Private/UI/Controller/SRCControllerPanelList.cpp
@@ -53,6 +53,7 @@ namespace UE::RCControllerPanelList
 	{
 		const FName TypeColor = TEXT("TypeColor");
 		const FName Name = TEXT("Controller Name");
+		const FName Description = TEXT("Controller Description");
 		const FName Value = TEXT("Controller Value");
 		const FName DragHandle = TEXT("Drag Handle");
 		const FName FieldId = TEXT("Controller Field Id");
@@ -138,6 +139,10 @@ namespace UE::RCControllerPanelList
 			{
 				return WrapWithDropTarget(ControllerItem->GetNameWidget());
 			}
+			else if (ColumnName == UE::RCControllerPanelList::Columns::Description)
+			{
+				return WrapWithDropTarget(ControllerItem->GetDescriptionWidget());
+			}
 			else if (ColumnName == UE::RCControllerPanelList::Columns::Value)
 			{
 				return ControllerItem->GetWidget();
@@ -321,7 +326,7 @@ namespace UE::RCControllerPanelList
 		{
 			if (ControllerItem.IsValid())
 			{
-				ControllerItem->EnterRenameMode();
+				ControllerItem->EnterDescriptionEditingMode();
 			}
 
 			return FSuperRowType::OnMouseButtonDoubleClick(InMyGeometry, InMouseEvent);
@@ -353,20 +358,24 @@ void SRCControllerPanelList::Construct(const FArguments& InArgs, const TSharedRe
 			.Style(&RCPanelStyle->HeaderRowStyle)
 
 			+ SHeaderRow::Column(UE::RCControllerPanelList::Columns::TypeColor)
-			.DefaultLabel(LOCTEXT("ControllerNameColumnName", ""))
-			.FixedWidth(30)
+			.DefaultLabel(LOCTEXT("ControllerColorColumnName", ""))
+			.FixedWidth(15)
 			.HeaderContentPadding(RCPanelStyle->HeaderRowPadding)
 
 			+ SHeaderRow::Column(UE::RCControllerPanelList::Columns::DragHandle)
 			.DefaultLabel(FText::GetEmpty())
-			.FixedWidth(30)
+			.FixedWidth(15)
 			.HeaderContentPadding(RCPanelStyle->HeaderRowPadding)
 
 			+ SHeaderRow::Column(UE::RCControllerPanelList::Columns::Name)
 			.DefaultLabel(LOCTEXT("ControllerNameColumnName", "Name"))
-			.FillWidth(0.15f)
+			.FillWidth(0.2f)
 			.HeaderContentPadding(RCPanelStyle->HeaderRowPadding)
 
+			+ SHeaderRow::Column(UE::RCControllerPanelList::Columns::Description)
+			.DefaultLabel(LOCTEXT("ControllerNameColumnDescription", "Description"))
+			.FillWidth(0.35f)
+
 			+ SHeaderRow::Column(UE::RCControllerPanelList::Columns::Value)
 			.DefaultLabel(LOCTEXT("ControllerValueColumnName", "Input"))
 			.FillWidth(0.45f)
@@ -764,7 +773,7 @@ void SRCControllerPanelList::EnterRenameMode()
 {
 	if (TSharedPtr<FRCControllerModel> SelectedItem = SelectedControllerItemWeakPtr.Pin())
 	{
-		SelectedItem->EnterRenameMode();
+		SelectedItem->EnterDescriptionEditingMode();
 	}
 }
 
diff --git a/Engine/Plugins/VirtualProduction/RemoteControl/Source/WebRemoteControl/Private/RemoteControlModels.h b/Engine/Plugins/VirtualProduction/RemoteControl/Source/WebRemoteControl/Private/RemoteControlModels.h
--- a/Engine/Plugins/VirtualProduction/RemoteControl/Source/WebRemoteControl/Private/RemoteControlModels.h
+++ b/Engine/Plugins/VirtualProduction/RemoteControl/Source/WebRemoteControl/Private/RemoteControlModels.h
@@ -386,6 +386,7 @@ struct FRCControllerDescription
 		}
 
 		DisplayName = InController->DisplayName;
+		Description = InController->Description;
 		ID = InController->Id.ToString();
 		Path = InController->GetPathName();
 		Type = InController->GetProperty()->GetCPPType();
@@ -395,7 +396,11 @@ struct FRCControllerDescription
 	/** The label displayed in the remote control panel for this controller. */
 	UPROPERTY()
 	FName DisplayName;
-	
+
+	/** Controller description. */
+	UPROPERTY()
+	FText Description;
+
 	/** Unique identifier for the controller. */
 	UPROPERTY()
 	FString ID;
