[RemoteControl]
Fix controllers unique names to take into account also deleted ones

diff --git a/Engine/Plugins/VirtualProduction/RemoteControl/Source/RemoteControl/Private/RCVirtualPropertyContainer.cpp b/Engine/Plugins/VirtualProduction/RemoteControl/Source/RemoteControl/Private/RCVirtualPropertyContainer.cpp
--- a/Engine/Plugins/VirtualProduction/RemoteControl/Source/RemoteControl/Private/RCVirtualPropertyContainer.cpp
+++ b/Engine/Plugins/VirtualProduction/RemoteControl/Source/RemoteControl/Private/RCVirtualPropertyContainer.cpp
@@ -133,7 +133,6 @@ URCVirtualPropertyInContainer* URCVirtualPropertyContainerBase::DuplicateVirtual
 bool URCVirtualPropertyContainerBase::RemoveProperty(const FName& InPropertyName)
 {
 	Bag.RemovePropertyByName(InPropertyName);
-	ControllerLabelToIdCache.Remove(InPropertyName);
 
 	for (TSet<TObjectPtr<URCVirtualPropertyBase>>::TIterator PropertiesIt = VirtualProperties.CreateIterator(); PropertiesIt; ++PropertiesIt)
 	{
@@ -141,6 +140,7 @@ bool URCVirtualPropertyContainerBase::RemoveProperty(const FName& InPropertyName
 		{
 			if (VirtualProperty->PropertyName == InPropertyName)
 			{
+				ControllerLabelToIdCache.Remove(VirtualProperty->DisplayName);
 				PropertiesIt.RemoveCurrent();
 				return true;
 			}
