- Fixing PVS warning in remote control about useless condition

diff --git a/Engine/Plugins/VirtualProduction/RemoteControl/Source/RemoteControlUI/Private/UI/SRemoteControlPanel.cpp b/Engine/Plugins/VirtualProduction/RemoteControl/Source/RemoteControlUI/Private/UI/SRemoteControlPanel.cpp
--- a/Engine/Plugins/VirtualProduction/RemoteControl/Source/RemoteControlUI/Private/UI/SRemoteControlPanel.cpp
+++ b/Engine/Plugins/VirtualProduction/RemoteControl/Source/RemoteControlUI/Private/UI/SRemoteControlPanel.cpp
@@ -2368,7 +2368,7 @@ void SRemoteControlPanel::ToggleLogicEditor_Execute()
 	Settings->PostEditChange();
 	Settings->SaveConfig();	
 
-	if (PanelDrawer.IsValid() && (ActivePanel != ERCPanels::RCP_Properties || ActivePanel == ERCPanels::RCP_None))
+	if (PanelDrawer.IsValid() && (ActivePanel != ERCPanels::RCP_Properties))
 	{
 		TSharedRef<FRCPanelDrawerArgs> PropertiesPanel = RegisteredDrawers.FindChecked(ERCPanels::RCP_Properties);
 
