[RemoteControl]
UI change for ActionPanel and Behaviours

diff --git a/Engine/Plugins/VirtualProduction/RemoteControl/Source/RemoteControlUI/Private/UI/Action/SRCActionPanel.cpp b/Engine/Plugins/VirtualProduction/RemoteControl/Source/RemoteControlUI/Private/UI/Action/SRCActionPanel.cpp
--- a/Engine/Plugins/VirtualProduction/RemoteControl/Source/RemoteControlUI/Private/UI/Action/SRCActionPanel.cpp
+++ b/Engine/Plugins/VirtualProduction/RemoteControl/Source/RemoteControlUI/Private/UI/Action/SRCActionPanel.cpp
@@ -177,24 +177,30 @@ void SRCActionPanel::UpdateWrappedWidget(TSharedPtr<FRCBehaviourModel> InBehavio
 		ActionDockPanel->AddHeaderToolbarItem(EToolbar::Right, AddAllSelectedActionsButton);
 		ActionDockPanel->AddHeaderToolbarItem(EToolbar::Right, AddAllActionsButton);
 
-		// Header Dock Panel
-		TSharedPtr<SRCMinorPanel> BehaviourDetailsPanel = SNew(SRCMinorPanel)
-			.EnableHeader(false)
-			[
-				SAssignNew(BehaviourDetailsWidget, SRCBehaviourDetails, SharedThis(this), InBehaviourItem.ToSharedRef())
-			];
-
 		TSharedRef<SRCMajorPanel> ActionsPanel = SNew(SRCMajorPanel)
 			.EnableFooter(false)
 			.EnableHeader(false)
 			.ChildOrientation(Orient_Vertical);
 
-		// Panel size of zero forces use of "SizeToContent" ensuring that each Behaviour only takes up as much space as necessary
-		ActionsPanel->AddPanel(BehaviourDetailsPanel.ToSharedRef(), 0.f);
+		if (InBehaviourItem->HasBehaviourDetailsWidget())
+		{
+			// Header Dock Panel
+			TSharedPtr<SRCMinorPanel> BehaviourDetailsPanel = SNew(SRCMinorPanel)
+				.EnableHeader(false)
+				[
+					SAssignNew(BehaviourDetailsWidget, SRCBehaviourDetails, SharedThis(this), InBehaviourItem.ToSharedRef())
+				];
 
+			// Panel size of zero forces use of "SizeToContent" ensuring that each Behaviour only takes up as much space as necessary
+			ActionsPanel->AddPanel(BehaviourDetailsPanel.ToSharedRef(), 0.f);
+		}
 		ActionsPanel->AddPanel(ActionDockPanel.ToSharedRef(), 0.5f);
 
 		WrappedBoxWidget->SetContent(ActionsPanel);
+
+		const bool bIsBehaviourEnabled = InBehaviourItem->IsBehaviourEnabled();
+		InBehaviourItem->RefreshIsBehaviourEnabled(bIsBehaviourEnabled);
+		RefreshIsBehaviourEnabled(bIsBehaviourEnabled);
 	}
 	else
 	{
diff --git a/Engine/Plugins/VirtualProduction/RemoteControl/Source/RemoteControlUI/Private/UI/Action/SRCBehaviourDetails.cpp b/Engine/Plugins/VirtualProduction/RemoteControl/Source/RemoteControlUI/Private/UI/Action/SRCBehaviourDetails.cpp
--- a/Engine/Plugins/VirtualProduction/RemoteControl/Source/RemoteControlUI/Private/UI/Action/SRCBehaviourDetails.cpp
+++ b/Engine/Plugins/VirtualProduction/RemoteControl/Source/RemoteControlUI/Private/UI/Action/SRCBehaviourDetails.cpp
@@ -19,7 +19,6 @@
 #include "UI/SRemoteControlPanel.h"
 #include "Widgets/Input/SButton.h"
 #include "Widgets/Layout/SSpacer.h"
-#include "Widgets/Layout/SSeparator.h"
 
 #define LOCTEXT_NAMESPACE "SRCBehaviourDetails"
 
@@ -56,56 +55,6 @@ void SRCBehaviourDetails::Construct(const FArguments& InArgs, TSharedRef<SRCActi
 		[
 			SNew(SVerticalBox)
 
-			+ SVerticalBox::Slot()
-			.HAlign(HAlign_Fill)
-			.Padding(6.f)
-			.AutoHeight()
-			[
-				SNew(SHorizontalBox)
-
-				// Behaviour Name
-				+ SHorizontalBox::Slot()
-				.HAlign(HAlign_Left)
-				.VAlign(VAlign_Center)
-				.FillWidth(1.f)
-				.Padding(4.f)
-				[
-					SAssignNew(BehaviourTitleWidget, STextBlock)
-					.ColorAndOpacity(FLinearColor::White)
-					.Text(BehaviourDisplayName)
-					.TextStyle(&RCPanelStyle->SectionHeaderTextStyle)
-					.ToolTipText(BehaviourDescription)
-				]
-
-				// Toggle Behaviour Button
-				+ SHorizontalBox::Slot()
-				.VAlign(VAlign_Center)
-				.AutoWidth()
-				.Padding(4.f, 0.f)
-				[
-					SNew(SCheckBox)
-					.AddMetaData<FTagMetaData>(FTagMetaData(TEXT("Toggle Behavior")))
-					.ToolTipText(LOCTEXT("EditModeTooltip", "Enable/Disable this Behaviour.\nWhen a behaviour is disabled its Actions will not be processed when the Controller value changes"))
-					.HAlign(HAlign_Center)
-					.Style(&RCPanelStyle->ToggleButtonStyle)
-					.ForegroundColor(FSlateColor::UseForeground())
-					.IsChecked_Lambda([this]() { return BehaviourItemWeakPtr.IsValid() && BehaviourItemWeakPtr.Pin()->IsBehaviourEnabled() ? ECheckBoxState::Checked : ECheckBoxState::Unchecked; })
-					.OnCheckStateChanged(this, &SRCBehaviourDetails::OnToggleEnableBehaviour)
-				]
-			]
-			
-			// Border (separating Header and Behaviour specific details panel)
-			+ SVerticalBox::Slot()
-			.Padding(2.f, 4.f)
-			.AutoHeight()
-			[
-				SNew(SSeparator)
-				.SeparatorImage(FAppStyle::Get().GetBrush("Separator"))
-				.Thickness(2.f)
-				.Orientation(EOrientation::Orient_Horizontal)
-				.Visibility_Lambda([this]() { return this->BehaviourDetailsWidget.IsValid() ? EVisibility::Visible : EVisibility::Collapsed; })
-			]
-
 			// Behaviour Specific Details Panel
 			+ SVerticalBox::Slot()
 			.Padding(8.f, 4.f)
@@ -126,11 +75,6 @@ void SRCBehaviourDetails::Construct(const FArguments& InArgs, TSharedRef<SRCActi
 	RefreshIsBehaviourEnabled(bIsChecked);
 }
 
-void SRCBehaviourDetails::OnToggleEnableBehaviour(ECheckBoxState State)
-{
-	SetIsBehaviourEnabled(State == ECheckBoxState::Checked);
-}
-
 void SRCBehaviourDetails::SetIsBehaviourEnabled(const bool bIsEnabled)
 {
 	if (TSharedPtr<FRCBehaviourModel> BehaviourItem = BehaviourItemWeakPtr.Pin())
diff --git a/Engine/Plugins/VirtualProduction/RemoteControl/Source/RemoteControlUI/Private/UI/Action/SRCBehaviourDetails.h b/Engine/Plugins/VirtualProduction/RemoteControl/Source/RemoteControlUI/Private/UI/Action/SRCBehaviourDetails.h
--- a/Engine/Plugins/VirtualProduction/RemoteControl/Source/RemoteControlUI/Private/UI/Action/SRCBehaviourDetails.h
+++ b/Engine/Plugins/VirtualProduction/RemoteControl/Source/RemoteControlUI/Private/UI/Action/SRCBehaviourDetails.h
@@ -31,9 +31,6 @@ private:
 	/** Refreshes the UI widgets enabled state depending on whether the parent behaviour is currently enabled */
 	void RefreshIsBehaviourEnabled(const bool bIsEnabled);
 
-	/** Handles click event for toggling Enable/Disable behaviour*/
-	void OnToggleEnableBehaviour(ECheckBoxState State);
-
 	/** The parent Action panel UI widget holding this Header*/
 	TWeakPtr<SRCActionPanel> ActionPanelWeakPtr;
 	
diff --git a/Engine/Plugins/VirtualProduction/RemoteControl/Source/RemoteControlUI/Private/UI/Behaviour/Builtin/Bind/RCBehaviourBindModel.cpp b/Engine/Plugins/VirtualProduction/RemoteControl/Source/RemoteControlUI/Private/UI/Behaviour/Builtin/Bind/RCBehaviourBindModel.cpp
--- a/Engine/Plugins/VirtualProduction/RemoteControl/Source/RemoteControlUI/Private/UI/Behaviour/Builtin/Bind/RCBehaviourBindModel.cpp
+++ b/Engine/Plugins/VirtualProduction/RemoteControl/Source/RemoteControlUI/Private/UI/Behaviour/Builtin/Bind/RCBehaviourBindModel.cpp
@@ -16,6 +16,11 @@ FRCBehaviourBindModel::FRCBehaviourBindModel(URCBehaviourBind* BindBehaviour, co
 {
 }
 
+bool FRCBehaviourBindModel::HasBehaviourDetailsWidget()
+{
+	return true;
+}
+
 TSharedRef<SWidget> FRCBehaviourBindModel::GetBehaviourDetailsWidget()
 {
 	return SAssignNew(BehaviourDetailsWidget, SRCBehaviourBind, SharedThis(this));
diff --git a/Engine/Plugins/VirtualProduction/RemoteControl/Source/RemoteControlUI/Private/UI/Behaviour/Builtin/Bind/RCBehaviourBindModel.h b/Engine/Plugins/VirtualProduction/RemoteControl/Source/RemoteControlUI/Private/UI/Behaviour/Builtin/Bind/RCBehaviourBindModel.h
--- a/Engine/Plugins/VirtualProduction/RemoteControl/Source/RemoteControlUI/Private/UI/Behaviour/Builtin/Bind/RCBehaviourBindModel.h
+++ b/Engine/Plugins/VirtualProduction/RemoteControl/Source/RemoteControlUI/Private/UI/Behaviour/Builtin/Bind/RCBehaviourBindModel.h
@@ -20,6 +20,9 @@ class FRCBehaviourBindModel : public FRCBehaviourModel
 public:
 	FRCBehaviourBindModel(URCBehaviourBind* BindBehaviour, const TSharedPtr<SRemoteControlPanel> InRemoteControlPanel);
 
+	/** Returns true if this behaviour have a details widget or false if not*/
+	virtual bool HasBehaviourDetailsWidget() override;
+
 	/** Builds a Behaviour specific details widget for Bind Behaviour*/
 	virtual TSharedRef<SWidget> GetBehaviourDetailsWidget() override;
 
diff --git a/Engine/Plugins/VirtualProduction/RemoteControl/Source/RemoteControlUI/Private/UI/Behaviour/Builtin/Conditional/RCBehaviourConditionalModel.cpp b/Engine/Plugins/VirtualProduction/RemoteControl/Source/RemoteControlUI/Private/UI/Behaviour/Builtin/Conditional/RCBehaviourConditionalModel.cpp
--- a/Engine/Plugins/VirtualProduction/RemoteControl/Source/RemoteControlUI/Private/UI/Behaviour/Builtin/Conditional/RCBehaviourConditionalModel.cpp
+++ b/Engine/Plugins/VirtualProduction/RemoteControl/Source/RemoteControlUI/Private/UI/Behaviour/Builtin/Conditional/RCBehaviourConditionalModel.cpp
@@ -44,6 +44,11 @@ URCAction* FRCBehaviourConditionalModel::AddAction(const TSharedRef<const FRemot
 	return NewAction;
 }
 
+bool FRCBehaviourConditionalModel::HasBehaviourDetailsWidget()
+{
+	return true;
+}
+
 void FRCBehaviourConditionalModel::CreateComparandInputField()
 {
 	if (URCBehaviourConditional* ConditionalBehaviour = GetConditionalBehaviour())
diff --git a/Engine/Plugins/VirtualProduction/RemoteControl/Source/RemoteControlUI/Private/UI/Behaviour/Builtin/Conditional/RCBehaviourConditionalModel.h b/Engine/Plugins/VirtualProduction/RemoteControl/Source/RemoteControlUI/Private/UI/Behaviour/Builtin/Conditional/RCBehaviourConditionalModel.h
--- a/Engine/Plugins/VirtualProduction/RemoteControl/Source/RemoteControlUI/Private/UI/Behaviour/Builtin/Conditional/RCBehaviourConditionalModel.h
+++ b/Engine/Plugins/VirtualProduction/RemoteControl/Source/RemoteControlUI/Private/UI/Behaviour/Builtin/Conditional/RCBehaviourConditionalModel.h
@@ -30,6 +30,9 @@ public:
 	/** Add a Logic Action using a remote control field as input */
 	virtual URCAction* AddAction(const TSharedRef<const FRemoteControlField> InRemoteControlField);
 
+	/** Returns true if this behaviour have a details widget or false if not*/
+	virtual bool HasBehaviourDetailsWidget() override;
+
 	/** Builds a Behaviour specific details widget enabling the user to define Conditions*/
 	virtual TSharedRef<SWidget> GetBehaviourDetailsWidget() override;
 
diff --git a/Engine/Plugins/VirtualProduction/RemoteControl/Source/RemoteControlUI/Private/UI/Behaviour/Builtin/Conditional/SRCBehaviourConditional.cpp b/Engine/Plugins/VirtualProduction/RemoteControl/Source/RemoteControlUI/Private/UI/Behaviour/Builtin/Conditional/SRCBehaviourConditional.cpp
--- a/Engine/Plugins/VirtualProduction/RemoteControl/Source/RemoteControlUI/Private/UI/Behaviour/Builtin/Conditional/SRCBehaviourConditional.cpp
+++ b/Engine/Plugins/VirtualProduction/RemoteControl/Source/RemoteControlUI/Private/UI/Behaviour/Builtin/Conditional/SRCBehaviourConditional.cpp
@@ -55,53 +55,78 @@ void SRCBehaviourConditional::Construct(const FArguments& InArgs, TSharedRef<FRC
 
 	ChildSlot
 	[
-		SNew(SVerticalBox)
-		// Conditions Panel Title
-		+ SVerticalBox::Slot()
-		.HAlign(HAlign_Left)
-		.Padding(FMargin(3.f))
-		.AutoHeight()
+		SNew(SHorizontalBox)
+		// Conditions Panel
+		+ SHorizontalBox::Slot()
+		.HAlign(HAlign_Center)
+		.FillWidth(0.5f)
 		[
-			SNew(STextBlock)
-			.Text(LOCTEXT("ConditionsPanelTitle", "Condition"))
-			.Font(FRemoteControlPanelStyle::Get()->GetFontStyle("RemoteControlPanel.Actions.ValuePanelHeader"))
-		]
+			SNew(SVerticalBox)
+			// Conditions Panel Title
+			+ SVerticalBox::Slot()
+			.HAlign(HAlign_Center)
+			.Padding(FMargin(3.f))
+			.AutoHeight()
+			[
+				SNew(STextBlock)
+				.Text(LOCTEXT("ConditionsPanelTitle", "Condition"))
+				.Font(FRemoteControlPanelStyle::Get()->GetFontStyle("RemoteControlPanel.Actions.ValuePanelHeader"))
+			]
 
-		// Conditions List
-		+ SVerticalBox::Slot()
-		.HAlign(HAlign_Left)
-		.Padding(FMargin(3.f))
-		.AutoHeight()
-		[
-			SAssignNew(ListViewConditions, SListView<TSharedPtr<ERCBehaviourConditionType>>)
-			.ListItemsSource(&Conditions)
-			.Orientation(Orient_Horizontal) // Horizontal List
-			.OnSelectionChanged(this, &SRCBehaviourConditional::OnConditionsListSelectionChanged)
-			.OnGenerateRow(this, &SRCBehaviourConditional::OnGenerateWidgetForConditionsList)
-			// @todo: Add dark background for list as per Figma
+			// Conditions List
+			+ SVerticalBox::Slot()
+			.HAlign(HAlign_Center)
+			.Padding(FMargin(3.f))
+			.AutoHeight()
+			[
+				SAssignNew(ListViewConditions, SListView<TSharedPtr<ERCBehaviourConditionType>>)
+				.ListItemsSource(&Conditions)
+				.Orientation(Orient_Horizontal) // Horizontal List
+				.OnSelectionChanged(this, &SRCBehaviourConditional::OnConditionsListSelectionChanged)
+				.OnGenerateRow(this, &SRCBehaviourConditional::OnGenerateWidgetForConditionsList)
+				// @todo: Add dark background for list as per Figma
+			]
 		]
 
-		// Comparand Label
-		+ SVerticalBox::Slot()
-		.HAlign(HAlign_Left)
-		.Padding(FMargin(3.f, 6.f, 3.f, 3.f))
-		.AutoHeight()
+		// Vertical Separator
+		+ SHorizontalBox::Slot()
+		.Padding(FMargin(5.f, 0.f))
+		.AutoWidth()
 		[
-			SNew(STextBlock)
-			.Text(LOCTEXT("ConditionsPanelComparandLabel", "Input Value"))
-			.Font(FRemoteControlPanelStyle::Get()->GetFontStyle("RemoteControlPanel.Actions.ValuePanelHeader"))
+			SNew(SSeparator)
+			.Orientation(Orient_Vertical)
+			.SeparatorImage(FAppStyle::Get().GetBrush("Separator"))
+			.Thickness(2.5f)
 		]
 
-		// Comparand Value Widget
-		+ SVerticalBox::Slot()
-		.HAlign(HAlign_Left)
-		.Padding(FMargin(3.f))
-		.AutoHeight()
+		// Comparand Panel
+		+ SHorizontalBox::Slot()
+		.HAlign(HAlign_Center)
+		.FillWidth(0.5f)
 		[
-			SAssignNew(ComparandFieldBoxWidget, SBox)
+			SNew(SVerticalBox)
+			// Comparand Label
+			+ SVerticalBox::Slot()
+			.HAlign(HAlign_Center)
+			.Padding(FMargin(3.f))
+			.AutoHeight()
 			[
-				InBehaviourItem->GetComparandFieldWidget()
-			]			
+				SNew(STextBlock)
+				.Text(LOCTEXT("ConditionsPanelComparandLabel", "Input Value"))
+				.Font(FRemoteControlPanelStyle::Get()->GetFontStyle("RemoteControlPanel.Actions.ValuePanelHeader"))
+			]
+
+			// Comparand Value Widget
+			+ SVerticalBox::Slot()
+			.HAlign(HAlign_Center)
+			.Padding(FMargin(3.f))
+			.AutoHeight()
+			[
+				SAssignNew(ComparandFieldBoxWidget, SBox)
+				[
+					InBehaviourItem->GetComparandFieldWidget()
+				]			
+			]
 		]
 	];
 
diff --git a/Engine/Plugins/VirtualProduction/RemoteControl/Source/RemoteControlUI/Private/UI/Behaviour/Builtin/Path/RCBehaviourSetAssetByPathModel.cpp b/Engine/Plugins/VirtualProduction/RemoteControl/Source/RemoteControlUI/Private/UI/Behaviour/Builtin/Path/RCBehaviourSetAssetByPathModel.cpp
--- a/Engine/Plugins/VirtualProduction/RemoteControl/Source/RemoteControlUI/Private/UI/Behaviour/Builtin/Path/RCBehaviourSetAssetByPathModel.cpp
+++ b/Engine/Plugins/VirtualProduction/RemoteControl/Source/RemoteControlUI/Private/UI/Behaviour/Builtin/Path/RCBehaviourSetAssetByPathModel.cpp
@@ -82,6 +82,11 @@ FRCSetAssetByPathBehaviourModel::FRCSetAssetByPathBehaviourModel(URCSetAssetByPa
 	SelectorBox = GetSelectorWidget(InitialSelected);
 }
 
+bool FRCSetAssetByPathBehaviourModel::HasBehaviourDetailsWidget()
+{
+	return true;
+}
+
 TSharedRef<SWidget> FRCSetAssetByPathBehaviourModel::GetBehaviourDetailsWidget()
 {
 	return SNew(SRCBehaviourSetAssetByPath, SharedThis(this));
diff --git a/Engine/Plugins/VirtualProduction/RemoteControl/Source/RemoteControlUI/Private/UI/Behaviour/Builtin/Path/RCBehaviourSetAssetByPathModel.h b/Engine/Plugins/VirtualProduction/RemoteControl/Source/RemoteControlUI/Private/UI/Behaviour/Builtin/Path/RCBehaviourSetAssetByPathModel.h
--- a/Engine/Plugins/VirtualProduction/RemoteControl/Source/RemoteControlUI/Private/UI/Behaviour/Builtin/Path/RCBehaviourSetAssetByPathModel.h
+++ b/Engine/Plugins/VirtualProduction/RemoteControl/Source/RemoteControlUI/Private/UI/Behaviour/Builtin/Path/RCBehaviourSetAssetByPathModel.h
@@ -26,6 +26,9 @@ class FRCSetAssetByPathBehaviourModel : public FRCBehaviourModel
 public:
 	FRCSetAssetByPathBehaviourModel(URCSetAssetByPathBehaviour* SetAssetByPathBehaviour);
 
+	/** Returns true if this behaviour have a details widget or false if not*/
+	virtual bool HasBehaviourDetailsWidget() override;
+
 	/** Builds a Behaviour specific widget as required for the Set Asset By Path Behaviour */
 	virtual TSharedRef<SWidget> GetBehaviourDetailsWidget() override;
 
diff --git a/Engine/Plugins/VirtualProduction/RemoteControl/Source/RemoteControlUI/Private/UI/Behaviour/Builtin/RangeMap/RCBehaviourRangeMapModel.cpp b/Engine/Plugins/VirtualProduction/RemoteControl/Source/RemoteControlUI/Private/UI/Behaviour/Builtin/RangeMap/RCBehaviourRangeMapModel.cpp
--- a/Engine/Plugins/VirtualProduction/RemoteControl/Source/RemoteControlUI/Private/UI/Behaviour/Builtin/RangeMap/RCBehaviourRangeMapModel.cpp
+++ b/Engine/Plugins/VirtualProduction/RemoteControl/Source/RemoteControlUI/Private/UI/Behaviour/Builtin/RangeMap/RCBehaviourRangeMapModel.cpp
@@ -46,6 +46,11 @@ FRCRangeMapBehaviourModel::FRCRangeMapBehaviourModel(URCRangeMapBehaviour* Range
 	}
 }
 
+bool FRCRangeMapBehaviourModel::HasBehaviourDetailsWidget()
+{
+	return true;
+}
+
 TSharedRef<SWidget> FRCRangeMapBehaviourModel::GetBehaviourDetailsWidget()
 {
 	return SNew(SRCBehaviourRangeMap, SharedThis(this));
diff --git a/Engine/Plugins/VirtualProduction/RemoteControl/Source/RemoteControlUI/Private/UI/Behaviour/Builtin/RangeMap/RCBehaviourRangeMapModel.h b/Engine/Plugins/VirtualProduction/RemoteControl/Source/RemoteControlUI/Private/UI/Behaviour/Builtin/RangeMap/RCBehaviourRangeMapModel.h
--- a/Engine/Plugins/VirtualProduction/RemoteControl/Source/RemoteControlUI/Private/UI/Behaviour/Builtin/RangeMap/RCBehaviourRangeMapModel.h
+++ b/Engine/Plugins/VirtualProduction/RemoteControl/Source/RemoteControlUI/Private/UI/Behaviour/Builtin/RangeMap/RCBehaviourRangeMapModel.h
@@ -21,6 +21,9 @@ class FRCRangeMapBehaviourModel : public FRCBehaviourModel
 public:
 	FRCRangeMapBehaviourModel(URCRangeMapBehaviour* RangeMapBehaviour);
 
+	/** Returns true if this behaviour have a details widget or false if not*/
+	virtual bool HasBehaviourDetailsWidget() override;
+
 	/** Builds a Behaviour specific widget as required for Range Mapping behaviour*/
 	virtual TSharedRef<SWidget> GetBehaviourDetailsWidget() override;
 
diff --git a/Engine/Plugins/VirtualProduction/RemoteControl/Source/RemoteControlUI/Private/UI/Behaviour/RCBehaviourModel.cpp b/Engine/Plugins/VirtualProduction/RemoteControl/Source/RemoteControlUI/Private/UI/Behaviour/RCBehaviourModel.cpp
--- a/Engine/Plugins/VirtualProduction/RemoteControl/Source/RemoteControlUI/Private/UI/Behaviour/RCBehaviourModel.cpp
+++ b/Engine/Plugins/VirtualProduction/RemoteControl/Source/RemoteControlUI/Private/UI/Behaviour/RCBehaviourModel.cpp
@@ -26,10 +26,10 @@ FRCBehaviourModel::FRCBehaviourModel(URCBehaviour* InBehaviour
 
 	if (BehaviourWeakPtr.IsValid())
 	{
-		const FText BehaviorDisplayName = BehaviourWeakPtr->GetDisplayName();
+		const FText BehaviourDisplayName = BehaviourWeakPtr->GetDisplayName();
 		
 		SAssignNew(BehaviourTitleText, STextBlock)
-			.Text(BehaviorDisplayName)
+			.Text(BehaviourDisplayName)
 			.TextStyle(&RCPanelStyle->HeaderTextStyle);
 
 		RefreshIsBehaviourEnabled(BehaviourWeakPtr->bIsEnabled);
@@ -91,6 +91,11 @@ TSharedRef<SWidget> FRCBehaviourModel::GetWidget() const
 		];
 }
 
+bool FRCBehaviourModel::HasBehaviourDetailsWidget()
+{
+	return false;
+}
+
 TSharedRef<SWidget> FRCBehaviourModel::GetBehaviourDetailsWidget()
 {
 	return SNullWidget::NullWidget;
diff --git a/Engine/Plugins/VirtualProduction/RemoteControl/Source/RemoteControlUI/Private/UI/Behaviour/RCBehaviourModel.h b/Engine/Plugins/VirtualProduction/RemoteControl/Source/RemoteControlUI/Private/UI/Behaviour/RCBehaviourModel.h
--- a/Engine/Plugins/VirtualProduction/RemoteControl/Source/RemoteControlUI/Private/UI/Behaviour/RCBehaviourModel.h
+++ b/Engine/Plugins/VirtualProduction/RemoteControl/Source/RemoteControlUI/Private/UI/Behaviour/RCBehaviourModel.h
@@ -38,6 +38,9 @@ public:
 	*/
 	virtual TSharedRef<SWidget> GetWidget() const override;
 
+	/** Returns true if this behaviour have a details widget or false if not*/
+	virtual bool HasBehaviourDetailsWidget();
+
 	/** Builds a Behaviour specific widget that child Behaviour classes can implement as required*/
 	virtual TSharedRef<SWidget> GetBehaviourDetailsWidget();
 
diff --git a/Engine/Plugins/VirtualProduction/RemoteControl/Source/RemoteControlUI/Private/UI/Behaviour/SRCBehaviourPanelList.cpp b/Engine/Plugins/VirtualProduction/RemoteControl/Source/RemoteControlUI/Private/UI/Behaviour/SRCBehaviourPanelList.cpp
--- a/Engine/Plugins/VirtualProduction/RemoteControl/Source/RemoteControlUI/Private/UI/Behaviour/SRCBehaviourPanelList.cpp
+++ b/Engine/Plugins/VirtualProduction/RemoteControl/Source/RemoteControlUI/Private/UI/Behaviour/SRCBehaviourPanelList.cpp
@@ -131,7 +131,28 @@ void SRCBehaviourPanelList::SetIsBehaviourEnabled(const bool bIsEnabled)
 	{
 		if (const TSharedPtr<SRCActionPanel> ActionPanel = RemoteControlPanel->GetLogicActionPanel())
 		{
-			ActionPanel->SetIsBehaviourEnabled(bIsEnabled);
+			ActionPanel->RefreshIsBehaviourEnabled(bIsEnabled);
+		}
+	}
+}
+
+void SRCBehaviourPanelList::SetIsBehaviourEnabled(const TSharedPtr<FRCBehaviourModel>& InBehaviourModel, const bool bIsEnabled)
+{
+	// Disable the behaviour
+	if (InBehaviourModel.IsValid())
+	{
+		InBehaviourModel->SetIsBehaviourEnabled(bIsEnabled);
+	}
+
+	// Disable the action panel
+	if (const TSharedPtr<SRemoteControlPanel> RemoteControlPanel = GetRemoteControlPanel())
+	{
+		if (const TSharedPtr<SRCActionPanel> ActionPanel = RemoteControlPanel->GetLogicActionPanel())
+		{
+			if (InBehaviourModel == ActionPanel->GetSelectedBehaviourItem())
+			{
+				ActionPanel->RefreshIsBehaviourEnabled(bIsEnabled);
+			}
 		}
 	}
 }
@@ -193,10 +214,40 @@ TSharedRef<ITableRow> SRCBehaviourPanelList::OnGenerateWidgetForList(TSharedPtr<
 	return SNew(STableRow<TSharedPtr<FString>>, OwnerTable)
 		.Style(&RCPanelStyle->TableRowStyle)
 		[
-			InItem->GetWidget()
+			SNew(SHorizontalBox)
+			+ SHorizontalBox::Slot()
+			[
+				InItem->GetWidget()
+			]
+
+			// Toggle Behaviour Button
+			+ SHorizontalBox::Slot()
+			.VAlign(VAlign_Center)
+			.AutoWidth()
+			.Padding(4.f, 0.f)
+			[
+				SNew(SCheckBox)
+				.AddMetaData<FTagMetaData>(FTagMetaData(TEXT("Toggle Behaviour")))
+				.ToolTipText(LOCTEXT("EditModeTooltip", "Enable/Disable this Behaviour.\nWhen a behaviour is disabled its Actions will not be processed when the Controller value changes"))
+				.HAlign(HAlign_Center)
+				.ForegroundColor(FLinearColor::White)
+				.Style(&RCPanelStyle->ToggleButtonStyle)
+				.IsChecked(this, &SRCBehaviourPanelList::IsBehaviourChecked, InItem)
+				.OnCheckStateChanged(this, &SRCBehaviourPanelList::OnToggleEnableBehaviour, InItem)
+			]
 		];
 }
 
+ECheckBoxState SRCBehaviourPanelList::IsBehaviourChecked(const TSharedPtr<FRCBehaviourModel> InBehaviourModel) const
+{
+	return InBehaviourModel.IsValid() && InBehaviourModel->IsBehaviourEnabled() ? ECheckBoxState::Checked : ECheckBoxState::Unchecked;
+}
+
+void SRCBehaviourPanelList::OnToggleEnableBehaviour(ECheckBoxState State, const TSharedPtr<FRCBehaviourModel> InBehaviourModel)
+{
+	SetIsBehaviourEnabled(InBehaviourModel, State == ECheckBoxState::Checked);
+}
+
 void SRCBehaviourPanelList::OnTreeSelectionChanged(TSharedPtr<FRCBehaviourModel> InItem, ESelectInfo::Type)
 {
 	if (const TSharedPtr<SRemoteControlPanel> RemoteControlPanel = BehaviourPanelWeakPtr.Pin()->GetRemoteControlPanel())
diff --git a/Engine/Plugins/VirtualProduction/RemoteControl/Source/RemoteControlUI/Private/UI/Behaviour/SRCBehaviourPanelList.h b/Engine/Plugins/VirtualProduction/RemoteControl/Source/RemoteControlUI/Private/UI/Behaviour/SRCBehaviourPanelList.h
--- a/Engine/Plugins/VirtualProduction/RemoteControl/Source/RemoteControlUI/Private/UI/Behaviour/SRCBehaviourPanelList.h
+++ b/Engine/Plugins/VirtualProduction/RemoteControl/Source/RemoteControlUI/Private/UI/Behaviour/SRCBehaviourPanelList.h
@@ -2,6 +2,7 @@
 
 #pragma once
 
+#include "Styling/SlateTypes.h"
 #include "UI/BaseLogicUI/SRCLogicPanelListBase.h"
 
 struct FRCPanelStyle;
@@ -69,11 +70,20 @@ private:
 	/** Enables or Disables the currently selected behaviour */
 	void SetIsBehaviourEnabled(const bool bIsEnabled);
 
+	/** Enables or Disables the passed behaviour */
+	void SetIsBehaviourEnabled(const TSharedPtr<FRCBehaviourModel>& InBehaviourModel, const bool bIsEnabled);
+
 	void AddBehaviourToList(URCBehaviour* InBehaviour);
 
 	/** OnGenerateRow delegate for the Behaviours List View */
 	TSharedRef<ITableRow> OnGenerateWidgetForList( TSharedPtr<FRCBehaviourModel> InItem, const TSharedRef<STableViewBase>& OwnerTable );
-	
+
+	/** Return whether or not the given behaviour is checked */
+	ECheckBoxState IsBehaviourChecked(const TSharedPtr<FRCBehaviourModel> InBehaviourModel) const;
+
+	/** Executed when toggling the behaviour state */
+	void OnToggleEnableBehaviour(ECheckBoxState State, TSharedPtr<FRCBehaviourModel> InBehaviourModel);
+
 	/** OnSelectionChanged delegate for Behaviours List View */
 	void OnTreeSelectionChanged(TSharedPtr<FRCBehaviourModel> InItem , ESelectInfo::Type);
 
