[RemoteControl]
Adding support for FVector2D Controller

diff --git a/Engine/Plugins/VirtualProduction/RemoteControl/Config/BaseRemoteControl.ini b/Engine/Plugins/VirtualProduction/RemoteControl/Config/BaseRemoteControl.ini
--- a/Engine/Plugins/VirtualProduction/RemoteControl/Config/BaseRemoteControl.ini
+++ b/Engine/Plugins/VirtualProduction/RemoteControl/Config/BaseRemoteControl.ini
@@ -12,6 +12,7 @@
 +SupportedControllerTypes=EPropertyBagPropertyType::Struct
 +SupportedControllerTypes=EPropertyBagPropertyType::Object
 +SupportedControllerStructTypes=Vector
++SupportedControllerStructTypes=Vector2D
 +SupportedControllerStructTypes=Color
 +SupportedControllerStructTypes=Rotator
 +SupportedControllerObjectClassPaths=/Script/Engine.StaticMesh
diff --git a/Engine/Plugins/VirtualProduction/RemoteControl/Source/RemoteControl/Private/RCVirtualProperty.cpp b/Engine/Plugins/VirtualProduction/RemoteControl/Source/RemoteControl/Private/RCVirtualProperty.cpp
--- a/Engine/Plugins/VirtualProduction/RemoteControl/Source/RemoteControl/Private/RCVirtualProperty.cpp
+++ b/Engine/Plugins/VirtualProduction/RemoteControl/Source/RemoteControl/Private/RCVirtualProperty.cpp
@@ -188,6 +188,17 @@ bool URCVirtualPropertyBase::IsVectorType() const
 	return false;
 }
 
+bool URCVirtualPropertyBase::IsVector2DType() const
+{
+	const FProperty* Property = GetProperty();
+	if (const FStructProperty* StructProperty = CastField<FStructProperty>(Property))
+	{
+		return StructProperty->Struct == TBaseStructure<FVector2D>::Get();
+	}
+
+	return false;
+}
+
 bool URCVirtualPropertyBase::IsColorType() const
 {
 	const FProperty* Property = GetProperty();
@@ -523,6 +534,11 @@ bool URCVirtualPropertyBase::GetValueVector(FVector& OutVector) const
 	return FRCVirtualPropertyCastHelpers::GetStructValue<FVector>(this, TBaseStructure<FVector>::Get(), &OutVector);
 }
 
+bool URCVirtualPropertyBase::GetValueVector2D(FVector2D& OutVector2D) const
+{
+	return FRCVirtualPropertyCastHelpers::GetStructValue<FVector2D>(this, TBaseStructure<FVector2D>::Get(), &OutVector2D);
+}
+
 bool URCVirtualPropertyBase::GetValueRotator(FRotator& OutRotator) const
 {
 	return FRCVirtualPropertyCastHelpers::GetStructValue<FRotator>(this, TBaseStructure<FRotator>::Get(), &OutRotator);
@@ -612,6 +628,13 @@ FString URCVirtualPropertyBase::GetDisplayValueAsString() const
 
 		return Vector.ToString();
 	}
+	else if (IsVector2DType())
+	{
+		FVector2D Vector2D;
+		GetValueVector2D(Vector2D);
+
+		return Vector2D.ToString();
+	}
 	else if (IsRotatorType())
 	{
 		FRotator Rotator;
@@ -729,6 +752,11 @@ bool URCVirtualPropertyBase::SetValueVector(const FVector& InVector)
 	return FRCVirtualPropertyCastHelpers::SetStructValue<FVector>(this, TBaseStructure<FVector>::Get(), InVector);
 }
 
+bool URCVirtualPropertyBase::SetValueVector2D(const FVector2D& InVector2D)
+{
+	return FRCVirtualPropertyCastHelpers::SetStructValue<FVector2D>(this, TBaseStructure<FVector2D>::Get(), InVector2D);
+}
+
 bool URCVirtualPropertyBase::SetValueRotator(const FRotator& InRotator)
 {
 	return FRCVirtualPropertyCastHelpers::SetStructValue<FRotator>(this, TBaseStructure<FRotator>::Get(), InRotator);
diff --git a/Engine/Plugins/VirtualProduction/RemoteControl/Source/RemoteControl/Private/RemoteControlField.cpp b/Engine/Plugins/VirtualProduction/RemoteControl/Source/RemoteControl/Private/RemoteControlField.cpp
--- a/Engine/Plugins/VirtualProduction/RemoteControl/Source/RemoteControl/Private/RemoteControlField.cpp
+++ b/Engine/Plugins/VirtualProduction/RemoteControl/Source/RemoteControl/Private/RemoteControlField.cpp
@@ -455,6 +455,7 @@ void FRemoteControlProperty::InitializeMetadata()
 		{
 			UStruct* Struct = CastFieldChecked<FStructProperty>(Property)->Struct;
 			return Struct->IsChildOf(TBaseStructure<FVector>::Get())
+				|| Struct->IsChildOf(TBaseStructure<FVector2D>::Get())
 				|| Struct->IsChildOf(TBaseStructure<FRotator>::Get());
 		}
 
diff --git a/Engine/Plugins/VirtualProduction/RemoteControl/Source/RemoteControl/Public/RCVirtualProperty.h b/Engine/Plugins/VirtualProduction/RemoteControl/Source/RemoteControl/Public/RCVirtualProperty.h
--- a/Engine/Plugins/VirtualProduction/RemoteControl/Source/RemoteControl/Public/RCVirtualProperty.h
+++ b/Engine/Plugins/VirtualProduction/RemoteControl/Source/RemoteControl/Public/RCVirtualProperty.h
@@ -101,6 +101,9 @@ public:
 	/** Whether this Virtual Property represents an FVector*/
 	bool IsVectorType() const;
 
+	/** Whether this Virtual Property represents an FVector2D*/
+	bool IsVector2DType() const;
+
 	/** Whether this Virtual Property represents an FColor*/
 	bool IsColorType() const;
 
@@ -213,6 +216,10 @@ public:
 	UFUNCTION(BlueprintCallable, Category = "Remote Control Behaviour")
 	bool GetValueVector(FVector& OutVector) const;
 
+	/** Get Vector2D value from Virtual Property */
+	UFUNCTION(BlueprintCallable, Category = "Remote Control Behaviour")
+	bool GetValueVector2D(FVector2D& OutVector2D) const;
+
 	/** Get Rotator value from Virtual Property */
 	UFUNCTION(BlueprintCallable, Category = "Remote Control Behaviour")
 	bool GetValueRotator(FRotator& OutRotator) const;
@@ -295,6 +302,10 @@ public:
 	UFUNCTION(BlueprintCallable, Category = "Remote Control Behaviour")
 	bool SetValueVector(const FVector& InVector);
 
+	/** Set Vector2D value from Virtual Property */
+	UFUNCTION(BlueprintCallable, Category = "Remote Control Behaviour")
+	bool SetValueVector2D(const FVector2D& InVector2D);
+
 	/** Set Rotator value from Virtual Property */
 	UFUNCTION(BlueprintCallable, Category = "Remote Control Behaviour")
 	bool SetValueRotator(const FRotator& InRotator);
diff --git a/Engine/Plugins/VirtualProduction/RemoteControl/Source/RemoteControlLogic/Private/Action/Bind/RCPropertyBindAction.cpp b/Engine/Plugins/VirtualProduction/RemoteControl/Source/RemoteControlLogic/Private/Action/Bind/RCPropertyBindAction.cpp
--- a/Engine/Plugins/VirtualProduction/RemoteControl/Source/RemoteControlLogic/Private/Action/Bind/RCPropertyBindAction.cpp
+++ b/Engine/Plugins/VirtualProduction/RemoteControl/Source/RemoteControlLogic/Private/Action/Bind/RCPropertyBindAction.cpp
@@ -33,6 +33,14 @@ static void SetStructPropertyFromController(const FProperty* RemoteControlProper
 
 				RemoteControlHandle->SetValue(VectorValue);
 			}
+			// FVector2D
+			else if (TargetStructProperty->Struct == TBaseStructure<FVector2D>::Get())
+			{
+				FVector2D Vector2DValue;
+				Controller->GetValueVector2D(Vector2DValue);
+
+				RemoteControlHandle->SetValue(Vector2DValue);
+			}
 			// FColor
 			else if (TargetStructProperty->Struct == TBaseStructure<FColor>::Get())
 			{
diff --git a/Engine/Plugins/VirtualProduction/RemoteControl/Source/RemoteControlLogic/Private/Behaviour/Builtin/Bind/RCBehaviourBind.cpp b/Engine/Plugins/VirtualProduction/RemoteControl/Source/RemoteControlLogic/Private/Behaviour/Builtin/Bind/RCBehaviourBind.cpp
--- a/Engine/Plugins/VirtualProduction/RemoteControl/Source/RemoteControlLogic/Private/Behaviour/Builtin/Bind/RCBehaviourBind.cpp
+++ b/Engine/Plugins/VirtualProduction/RemoteControl/Source/RemoteControlLogic/Private/Behaviour/Builtin/Bind/RCBehaviourBind.cpp
@@ -470,6 +470,12 @@ bool URCBehaviourBind::CopyPropertyValueToController(URCController* InController
 			PropertyHandle->GetValue(VectorValue);
 			InController->SetValueVector(VectorValue);
 		}
+		if (ControllerAsStructProperty->Struct == TBaseStructure<FVector2D>::Get())
+		{
+			FVector2D Vector2DValue;
+			PropertyHandle->GetValue(Vector2DValue);
+			InController->SetValueVector2D(Vector2DValue);
+		}
 		else if (ControllerAsStructProperty->Struct == TBaseStructure<FRotator>::Get())
 		{
 			FRotator RotatorValue;
diff --git a/Engine/Plugins/VirtualProduction/RemoteControl/Source/RemoteControlLogic/Private/Behaviour/Builtin/RCBehaviourOnValueChangedNode.cpp b/Engine/Plugins/VirtualProduction/RemoteControl/Source/RemoteControlLogic/Private/Behaviour/Builtin/RCBehaviourOnValueChangedNode.cpp
--- a/Engine/Plugins/VirtualProduction/RemoteControl/Source/RemoteControlLogic/Private/Behaviour/Builtin/RCBehaviourOnValueChangedNode.cpp
+++ b/Engine/Plugins/VirtualProduction/RemoteControl/Source/RemoteControlLogic/Private/Behaviour/Builtin/RCBehaviourOnValueChangedNode.cpp
@@ -1,4 +1,4 @@
-﻿// Copyright Epic Games, Inc. All Rights Reserved.
+// Copyright Epic Games, Inc. All Rights Reserved.
 
 #include "Behaviour/Builtin/RCBehaviourOnValueChangedNode.h"
 
@@ -29,6 +29,7 @@ bool URCBehaviourOnValueChangedNode::IsSupported_Implementation(URCBehaviour* In
 		{ EPropertyBagPropertyType::String, nullptr },
 		{ EPropertyBagPropertyType::Text, nullptr },
 		{ EPropertyBagPropertyType::Struct, TBaseStructure<FVector>::Get() },
+		{ EPropertyBagPropertyType::Struct, TBaseStructure<FVector2D>::Get() },
 		{ EPropertyBagPropertyType::Struct, TBaseStructure<FColor>::Get() },
 		{ EPropertyBagPropertyType::Struct, TBaseStructure<FRotator>::Get() }
 	};
diff --git a/Engine/Plugins/VirtualProduction/RemoteControl/Source/RemoteControlLogic/Private/Controller/RCControllerContainer.cpp b/Engine/Plugins/VirtualProduction/RemoteControl/Source/RemoteControlLogic/Private/Controller/RCControllerContainer.cpp
--- a/Engine/Plugins/VirtualProduction/RemoteControl/Source/RemoteControlLogic/Private/Controller/RCControllerContainer.cpp
+++ b/Engine/Plugins/VirtualProduction/RemoteControl/Source/RemoteControlLogic/Private/Controller/RCControllerContainer.cpp
@@ -1,4 +1,4 @@
-﻿// Copyright Epic Games, Inc. All Rights Reserved.
+// Copyright Epic Games, Inc. All Rights Reserved.
 
 #include "Controller/RCControllerContainer.h"
 
@@ -29,7 +29,7 @@ URCVirtualPropertyInContainer* URCControllerContainer::AddProperty(const FName&
 	// Vector Controllers
 	if (InValueType == EPropertyBagPropertyType::Struct)
 	{
-		if (InValueTypeObject == TBaseStructure<FVector>::Get())
+		if (InValueTypeObject == TBaseStructure<FVector>::Get() || InValueTypeObject == TBaseStructure<FVector2D>::Get())
 		{
 			MetaData.Add(FPropertyBagPropertyDescMetaData(FName("Delta"), FString::Printf(TEXT("%f"), VectorSliderDelta)));
 			MetaData.Add(FPropertyBagPropertyDescMetaData(FName("LinearDeltaSensitivity"), FString::Printf(TEXT("%f"), VectorLinearDeltaSensitivity)));
diff --git a/Engine/Plugins/VirtualProduction/RemoteControl/Source/RemoteControlUI/Private/UI/Controller/SRCControllerPanel.cpp b/Engine/Plugins/VirtualProduction/RemoteControl/Source/RemoteControlUI/Private/UI/Controller/SRCControllerPanel.cpp
--- a/Engine/Plugins/VirtualProduction/RemoteControl/Source/RemoteControlUI/Private/UI/Controller/SRCControllerPanel.cpp
+++ b/Engine/Plugins/VirtualProduction/RemoteControl/Source/RemoteControlUI/Private/UI/Controller/SRCControllerPanel.cpp
@@ -108,6 +108,10 @@ static UObject* GetBaseStructForType(FName StructType)
 	{
 		return TBaseStructure<FVector>::Get();
 	}
+	else if (StructType == NAME_Vector2D)
+	{
+		return TBaseStructure<FVector2d>::Get();
+	}
 	else if (StructType == NAME_Color)
 	{
 		return TBaseStructure<FColor>::Get();
diff --git a/Engine/Plugins/VirtualProduction/RemoteControl/Source/RemoteControlUI/Private/UI/Controller/SRCControllerPanelList.cpp b/Engine/Plugins/VirtualProduction/RemoteControl/Source/RemoteControlUI/Private/UI/Controller/SRCControllerPanelList.cpp
--- a/Engine/Plugins/VirtualProduction/RemoteControl/Source/RemoteControlUI/Private/UI/Controller/SRCControllerPanelList.cpp
+++ b/Engine/Plugins/VirtualProduction/RemoteControl/Source/RemoteControlUI/Private/UI/Controller/SRCControllerPanelList.cpp
@@ -64,6 +64,7 @@ namespace UE::RCControllerPanelList
 	{
 		static const TSet<UScriptStruct*> SupportedStructs = {
 			TBaseStructure<FVector>::Get(),
+			TBaseStructure<FVector2D>::Get(),
 			TBaseStructure<FRotator>::Get(),
 			TBaseStructure<FColor>::Get()
 		};
