[Remote Control] Fix another edge case with unregistering embedded presets on a "Save As".

#rnx

diff --git a/Engine/Plugins/VirtualProduction/RemoteControl/Source/RemoteControl/Private/RemoteControlModule.cpp b/Engine/Plugins/VirtualProduction/RemoteControl/Source/RemoteControl/Private/RemoteControlModule.cpp
--- a/Engine/Plugins/VirtualProduction/RemoteControl/Source/RemoteControl/Private/RemoteControlModule.cpp
+++ b/Engine/Plugins/VirtualProduction/RemoteControl/Source/RemoteControl/Private/RemoteControlModule.cpp
@@ -917,30 +917,31 @@ void FRemoteControlModule::UnregisterEmbeddedPreset(URemoteControlPreset* Preset
 	// Remark: the PresetName is unreliable. When unregistering during the garbage collection,
 	// the name can be "None". In that case we can recover the name from CachedPresetNamesById
 	// by using the PresetId.
-	FName PresetName = Preset->GetPresetName();
+	// The other case is when the package has been renamed (save level as), in that case
+	// the PresetName is the renamed one, but both CachedPresetNamesById and EmbeddedPresets
+	// still have the old name.
+	const FName PresetName = Preset->GetPresetName();
+	bool bPresetRemoved = false;
 
+	if (PresetName != NAME_None)
+	{
+		// Attempt to remove with current name (may not work if the package has been renamed).
+		bPresetRemoved = (EmbeddedPresets.Remove(PresetName) > 0) ? true : false;
+	}
+	
 	const FGuid PresetId = Preset->GetPresetId();
 	
 	if (PresetId.IsValid())
 	{
-		if (const FName* FoundPresetName = CachedPresetNamesById.Find(PresetId))
+		if (!bPresetRemoved)
 		{
-			// Recover the name from the id.
-			if (PresetName == NAME_None)
-			{
-				PresetName = *FoundPresetName;
-			}
-			
-			if (PresetName == *FoundPresetName)
+			if (const FName* FoundPresetName = CachedPresetNamesById.Find(PresetId))
 			{
-				CachedPresetNamesById.Remove(PresetId);
+				// Attempt to remove from recovered the name from the id.
+				EmbeddedPresets.Remove(*FoundPresetName);
 			}
 		}
-	}
-
-	if (PresetName != NAME_None)
-	{
-		EmbeddedPresets.Remove(PresetName);
+		CachedPresetNamesById.Remove(PresetId);
 	}
 }
 
