[RemoteControl]
Sub-Property Id

diff --git a/Engine/Plugins/VirtualProduction/RemoteControl/Source/RemoteControl/Private/PropertyIdHandler/BasePropertyIdHandler.cpp b/Engine/Plugins/VirtualProduction/RemoteControl/Source/RemoteControl/Private/PropertyIdHandler/BasePropertyIdHandler.cpp
--- a/Engine/Plugins/VirtualProduction/RemoteControl/Source/RemoteControl/Private/PropertyIdHandler/BasePropertyIdHandler.cpp
+++ b/Engine/Plugins/VirtualProduction/RemoteControl/Source/RemoteControl/Private/PropertyIdHandler/BasePropertyIdHandler.cpp
@@ -6,7 +6,7 @@ bool FBasePropertyIdHandler::IsPropertySupported(const FProperty* InProperty) co
 {
 	if (InProperty)
 	{
-		const FProperty* Property = TryGetPropertyInContainerOrSelf(InProperty);
+		const FProperty* Property = GetPropertyInsideContainer(InProperty);
 		return !(Property->GetClass()->IsChildOf(FStructProperty::StaticClass()) ||
 				 Property->GetClass()->IsChildOf(FObjectProperty::StaticClass()) ||
 				 Property->GetClass()->IsChildOf(FEnumProperty::StaticClass())	  );
@@ -19,7 +19,7 @@ EPropertyBagPropertyType FBasePropertyIdHandler::GetPropertyType(const FProperty
 	EPropertyBagPropertyType PropertyType = EPropertyBagPropertyType::None;
 	if (InProperty)
 	{
-		const FProperty* Property = TryGetPropertyInContainerOrSelf(InProperty);
+		const FProperty* Property = GetPropertyInsideContainer(InProperty);
 		const FName PropertyName = Property->GetClass()->GetFName();
 		if (PropertyName == NAME_BoolProperty)
 		{
@@ -71,7 +71,7 @@ FName FBasePropertyIdHandler::GetPropertyTypeName(const FProperty* InProperty) c
 {
 	if (InProperty)
 	{
-		const FProperty* Property = TryGetPropertyInContainerOrSelf(InProperty);
+		const FProperty* Property = GetPropertyInsideContainer(InProperty);
 		if (const FByteProperty* ByteProperty = CastField<FByteProperty>(Property))
 		{
 			if (ByteProperty->Enum)
@@ -88,7 +88,7 @@ UObject* FBasePropertyIdHandler::GetPropertyTypeObject(const FProperty* InProper
 {
 	if (InProperty)
 	{
-		const FProperty* Property = TryGetPropertyInContainerOrSelf(InProperty);
+		const FProperty* Property = GetPropertyInsideContainer(InProperty);
 		if (const FByteProperty* ByteProperty = CastField<FByteProperty>(Property))
 		{
 			if (ByteProperty->Enum)
diff --git a/Engine/Plugins/VirtualProduction/RemoteControl/Source/RemoteControl/Private/PropertyIdHandler/EnumPropertyIdHandler.cpp b/Engine/Plugins/VirtualProduction/RemoteControl/Source/RemoteControl/Private/PropertyIdHandler/EnumPropertyIdHandler.cpp
--- a/Engine/Plugins/VirtualProduction/RemoteControl/Source/RemoteControl/Private/PropertyIdHandler/EnumPropertyIdHandler.cpp
+++ b/Engine/Plugins/VirtualProduction/RemoteControl/Source/RemoteControl/Private/PropertyIdHandler/EnumPropertyIdHandler.cpp
@@ -6,7 +6,7 @@ bool FEnumPropertyIdHandler::IsPropertySupported(const FProperty* InProperty) co
 {
 	if (InProperty)
 	{
-		const FProperty* Property = TryGetPropertyInContainerOrSelf(InProperty);
+		const FProperty* Property = GetPropertyInsideContainer(InProperty);
 		return Property->GetClass()->IsChildOf(FEnumProperty::StaticClass());
 	}
 	return false;
@@ -21,7 +21,7 @@ FName FEnumPropertyIdHandler::GetPropertyTypeName(const FProperty* InProperty) c
 {
 	if (InProperty)
 	{
-		const FProperty* Property = TryGetPropertyInContainerOrSelf(InProperty);
+		const FProperty* Property = GetPropertyInsideContainer(InProperty);
 		if (const FEnumProperty* EnumProperty = CastField<FEnumProperty>(Property))
 		{
 			return EnumProperty->GetEnum()->GetFName();
@@ -34,7 +34,7 @@ UObject* FEnumPropertyIdHandler::GetPropertyTypeObject(const FProperty* InProper
 {
 	if (InProperty)
 	{
-		const FProperty* Property = TryGetPropertyInContainerOrSelf(InProperty);
+		const FProperty* Property = GetPropertyInsideContainer(InProperty);
 		if (const FEnumProperty* EnumProperty = CastField<FEnumProperty>(Property))
 		{
 			return EnumProperty->GetEnum();
diff --git a/Engine/Plugins/VirtualProduction/RemoteControl/Source/RemoteControl/Private/PropertyIdHandler/ObjectPropertyIdHandler.cpp b/Engine/Plugins/VirtualProduction/RemoteControl/Source/RemoteControl/Private/PropertyIdHandler/ObjectPropertyIdHandler.cpp
--- a/Engine/Plugins/VirtualProduction/RemoteControl/Source/RemoteControl/Private/PropertyIdHandler/ObjectPropertyIdHandler.cpp
+++ b/Engine/Plugins/VirtualProduction/RemoteControl/Source/RemoteControl/Private/PropertyIdHandler/ObjectPropertyIdHandler.cpp
@@ -11,7 +11,7 @@ bool FObjectPropertyIdHandler::IsPropertySupported(const FProperty* InProperty)
 {
 	if (InProperty)
 	{
-		const FProperty* Property = TryGetPropertyInContainerOrSelf(InProperty);
+		const FProperty* Property = GetPropertyInsideContainer(InProperty);
 		return Property->GetClass()->IsChildOf(FObjectProperty::StaticClass());
 	}
 	return false;
@@ -26,7 +26,7 @@ FName FObjectPropertyIdHandler::GetPropertyTypeName(const FProperty* InProperty)
 {
 	if (InProperty)
 	{
-		const FProperty* Property = TryGetPropertyInContainerOrSelf(InProperty);
+		const FProperty* Property = GetPropertyInsideContainer(InProperty);
 		if (const FObjectProperty* ObjectProperty = CastField<FObjectProperty>(Property))
 		{
 			return ObjectProperty->PropertyClass->GetFName();
@@ -44,17 +44,17 @@ TObjectPtr<UObject> FObjectPropertyIdHandler::GetObjectPropertyDefaultValue(cons
 {
 	if (InClassToCreate)
 	{
+		if (InClassToCreate->IsChildOf(UMaterialInterface::StaticClass()))
+    	{
+    		return UMaterial::GetDefaultMaterial(MD_Surface);
+    	}
+
 		if (InClassToCreate->HasAnyClassFlags(CLASS_Abstract))
 		{
 			// for now skip any AbstractClass
 			return nullptr;
 		}
 
-		if (InClassToCreate->IsChildOf(UMaterialInterface::StaticClass()))
-		{
-			return UMaterial::GetDefaultMaterial(MD_Surface);
-		}
-
 		return NewObject<UObject>(GetTransientPackage(), InClassToCreate);
 	}
 	return nullptr;
diff --git a/Engine/Plugins/VirtualProduction/RemoteControl/Source/RemoteControl/Private/PropertyIdHandler/StructPropertyIdHandler.cpp b/Engine/Plugins/VirtualProduction/RemoteControl/Source/RemoteControl/Private/PropertyIdHandler/StructPropertyIdHandler.cpp
--- a/Engine/Plugins/VirtualProduction/RemoteControl/Source/RemoteControl/Private/PropertyIdHandler/StructPropertyIdHandler.cpp
+++ b/Engine/Plugins/VirtualProduction/RemoteControl/Source/RemoteControl/Private/PropertyIdHandler/StructPropertyIdHandler.cpp
@@ -6,7 +6,7 @@ bool FStructPropertyIdHandler::IsPropertySupported(const FProperty* InProperty)
 {
 	if (InProperty)
 	{
-		const FProperty* Property = TryGetPropertyInContainerOrSelf(InProperty);
+		const FProperty* Property = GetPropertyInsideContainer(InProperty);
 		return Property->GetClass()->IsChildOf(FStructProperty::StaticClass());
 	}
 	return false;
@@ -21,7 +21,7 @@ FName FStructPropertyIdHandler::GetPropertyTypeName(const FProperty* InProperty)
 {
 	if (InProperty)
 	{
-		const FProperty* Property = TryGetPropertyInContainerOrSelf(InProperty);
+		const FProperty* Property = GetPropertyInsideContainer(InProperty);
 		if (const FStructProperty* StructProperty = CastField<FStructProperty>(Property))
 		{
 			if (StructProperty->Struct->GetFName() == NAME_LinearColor)
@@ -38,7 +38,7 @@ UObject* FStructPropertyIdHandler::GetPropertyTypeObject(const FProperty* InProp
 {
 	if (InProperty)
 	{
-		const FProperty* Property = TryGetPropertyInContainerOrSelf(InProperty);
+		const FProperty* Property = GetPropertyInsideContainer(InProperty);
 		if (const FStructProperty* StructProperty = CastField<FStructProperty>(Property))
 		{
 			if (StructProperty->Struct->GetFName() == NAME_LinearColor)
diff --git a/Engine/Plugins/VirtualProduction/RemoteControl/Source/RemoteControl/Private/RemoteControlPropertyIdRegistry.cpp b/Engine/Plugins/VirtualProduction/RemoteControl/Source/RemoteControl/Private/RemoteControlPropertyIdRegistry.cpp
--- a/Engine/Plugins/VirtualProduction/RemoteControl/Source/RemoteControl/Private/RemoteControlPropertyIdRegistry.cpp
+++ b/Engine/Plugins/VirtualProduction/RemoteControl/Source/RemoteControl/Private/RemoteControlPropertyIdRegistry.cpp
@@ -577,7 +577,7 @@ const UClass* URemoteControlPropertyIdRegistry::GetClassByEntityId(const FGuid&
 	return nullptr;
 }
 
-TSet<FName> URemoteControlPropertyIdRegistry::GetFieldIdsNameList()
+TSet<FName> URemoteControlPropertyIdRegistry::GetFieldIdsNameList() const
 {
 	TSet<FName> OutIds;
 	for (FRCPropertyIdWrapper Field : IdentifiedFields)
@@ -590,6 +590,20 @@ TSet<FName> URemoteControlPropertyIdRegistry::GetFieldIdsNameList()
 	return OutIds;
 }
 
+TSet<FName> URemoteControlPropertyIdRegistry::GetFullPropertyIdsNamePossibilitiesList() const
+{
+	TSet<FName> InitialIdsList = GetFieldIdsNameList();
+	TSet<FName> OutIds;
+
+	for (const FName& Id : InitialIdsList)
+	{
+		OutIds.Append(GetPossiblePropertyIds(Id));
+	}
+
+	OutIds.Sort(FNameLexicalLess());
+	return OutIds;
+}
+
 TSet<FGuid> URemoteControlPropertyIdRegistry::GetEntityIdsForPropertyId(const FName& InPropertyId) const
 {
 	TSet<FGuid> OutEntityIds;
@@ -620,6 +634,40 @@ TSet<FGuid> URemoteControlPropertyIdRegistry::GetEntityIdsList()
 	return OutIds;
 }
 
+bool URemoteControlPropertyIdRegistry::Contains(FName InContainerPropertyId, FName InTargetPropertyId) const
+{
+	if (InContainerPropertyId == InTargetPropertyId)
+	{
+		return true;
+	}
+
+	const TSet<FName> PossibilitiesSecond = GetPossiblePropertyIds(InContainerPropertyId);
+	return PossibilitiesSecond.Contains(InTargetPropertyId);
+}
+
+TSet<FName> URemoteControlPropertyIdRegistry::GetPossiblePropertyIds(FName InPropId) const
+{
+	TSet<FName> OutPossibleIds;
+
+	// Add the given PropId
+	OutPossibleIds.Add(InPropId);
+	const FString FullPropId = InPropId.ToString();
+
+	// Get all the substring that are before a "."
+	FString PossibleId;
+	TArray<FString> PossibleSubIds;
+	FullPropId.ParseIntoArray(PossibleSubIds, TEXT("."));
+
+	for (const FString& SubId : PossibleSubIds)
+	{
+		PossibleId.Append(SubId);
+		OutPossibleIds.Add(FName(PossibleId));
+		PossibleId.AppendChar('.');
+	}
+
+	return OutPossibleIds;
+}
+
 void URemoteControlPropertyIdRegistry::UpdateEntityIds(const TMap<FGuid, FGuid>& InEntityIdMap)
 {
 	bool bNeedsRehash = false;
diff --git a/Engine/Plugins/VirtualProduction/RemoteControl/Source/RemoteControl/Public/IPropertyIdHandler.h b/Engine/Plugins/VirtualProduction/RemoteControl/Source/RemoteControl/Public/IPropertyIdHandler.h
--- a/Engine/Plugins/VirtualProduction/RemoteControl/Source/RemoteControl/Public/IPropertyIdHandler.h
+++ b/Engine/Plugins/VirtualProduction/RemoteControl/Source/RemoteControl/Public/IPropertyIdHandler.h
@@ -31,9 +31,8 @@ public:
 	/** Get the Property default value (Used for UObject to create a new one for the ValueWidget) */
 	virtual TObjectPtr<UObject> GetObjectPropertyDefaultValue(const FProperty* InProperty, const UClass* InClassToCreate) const = 0;
 
-protected:
 	/** Get the property inside the container if not inside any container return the property passed */
-	const FProperty* TryGetPropertyInContainerOrSelf(const FProperty* InProperty) const
+	const FProperty* GetPropertyInsideContainer(const FProperty* InProperty) const
 	{
 		if (const FArrayProperty* ArrayProperty = CastField<FArrayProperty>(InProperty))
 		{
@@ -44,9 +43,9 @@ protected:
 			return MapProperty->ValueProp;
 		}
 		if (const FSetProperty* SetProperty = CastField<FSetProperty>(InProperty))
-        {
-        	return SetProperty->ElementProp;
-        }
+		{
+			return SetProperty->ElementProp;
+		}
 		return InProperty;
 	}
 };
diff --git a/Engine/Plugins/VirtualProduction/RemoteControl/Source/RemoteControl/Public/RemoteControlPropertyIdRegistry.h b/Engine/Plugins/VirtualProduction/RemoteControl/Source/RemoteControl/Public/RemoteControlPropertyIdRegistry.h
--- a/Engine/Plugins/VirtualProduction/RemoteControl/Source/RemoteControl/Public/RemoteControlPropertyIdRegistry.h
+++ b/Engine/Plugins/VirtualProduction/RemoteControl/Source/RemoteControl/Public/RemoteControlPropertyIdRegistry.h
@@ -171,7 +171,14 @@ public:
 	 * Get the list of existing Property Ids.
 	 * @return The list of existing Property Ids that has a valid FieldId.
 	 */
-	TSet<FName> GetFieldIdsNameList();
+	TSet<FName> GetFieldIdsNameList() const;
+
+	/**
+	 * Returns all the possible Ids that can be used.
+	 *
+	 * @return A Set with the possible Ids to use.
+	 */
+	TSet<FName> GetFullPropertyIdsNamePossibilitiesList() const;
 
 	/**
 	 * Return all EntityIds that has the given PropertyId.
@@ -186,6 +193,25 @@ public:
 	 */
 	TSet<FGuid> GetEntityIdsList();
 
+	/**
+	 * Returns whether the 2 Ids are equal or if the Target one is contained in the Container one.\n
+	 * ex: 100.A and 100 will return true but not the other way around.
+	 *
+	 * @param InContainerPropertyId Container PropertyId, will create sub-string from this to later check if the target one is equal to one of them.
+	 * @param InTargetPropertyId Target Id to check if it is part of the container one.
+	 * @return True if the PropertyIds are equal or if the Target one is contained in at least 1 sub-string of the Container one.
+	 */
+	bool Contains(FName InContainerPropertyId, FName InTargetPropertyId) const;
+
+	/**
+	 * Returns the possible Ids based on given one.\n
+	 * ex: 100.A.B will return 100 and 100.A and 100.A.B
+	 *
+	 * @param InPropId Property Id used to generate all the other Ids that can be used with this one.
+	 * @return A Set with the possible Ids to use.
+	 */
+	TSet<FName> GetPossiblePropertyIds(FName InPropId) const;
+
 	DECLARE_MULTICAST_DELEGATE(FOnPropertyIdUpdated)
 	FOnPropertyIdUpdated& OnPropertyIdUpdated() { return OnPropertyIdUpdatedDelegate; }
 
diff --git a/Engine/Plugins/VirtualProduction/RemoteControl/Source/RemoteControlLogic/Private/Action/RCPropertyIdAction.cpp b/Engine/Plugins/VirtualProduction/RemoteControl/Source/RemoteControlLogic/Private/Action/RCPropertyIdAction.cpp
--- a/Engine/Plugins/VirtualProduction/RemoteControl/Source/RemoteControlLogic/Private/Action/RCPropertyIdAction.cpp
+++ b/Engine/Plugins/VirtualProduction/RemoteControl/Source/RemoteControlLogic/Private/Action/RCPropertyIdAction.cpp
@@ -31,13 +31,13 @@ void URCPropertyIdAction::Execute() const
 		return;
 	}
 
-	for (const TPair<FName, TObjectPtr<URCVirtualPropertySelfContainer>>& PropertyContainer : PropertySelfContainer)
+	for (const TPair<FPropertyIdContainerKey, TObjectPtr<URCVirtualPropertySelfContainer>>& PropertyContainer : PropertySelfContainer)
 	{
 		if (FProperty* Property = PropertyContainer.Value->GetProperty())
 		{
 			FRemoteControlPropertyIdArgs PropertyIdArgs;
 			PropertyIdArgs.VirtualProperty = PropertyContainer.Value;
-			PropertyIdArgs.PropertyId = PropertyId;
+			PropertyIdArgs.PropertyId = PropertyContainer.Key.PropertyId;
 			PropertyIdArgs.SuperType = Property->GetClass()->GetFName();
 
 			TSharedPtr<IPropertyIdHandler> PropertyIdHandler = IRemoteControlModule::Get().GetPropertyIdHandlerFor(Property);
@@ -81,7 +81,14 @@ void URCPropertyIdAction::Execute() const
 			{
 				if (const FStructProperty* StructProperty = CastField<FStructProperty>(Property))
 				{
-					PropertyIdArgs.SubType = StructProperty->Struct->GetFName();
+					if (StructProperty->Struct->GetFName() == NAME_LinearColor)
+					{
+						PropertyIdArgs.SubType = NAME_Color;
+					}
+					else
+					{
+						PropertyIdArgs.SubType = StructProperty->Struct->GetFName();
+					}
 				}
 			}
 			PresetWeakPtr->PerformChainReaction(PropertyIdArgs);
@@ -92,14 +99,14 @@ void URCPropertyIdAction::Execute() const
 
 void URCPropertyIdAction::UpdateEntityIds(const TMap<FGuid, FGuid>& InEntityIdMap)
 {
-	for (const TPair<FName, TObjectPtr<URCVirtualPropertySelfContainer>>& PropertyContainerEntry : PropertySelfContainer)
+	for (const TPair<FPropertyIdContainerKey, TObjectPtr<URCVirtualPropertySelfContainer>>& PropertyContainerEntry : PropertySelfContainer)
 	{
 		if (PropertyContainerEntry.Value)
 		{
 			PropertyContainerEntry.Value->UpdateEntityIds(InEntityIdMap);
 		}
 	}
-	for (const TPair<FName, TObjectPtr<URCVirtualPropertySelfContainer>>& PropertyContainerEntry : CachedPropertySelfContainer)
+	for (const TPair<FPropertyIdContainerKey, TObjectPtr<URCVirtualPropertySelfContainer>>& PropertyContainerEntry : CachedPropertySelfContainer)
 	{
 		if (PropertyContainerEntry.Value)
 		{
@@ -132,15 +139,16 @@ void URCPropertyIdAction::PostEditChangeChainProperty(FPropertyChangedChainEvent
 
 void URCPropertyIdAction::UpdatePropertyId()
 {
-	DefaultObject = nullptr;
 	if (URemoteControlPreset* Preset = PresetWeakPtr.Get())
 	{
 		PropertySelfContainer.Empty();
-		for (const FGuid& TargetProperty : Preset->GetPropertyIdRegistry().Get()->GetEntityIdsList())
+		const TObjectPtr<URemoteControlPropertyIdRegistry> PropertyIdRegistry = Preset->GetPropertyIdRegistry();
+
+		for (const FGuid& TargetProperty : PropertyIdRegistry->GetEntityIdsList())
 		{
 			if (const TSharedPtr<FRemoteControlProperty> TargetRCProperty = Preset->GetExposedEntity<FRemoteControlProperty>(TargetProperty).Pin())
 			{
-				if (TargetRCProperty->PropertyId == PropertyId)
+				if (PropertyIdRegistry->Contains(TargetRCProperty->PropertyId, PropertyId))
 				{
 					if (FProperty* Property = TargetRCProperty->GetProperty())
 					{
@@ -149,51 +157,56 @@ void URCPropertyIdAction::UpdatePropertyId()
 						{
 							continue;
 						}
-						FName PropertyClassName = PropertyIdHandler->GetPropertyTypeName(Property);
-						const FName& NewPropertyIdName = *(PropertyId.ToString() + TEXT(".") + PropertyClassName.ToString());
-						if (PropertyIdHandler->GetPropertyType(Property) == EPropertyBagPropertyType::Object)
+
+						FString PrefixPropertyClassName = TargetRCProperty->PropertyId.ToString() + TEXT(".");
+						FName PropertyClassName = FName(PrefixPropertyClassName + PropertyIdHandler->GetPropertyTypeName(Property).ToString());
+						const FName& NewPropertyIdName = *(TargetRCProperty->PropertyId.ToString() + TEXT(".") + PropertyClassName.ToString());
+						const FProperty* PropToDuplicate = PropertyIdHandler->GetPropertyInsideContainer(Property);
+						FPropertyIdContainerKey CurrentKey = { TargetRCProperty->PropertyId, PropertyClassName };
+
+						if (CachedPropertySelfContainer.Contains(CurrentKey))
+						{
+							PropertySelfContainer.Add(CurrentKey, CachedPropertySelfContainer[CurrentKey]);
+						}
+						else
 						{
-							const FName PropertyNameToSearchFor = GET_MEMBER_NAME_CHECKED(URCPropertyIdAction, DefaultObject);
-							DefaultObject = PropertyIdHandler->GetObjectPropertyDefaultValue(Property, Preset->GetPropertyIdRegistry().Get()->GetClassByEntityId(TargetProperty));
-							if (DefaultObject && !PropertyNameToSearchFor.IsNone())
+							CachedPropertySelfContainer.Add(CurrentKey, NewObject<URCVirtualPropertySelfContainer>(this));
+
+							// Special case for LinearColor, since we want to base the copy on Color struct here we add the property instead of duplicating it
+							bool bColorAdded = false;
+
+							if (const FStructProperty* StructProp = CastField<FStructProperty>(PropToDuplicate))
 							{
-								for (TFieldIterator<FObjectProperty> FieldIt(GetClass()); FieldIt; ++FieldIt)
+								if (StructProp->Struct)
 								{
-									if (FieldIt->GetFName() == PropertyNameToSearchFor)
+									if (StructProp->Struct->GetFName() == NAME_LinearColor ||
+										StructProp->Struct->GetFName() == NAME_Color)
 									{
-										FieldIt->PropertyClass = DefaultObject->GetClass();
-										if (CachedPropertySelfContainer.Contains(PropertyClassName))
-										{
-											PropertySelfContainer.Add(PropertyClassName, CachedPropertySelfContainer[PropertyClassName]);
-										}
-										else
+										CachedPropertySelfContainer[CurrentKey]->AddProperty(NewPropertyIdName,
+										PropertyIdHandler->GetPropertyType(Property),
+										PropertyIdHandler->GetPropertyTypeObject(Property));
+
+#if WITH_EDITORONLY_DATA
+										FProperty* BagProperty = CachedPropertySelfContainer[CurrentKey]->GetProperty();
+										FStructProperty* StructProperty = CastField<FStructProperty>(BagProperty);
+										if (StructProperty && StructProperty->Struct &&
+											!StructProperty->HasMetaData("OnlyUpdateOnInteractionEnd"))
 										{
-											CachedPropertySelfContainer.Add(PropertyClassName, NewObject<URCVirtualPropertySelfContainer>(this));
-											CachedPropertySelfContainer[PropertyClassName]->DuplicateProperty(NewPropertyIdName, *FieldIt);
-											CachedPropertySelfContainer[PropertyClassName]->PresetWeakPtr = PresetWeakPtr;
-											PropertySelfContainer.Add(PropertyClassName, CachedPropertySelfContainer[PropertyClassName]);
+											StructProperty->AppendMetaData({{FName("OnlyUpdateOnInteractionEnd"), TEXT("true")}});
 										}
-										break;
+#endif
+										bColorAdded = true;
 									}
 								}
 							}
-						}
-						else
-						{
-							if (CachedPropertySelfContainer.Contains(PropertyClassName))
-							{
-								PropertySelfContainer.Add(PropertyClassName, CachedPropertySelfContainer[PropertyClassName]);
-							}
-							else
+
+							if (!bColorAdded)
 							{
-								CachedPropertySelfContainer.Add(PropertyClassName, NewObject<URCVirtualPropertySelfContainer>(this));
-								CachedPropertySelfContainer[PropertyClassName]->AddProperty(NewPropertyIdName,
-									PropertyIdHandler->GetPropertyType(Property),
-									PropertyIdHandler->GetPropertyTypeObject(Property));
-								CachedPropertySelfContainer[PropertyClassName]->PresetWeakPtr = PresetWeakPtr;
-								
-								PropertySelfContainer.Add(PropertyClassName, CachedPropertySelfContainer[PropertyClassName]);
+								CachedPropertySelfContainer[CurrentKey]->DuplicateProperty(NewPropertyIdName, PropToDuplicate);
 							}
+
+							CachedPropertySelfContainer[CurrentKey]->PresetWeakPtr = PresetWeakPtr;
+							PropertySelfContainer.Add(CurrentKey, CachedPropertySelfContainer[CurrentKey]);
 						}
 					}
 				}
diff --git a/Engine/Plugins/VirtualProduction/RemoteControl/Source/RemoteControlLogic/Public/Action/RCPropertyIdAction.h b/Engine/Plugins/VirtualProduction/RemoteControl/Source/RemoteControlLogic/Public/Action/RCPropertyIdAction.h
--- a/Engine/Plugins/VirtualProduction/RemoteControl/Source/RemoteControlLogic/Public/Action/RCPropertyIdAction.h
+++ b/Engine/Plugins/VirtualProduction/RemoteControl/Source/RemoteControlLogic/Public/Action/RCPropertyIdAction.h
@@ -8,6 +8,40 @@
 
 class URCVirtualPropertySelfContainer;
 
+/**
+ * Key for the PropertyContainer of the PropertyId Action
+ */
+USTRUCT()
+struct FPropertyIdContainerKey
+{
+	GENERATED_BODY()
+
+	/** PropertyId */
+	UPROPERTY()
+	FName PropertyId;
+
+	/** Name of the container */
+	UPROPERTY()
+	FName ContainerName;
+
+	/** Return hash value for this object, used when using this object as a key inside hashing containers. */
+	friend uint32 GetTypeHash(const FPropertyIdContainerKey& InKey)
+	{
+		return HashCombine(GetTypeHash(InKey.PropertyId), GetTypeHash(InKey.ContainerName));
+	}
+
+	/** Comparison operator, used by hashing containers. */
+	bool operator==(const FPropertyIdContainerKey& InOtherKey) const
+	{
+		return PropertyId == InOtherKey.PropertyId && ContainerName == InOtherKey.ContainerName;
+	}
+
+	bool operator<(const FPropertyIdContainerKey& InOtherKey) const
+	{
+		return PropertyId.Compare(InOtherKey.PropertyId) < 0;
+	}
+};
+
 /**
  * Action for PropertyId
  */
@@ -43,14 +77,9 @@ public:
 
 	/** Virtual Property Container */
 	UPROPERTY()
-	TMap<FName, TObjectPtr<URCVirtualPropertySelfContainer>> PropertySelfContainer;
+	TMap<FPropertyIdContainerKey, TObjectPtr<URCVirtualPropertySelfContainer>> PropertySelfContainer;
 
 	/** Cached Virtual Property Container */
 	UPROPERTY()
-	TMap<FName, TObjectPtr<URCVirtualPropertySelfContainer>> CachedPropertySelfContainer;
-
-private:
-	/** Holds the default Object. */
-	UPROPERTY(Transient)
-	TObjectPtr<UObject> DefaultObject = nullptr;
+	TMap<FPropertyIdContainerKey, TObjectPtr<URCVirtualPropertySelfContainer>> CachedPropertySelfContainer;
 };
diff --git a/Engine/Plugins/VirtualProduction/RemoteControl/Source/RemoteControlUI/Private/UI/Action/RCActionModel.cpp b/Engine/Plugins/VirtualProduction/RemoteControl/Source/RemoteControlUI/Private/UI/Action/RCActionModel.cpp
--- a/Engine/Plugins/VirtualProduction/RemoteControl/Source/RemoteControlUI/Private/UI/Action/RCActionModel.cpp
+++ b/Engine/Plugins/VirtualProduction/RemoteControl/Source/RemoteControlUI/Private/UI/Action/RCActionModel.cpp
@@ -2,6 +2,7 @@
 
 #include "RCActionModel.h"
 
+#include "Algo/Count.h"
 #include "Commands/RemoteControlCommands.h"
 #include "Controller/RCController.h"
 #include "IDetailTreeNode.h"
@@ -119,7 +120,7 @@ TSharedRef<SWidget> FRCActionModel::GetTypeColorTagWidget() const
 		[
 			SNew(SBorder)
 			.Visibility(EVisibility::HitTestInvisible)
-			.BorderImage(FAppStyle::Get().GetBrush("NumericEntrySpinBox.NarrowDecorator"))
+			.BorderImage(FAppStyle::Get().GetBrush("NumericEntrySpinBox.Decorator"))
 			.BorderBackgroundColor(TypeColor)
 			.Padding(FMargin(5.0f, 0.0f, 0.0f, 0.f))
 		];
@@ -351,7 +352,7 @@ FRCPropertyIdActionType::~FRCPropertyIdActionType()
 	}
 #endif
 
-	for (const TPair<FName, TSharedPtr<IPropertyRowGenerator>>& CachedGenerator : CachedPropertyIdValueRowGenerator)
+	for (const TPair<FPropertyIdContainerKey, TSharedPtr<IPropertyRowGenerator>>& CachedGenerator : CachedPropertyIdValueRowGenerator)
 	{
 		if (CachedGenerator.Value.IsValid())
 		{
@@ -359,7 +360,6 @@ FRCPropertyIdActionType::~FRCPropertyIdActionType()
 		}
 	}
 	PropertyIdValueRowGenerator.Reset();
-	CachedPropertyIdValueRowGenerator.Reset();
 }
 
 FLinearColor FRCPropertyIdActionType::GetPropertyIdTypeColor() const
@@ -374,17 +374,20 @@ TSharedRef<SWidget> FRCPropertyIdActionType::GetPropertyIdNameWidget() const
 	{
 		return SNullWidget::NullWidget;
 	}
+
 	const FNodeWidgets FieldIdNodeWidgets = FieldIdTreeNodeWeakPtr.Pin()->CreateNodeWidgets();
 	const TSharedRef<SHorizontalBox> NameWidget = SNew(SHorizontalBox);
+
 	if (FieldIdNodeWidgets.ValueWidget)
 	{
 		NameWidget->AddSlot()
-			.Padding(3.f, 2.f)
+			.Padding(10.f, 2.f)
 			.VAlign(VAlign_Center)
 			[
 				FieldIdNodeWidgets.ValueWidget.ToSharedRef()
 			];
 	}
+
 	return NameWidget;
 }
 
@@ -392,15 +395,71 @@ TSharedRef<SWidget> FRCPropertyIdActionType::GetPropertyIdValueWidget() const
 {
 	TSharedRef<SVerticalBox> VerticalBox = SNew(SVerticalBox);
 
-	for (TPair<FName, TWeakPtr<IDetailTreeNode>> ValueTreeNode : ValueTreeNodeWeakPtr)
+	TArray<FPropertyIdContainerKey> SortedTreeNodesKeys;
+	ValueTreeNodeWeakPtr.GetKeys(SortedTreeNodesKeys);
+	SortedTreeNodesKeys.Sort();
+
+	FName LastPropId = NAME_None;
+	FMargin TitlePadding = FMargin(0.f);
+	FMargin ValuePadding = FMargin(0.f);
+	int32 OriginalPropIdLength = 0;
+	int32 OriginalDotCount = 0;
+
+	if (const URCPropertyIdAction* PropIdAction = PropertyIdActionWeakPtr.Get())
 	{
-		// store it and just remove the unusued to not reset the value widget
-		VerticalBox->AddSlot()
-			.AutoHeight()
-			[
-				UE::RCUIHelpers::GetGenericFieldWidget(ValueTreeNode.Value.Pin())
-			];
+		const FString PropIdAsString = PropIdAction->PropertyId.ToString();
+		OriginalDotCount = Algo::Count(PropIdAsString, '.');
+		OriginalPropIdLength = PropIdAction->PropertyId.GetStringLength();
+	}
+
+	for (const FPropertyIdContainerKey& TreeNodeKey : SortedTreeNodesKeys)
+	{
+		FString PropIdAsString = TEXT("");
+
+		if (LastPropId != TreeNodeKey.PropertyId)
+		{
+			constexpr float Offset = 5.f;
+
+			if (LastPropId != NAME_None)
+			{
+				TitlePadding.Top = Offset;
+			}
+
+			LastPropId = TreeNodeKey.PropertyId;
+			PropIdAsString = LastPropId.ToString();
+			if (PropIdAsString.Len() != OriginalPropIdLength)
+			{
+				FString LabelToUse = PropIdAsString.RightChop(OriginalPropIdLength);
+				const int32 CurrentDotCount = Algo::Count(LabelToUse, '.');
+				TitlePadding.Left = Offset * (CurrentDotCount - OriginalDotCount);
+				ValuePadding.Left = TitlePadding.Left + 2.f;
+
+				VerticalBox->AddSlot()
+					.AutoHeight()
+					.Padding(TitlePadding)
+					[
+						SNew(STextBlock)
+						.Text(FText::FromString(LabelToUse))
+					];
+			}
+		}
+
+		if (const TWeakPtr<IDetailTreeNode>* Node = ValueTreeNodeWeakPtr.Find(TreeNodeKey))
+		{
+			VerticalBox->AddSlot()
+				.Padding(ValuePadding)
+				.AutoHeight()
+				[
+					SNew(SHorizontalBox)
+					+ SHorizontalBox::Slot()
+					.AutoWidth()
+					[
+						UE::RCUIHelpers::GetGenericFieldWidget(Node->Pin())
+					]
+				];
+		}
 	}
+
 	return VerticalBox;
 }
 
@@ -464,7 +523,13 @@ void FRCPropertyIdActionType::RefreshValueWidget()
 		// Generate UI widget for Action input
 		PropertyIdValueRowGenerator.Reset();
 		ValueTreeNodeWeakPtr.Reset();
-		for (const TPair<FName, TObjectPtr<URCVirtualPropertySelfContainer>>& PropertyContainer : PropertyIdAction->PropertySelfContainer)
+		PropertyIdAction->PropertySelfContainer.KeySort([] (const FPropertyIdContainerKey& InFirst, const FPropertyIdContainerKey& InSecond)
+		{
+			return InSecond < InFirst;
+		});
+
+
+		for (const TPair<FPropertyIdContainerKey, TObjectPtr<URCVirtualPropertySelfContainer>>& PropertyContainer : PropertyIdAction->PropertySelfContainer)
 		{
 			if (IsValid(PropertyContainer.Value))
 			{
@@ -485,6 +550,7 @@ void FRCPropertyIdActionType::RefreshValueWidget()
 
 						PropertyIdValueRowGenerator.Add(PropertyContainer.Key, CachedPropertyIdValueRowGenerator[PropertyContainer.Key]);
 					}
+
 					for (const TSharedRef<IDetailTreeNode>& CategoryNode : PropertyIdValueRowGenerator[PropertyContainer.Key]->GetRootTreeNodes())
 					{
 						TArray<TSharedRef<IDetailTreeNode>> Children;
diff --git a/Engine/Plugins/VirtualProduction/RemoteControl/Source/RemoteControlUI/Private/UI/Action/RCActionModel.h b/Engine/Plugins/VirtualProduction/RemoteControl/Source/RemoteControlUI/Private/UI/Action/RCActionModel.h
--- a/Engine/Plugins/VirtualProduction/RemoteControl/Source/RemoteControlUI/Private/UI/Action/RCActionModel.h
+++ b/Engine/Plugins/VirtualProduction/RemoteControl/Source/RemoteControlUI/Private/UI/Action/RCActionModel.h
@@ -192,19 +192,19 @@ protected:
 	TSharedPtr<IPropertyRowGenerator> PropertyIdNameRowGenerator;
 	
 	/** The cache of the row generator used to not recreate it */
-	TMap<FName, TSharedPtr<IPropertyRowGenerator>> PropertyIdValueRowGenerator;
+	TMap<FPropertyIdContainerKey, TSharedPtr<IPropertyRowGenerator>> PropertyIdValueRowGenerator;
 
 	/** The row generator used to represent value of this widget as a row, when used with SListView */
-	TMap<FName, TSharedPtr<IPropertyRowGenerator>> CachedPropertyIdValueRowGenerator;
+	TMap<FPropertyIdContainerKey, TSharedPtr<IPropertyRowGenerator>> CachedPropertyIdValueRowGenerator;
 
 	/** Used to create a generic Field Id Widget for the property row widget */
 	TWeakPtr<IDetailTreeNode> FieldIdTreeNodeWeakPtr;
 	
 	/** Used to create a generic Value Widget for the property row widget */
-	TMap<FName, TWeakPtr<IDetailTreeNode>> ValueTreeNodeWeakPtr;
+	TMap<FPropertyIdContainerKey, TWeakPtr<IDetailTreeNode>> ValueTreeNodeWeakPtr;
 
 	/** Used to cache a generic Value Widget for the property row widget to not recreate it */
-	TMap<FName, TWeakPtr<IDetailTreeNode>> CachedValueTreeNodeWeakPtr;
+	TMap<FPropertyIdContainerKey, TWeakPtr<IDetailTreeNode>> CachedValueTreeNodeWeakPtr;
 };
 
 /*
diff --git a/Engine/Plugins/VirtualProduction/RemoteControl/Source/RemoteControlUI/Private/UI/Action/SRCActionPanel.cpp b/Engine/Plugins/VirtualProduction/RemoteControl/Source/RemoteControlUI/Private/UI/Action/SRCActionPanel.cpp
--- a/Engine/Plugins/VirtualProduction/RemoteControl/Source/RemoteControlUI/Private/UI/Action/SRCActionPanel.cpp
+++ b/Engine/Plugins/VirtualProduction/RemoteControl/Source/RemoteControlUI/Private/UI/Action/SRCActionPanel.cpp
@@ -275,8 +275,8 @@ TSharedRef<SWidget> SRCActionPanel::GetActionMenuContentWidget()
 			{
 				if (Preset->GetPropertyIdRegistry())
 				{
-					TSet<FName> IdList = Preset->GetPropertyIdRegistry().Get()->GetFieldIdsNameList();
-					if (IdList.Num())
+					TSet<FName> IdList = Preset->GetPropertyIdRegistry().Get()->GetFullPropertyIdsNamePossibilitiesList();
+					if (!IdList.IsEmpty())
 					{
 						MenuBuilder.AddSubMenu(
 						LOCTEXT("AddActionSubMenu", "Add specific ID action"),
diff --git a/Engine/Plugins/VirtualProduction/RemoteControl/Source/RemoteControlUI/Private/UI/Action/SRCVirtualPropertyWidget.cpp b/Engine/Plugins/VirtualProduction/RemoteControl/Source/RemoteControlUI/Private/UI/Action/SRCVirtualPropertyWidget.cpp
--- a/Engine/Plugins/VirtualProduction/RemoteControl/Source/RemoteControlUI/Private/UI/Action/SRCVirtualPropertyWidget.cpp
+++ b/Engine/Plugins/VirtualProduction/RemoteControl/Source/RemoteControlUI/Private/UI/Action/SRCVirtualPropertyWidget.cpp
@@ -39,6 +39,7 @@ void SRCVirtualPropertyWidget::Construct(const FArguments& InArgs, URCVirtualPro
 			.OnMouseDoubleClick(this, &SRCVirtualPropertyWidget::OnMouseDoubleClick)
 			[
 				SAssignNew(VirtualPropertyWidgetBox, SBox)
+				.VAlign(VAlign_Center)
 				.Padding(FMargin(6.f))
 				[
 					VirtualPropertyDisplayWidget
diff --git a/Engine/Plugins/VirtualProduction/RemoteControl/Source/RemoteControlUI/Private/UI/RCUIHelpers.cpp b/Engine/Plugins/VirtualProduction/RemoteControl/Source/RemoteControlUI/Private/UI/RCUIHelpers.cpp
--- a/Engine/Plugins/VirtualProduction/RemoteControl/Source/RemoteControlUI/Private/UI/RCUIHelpers.cpp
+++ b/Engine/Plugins/VirtualProduction/RemoteControl/Source/RemoteControlUI/Private/UI/RCUIHelpers.cpp
@@ -206,7 +206,7 @@ TSharedRef<SWidget> UE::RCUIHelpers::GetTypeColorWidget(const FProperty* InPrope
 			[
 				SNew(SBorder)
 				.Visibility(EVisibility::Visible)
-				.BorderImage(FAppStyle::Get().GetBrush("NumericEntrySpinBox.NarrowDecorator"))
+				.BorderImage(FAppStyle::Get().GetBrush("NumericEntrySpinBox.Decorator"))
 				.BorderBackgroundColor(TypeColor)
 				.Padding(FMargin(5.0f, 0.0f, 0.0f, 0.f))
 			];
diff --git a/Engine/Plugins/VirtualProduction/RemoteControl/Source/RemoteControlUI/Private/UI/SRCPanelExposedEntity.cpp b/Engine/Plugins/VirtualProduction/RemoteControl/Source/RemoteControlUI/Private/UI/SRCPanelExposedEntity.cpp
--- a/Engine/Plugins/VirtualProduction/RemoteControl/Source/RemoteControlUI/Private/UI/SRCPanelExposedEntity.cpp
+++ b/Engine/Plugins/VirtualProduction/RemoteControl/Source/RemoteControlUI/Private/UI/SRCPanelExposedEntity.cpp
@@ -508,7 +508,6 @@ TSharedRef<SWidget> SRCPanelExposedEntity::CreateEntityWidget(TSharedPtr<SWidget
 	Args.PropertyIdWidget = SNew(SBox)
 		[
 			SNew(SEditableTextBox)
-			.Justification(ETextJustify::Center)
 			.MinDesiredWidth(50.f)
 			.SelectAllTextWhenFocused(true)
 			.RevertTextOnEscape(true)
