[RemoteControl]
Add search for BP_Set while searching for setter.

diff --git a/Engine/Plugins/VirtualProduction/RemoteControl/Source/RemoteControlCommon/Public/RCPropertyUtilities.inl b/Engine/Plugins/VirtualProduction/RemoteControl/Source/RemoteControlCommon/Public/RCPropertyUtilities.inl
--- a/Engine/Plugins/VirtualProduction/RemoteControl/Source/RemoteControlCommon/Public/RCPropertyUtilities.inl
+++ b/Engine/Plugins/VirtualProduction/RemoteControl/Source/RemoteControlCommon/Public/RCPropertyUtilities.inl
@@ -524,7 +524,8 @@ namespace RemoteControlPropertyUtilities
 
 			static const TArray<FString> SetterPrefixes = {
 				FString("Set"),
-				FString("K2_Set")
+				FString("K2_Set"),
+				FString("BP_Set")
 			};
 
 			for (const FString& Prefix : SetterPrefixes)
