[RemoteControl]
Add check for Editable properties in PropertyId action

diff --git a/Engine/Plugins/VirtualProduction/RemoteControl/Source/RemoteControl/Private/RemoteControlPropertyIdRegistry.cpp b/Engine/Plugins/VirtualProduction/RemoteControl/Source/RemoteControl/Private/RemoteControlPropertyIdRegistry.cpp
--- a/Engine/Plugins/VirtualProduction/RemoteControl/Source/RemoteControl/Private/RemoteControlPropertyIdRegistry.cpp
+++ b/Engine/Plugins/VirtualProduction/RemoteControl/Source/RemoteControl/Private/RemoteControlPropertyIdRegistry.cpp
@@ -529,10 +529,13 @@ void URemoteControlPropertyIdRegistry::AddIdentifiedField(const TSharedRef<FRemo
 	{
 		if (const TSharedPtr<FRemoteControlProperty> RCProperty = StaticCastSharedRef<FRemoteControlProperty>(InFieldToIdentify))
 		{
-			FRCPropertyIdWrapper Wrapper{ RCProperty.ToSharedRef()};
-			if (Wrapper.IsValid())
+			if (RCProperty->IsEditable())
 			{
-				IdentifiedFields.Add(MoveTemp(Wrapper));
+				FRCPropertyIdWrapper Wrapper{ RCProperty.ToSharedRef() };
+				if (Wrapper.IsValid())
+				{
+					IdentifiedFields.Add(MoveTemp(Wrapper));
+				}
 			}
 		}
 	}
