[RemoteControl]
Add check during BindObject for RootComponent validity.

diff --git a/Engine/Plugins/VirtualProduction/RemoteControl/Source/RemoteControl/Private/RemoteControlField.cpp b/Engine/Plugins/VirtualProduction/RemoteControl/Source/RemoteControl/Private/RemoteControlField.cpp
--- a/Engine/Plugins/VirtualProduction/RemoteControl/Source/RemoteControl/Private/RemoteControlField.cpp
+++ b/Engine/Plugins/VirtualProduction/RemoteControl/Source/RemoteControl/Private/RemoteControlField.cpp
@@ -86,7 +86,8 @@ void FRemoteControlField::BindObject(UObject* InObjectToBind)
 		else if (AActor* Actor = Cast<AActor>(InObjectToBind))
 		{
 			// Attempt to bind to the root component since it is a very common case.
-			if (Actor->GetRootComponent()->GetClass() == ResolvedOwnerClass)
+			const USceneComponent* RootComponent = Actor->GetRootComponent();
+			if (RootComponent && RootComponent->GetClass() == ResolvedOwnerClass)
 			{
 				FRemoteControlEntity::BindObject(Actor->GetRootComponent());
 			}
