[RemoteControl]
Fix PropertyId in certain cases will be displayed wrongly.

diff --git a/Engine/Plugins/VirtualProduction/RemoteControl/Source/RemoteControlUI/Private/UI/SRCPanelExposedEntity.cpp b/Engine/Plugins/VirtualProduction/RemoteControl/Source/RemoteControlUI/Private/UI/SRCPanelExposedEntity.cpp
--- a/Engine/Plugins/VirtualProduction/RemoteControl/Source/RemoteControlUI/Private/UI/SRCPanelExposedEntity.cpp
+++ b/Engine/Plugins/VirtualProduction/RemoteControl/Source/RemoteControlUI/Private/UI/SRCPanelExposedEntity.cpp
@@ -33,7 +33,6 @@
 #include "Widgets/Input/SButton.h"
 #include "Widgets/Input/SCheckBox.h"
 #include "Widgets/Input/SComboButton.h"
-#include "Widgets/Input/SEditableTextBox.h"
 #include "Widgets/Layout/SBorder.h"
 #include "Widgets/Layout/SBox.h"
 #include "Widgets/Notifications/SPopUpErrorText.h"
@@ -177,7 +176,6 @@ void SRCPanelExposedEntity::Initialize(const FGuid& InEntityId, URemoteControlPr
 			{
 				const TSharedPtr<FRemoteControlProperty> RCProperty = StaticCastSharedPtr<FRemoteControlProperty>(RCEntity);
 				CachedFieldPath = RCProperty->FieldPathInfo.ToString();
-				PropertyIdLabel = RCProperty->PropertyId;
 			}
 
 			FName OwnerFName;
@@ -429,27 +427,6 @@ void SRCPanelExposedEntity::OnLabelCommitted(const FText& InLabel, ETextCommit::
 	}
 }
 
-void SRCPanelExposedEntity::OnPropertyIdTextCommitted(const FText& InText, ETextCommit::Type InCommitInfo)
-{
-	if (const URemoteControlPreset* RCPreset = Preset.Get())
-	{
-		if (const TSharedPtr<FRemoteControlEntity> RCEntity = GetEntity())
-		{
-			if (const TSharedPtr<FRemoteControlField> RCField = StaticCastSharedPtr<FRemoteControlField>(RCEntity))
-			{
-				const FName NewId = FName(InText.ToString());
-				if (RCField->PropertyId.Compare(NewId) != 0)
-				{
-					RCField->PropertyId = NewId;
-					PropertyIdLabel = NewId;
-					RCPreset->UpdateIdentifiedField(RCField.ToSharedRef());
-					OnPropertyIdRenamed().ExecuteIfBound(NewId);
-				}
-			}
-		}
-	}
-}
-
 void SRCPanelExposedEntity::OnActorSelected(AActor* InActor) const
 {
 	if (TSharedPtr<FRemoteControlEntity> Entity = GetEntity())
@@ -499,22 +476,19 @@ bool SRCPanelExposedEntity::IsActorSelectable(const AActor* Actor) const
 
 TSharedRef<SWidget> SRCPanelExposedEntity::CreateEntityWidget(TSharedPtr<SWidget> ValueWidget, TSharedPtr<SWidget> ResetWidget, const FText& OptionalWarningMessage, TSharedRef<SWidget> EditConditionWidget)
 {
-	FMakeNodeWidgetArgs Args;
+	const FMakeNodeWidgetArgs Args = CreateEntityWidgetInternal(ValueWidget, ResetWidget, OptionalWarningMessage, EditConditionWidget);
 
 	TSharedRef<SBorder> Widget = SNew(SBorder)
 		.Padding(0.0f)
 		.BorderImage(this, &SRCPanelExposedEntity::GetBorderImage);
-	
-	Args.PropertyIdWidget = SNew(SBox)
-		[
-			SNew(SEditableTextBox)
-			.MinDesiredWidth(50.f)
-			.SelectAllTextWhenFocused(true)
-			.RevertTextOnEscape(true)
-			.ClearKeyboardFocusOnCommit(true)
-			.Text_Lambda([this] () { return FText::FromName(PropertyIdLabel); })
-			.OnTextCommitted(this, &SRCPanelExposedEntity::OnPropertyIdTextCommitted)
-		];
+
+	Widget->SetContent(MakeNodeWidget(Args));
+	return Widget;
+}
+
+SRCPanelTreeNode::FMakeNodeWidgetArgs SRCPanelExposedEntity::CreateEntityWidgetInternal(TSharedPtr<SWidget> ValueWidget, TSharedPtr<SWidget> ResetWidget, const FText& OptionalWarningMessage, TSharedRef<SWidget> EditConditionWidget)
+{
+	FMakeNodeWidgetArgs Args;
 
 	const FSlateBrush* TrashBrush = FAppStyle::Get().GetBrush("Icons.Delete");
 
@@ -587,8 +561,7 @@ TSharedRef<SWidget> SRCPanelExposedEntity::CreateEntityWidget(TSharedPtr<SWidget
 
 	Args.ResetButton = ResetWidget;
 
-	Widget->SetContent(MakeNodeWidget(Args));
-	return Widget;
+	return Args;
 }
 
 void SRCPanelExposedEntity::OnActorSelectedForRebindAllProperties(AActor* InActor) const
diff --git a/Engine/Plugins/VirtualProduction/RemoteControl/Source/RemoteControlUI/Private/UI/SRCPanelExposedEntity.h b/Engine/Plugins/VirtualProduction/RemoteControl/Source/RemoteControlUI/Private/UI/SRCPanelExposedEntity.h
--- a/Engine/Plugins/VirtualProduction/RemoteControl/Source/RemoteControlUI/Private/UI/SRCPanelExposedEntity.h
+++ b/Engine/Plugins/VirtualProduction/RemoteControl/Source/RemoteControlUI/Private/UI/SRCPanelExposedEntity.h
@@ -50,6 +50,9 @@ protected:
 	/** Create an exposed entity widget with a drag handle and unexpose button. */
 	TSharedRef<SWidget> CreateEntityWidget(TSharedPtr<SWidget> ValueWidget, TSharedPtr<SWidget> ResetWidget = SNullWidget::NullWidget, const FText& OptionalWarningMessage = FText::GetEmpty(), TSharedRef<SWidget> EditConditionWidget = SNullWidget::NullWidget);
 
+	/** Returns populated args to display this widget. */
+	virtual FMakeNodeWidgetArgs CreateEntityWidgetInternal(TSharedPtr<SWidget> ValueWidget, TSharedPtr<SWidget> ResetWidget = SNullWidget::NullWidget, const FText& OptionalWarningMessage = FText::GetEmpty(), TSharedRef<SWidget> EditConditionWidget = SNullWidget::NullWidget);
+
 protected:
 	/** Id of the entity. */
 	FGuid EntityId;
@@ -102,8 +105,6 @@ private:
 	bool OnVerifyItemLabelChanged(const FText& InLabel, FText& OutErrorMessage);
 	/** Handles committing a entity label. */
 	void OnLabelCommitted(const FText& InLabel, ETextCommit::Type InCommitInfo);
-	/** Handles committing a property id label. */
-	void OnPropertyIdTextCommitted(const FText& InText, ETextCommit::Type InCommitInfo);
 	/** Returns whether or not the actor is selectable for a binding replacement. */
 	bool IsActorSelectable(const AActor* Parent) const;
 	/** Handle clicking on the unexpose button. */
diff --git a/Engine/Plugins/VirtualProduction/RemoteControl/Source/RemoteControlUI/Private/UI/SRCPanelExposedField.cpp b/Engine/Plugins/VirtualProduction/RemoteControl/Source/RemoteControlUI/Private/UI/SRCPanelExposedField.cpp
--- a/Engine/Plugins/VirtualProduction/RemoteControl/Source/RemoteControlUI/Private/UI/SRCPanelExposedField.cpp
+++ b/Engine/Plugins/VirtualProduction/RemoteControl/Source/RemoteControlUI/Private/UI/SRCPanelExposedField.cpp
@@ -31,6 +31,7 @@
 #include "UObject/Object.h"
 #include "Widgets/Input/SButton.h"
 #include "Widgets/Input/SCheckBox.h"
+#include "Widgets/Input/SEditableTextBox.h"
 #include "Widgets/Layout/SBorder.h"
 #include "Widgets/Layout/SBox.h"
 #include "Widgets/Masks/SRCProtocolMask.h"
@@ -144,6 +145,7 @@ void SRCPanelExposedField::Construct(const FArguments& InArgs, TWeakPtr<FRemoteC
 		Initialize(FieldPtr->GetId(), InArgs._Preset.Get(), InArgs._LiveMode);
 
 		CachedLabel = FieldPtr->GetLabel();
+		PropertyIdLabel = FieldPtr->PropertyId;
 		EntityId = FieldPtr->GetId();
 
 		if (FieldPtr->FieldType == EExposedFieldType::Property)
@@ -182,6 +184,24 @@ FName SRCPanelExposedField::GetOwnerName() const
 	return CachedOwnerName;
 }
 
+SRCPanelTreeNode::FMakeNodeWidgetArgs SRCPanelExposedField::CreateEntityWidgetInternal(TSharedPtr<SWidget> ValueWidget, TSharedPtr<SWidget> ResetWidget, const FText& OptionalWarningMessage, TSharedRef<SWidget> EditConditionWidget)
+{
+	FMakeNodeWidgetArgs Args = SRCPanelExposedEntity::CreateEntityWidgetInternal(ValueWidget, ResetWidget, OptionalWarningMessage, EditConditionWidget);
+
+	Args.PropertyIdWidget = SNew(SBox)
+		[
+			SNew(SEditableTextBox)
+			.MinDesiredWidth(50.f)
+			.SelectAllTextWhenFocused(true)
+			.RevertTextOnEscape(true)
+			.ClearKeyboardFocusOnCommit(true)
+			.Text_Lambda([this] () { return GetPropertyIdText(); })
+			.OnTextCommitted(this, &SRCPanelExposedField::OnPropertyIdTextCommitted)
+		];
+
+	return Args;
+}
+
 EExposedFieldType SRCPanelExposedField::GetFieldType() const
 {
 	if (TSharedPtr<FRemoteControlField> Field = WeakField.Pin())
@@ -802,6 +822,36 @@ FReply SRCPanelExposedField::OnClickFunctionButton()
 	return FReply::Handled();
 }
 
+FText SRCPanelExposedField::GetPropertyIdText()
+{
+	if (const TSharedPtr<FRemoteControlField> RCField = WeakField.Pin())
+	{
+		if (RCField->PropertyId != PropertyIdLabel)
+		{
+			PropertyIdLabel = RCField->PropertyId;
+		}
+	}
+	return FText::FromName(PropertyIdLabel);
+}
+
+void SRCPanelExposedField::OnPropertyIdTextCommitted(const FText& InText, ETextCommit::Type InCommitInfo)
+{
+	if (const URemoteControlPreset* RCPreset = Preset.Get())
+	{
+		if (const TSharedPtr<FRemoteControlField> RCField = WeakField.Pin())
+		{
+			const FName NewId = FName(InText.ToString());
+			if (RCField->PropertyId.Compare(NewId) != 0)
+			{
+				RCField->PropertyId = NewId;
+				PropertyIdLabel = NewId;
+				RCPreset->UpdateIdentifiedField(RCField.ToSharedRef());
+				OnPropertyIdRenamed().ExecuteIfBound(NewId);
+			}
+		}
+	}
+}
+
 void SRCPanelFieldChildNode::Construct(const FArguments& InArgs, const TSharedRef<IDetailTreeNode>& InNode, FRCColumnSizeData InColumnSizeData)
 {
 	TArray<TSharedRef<IDetailTreeNode>> ChildNodes;
diff --git a/Engine/Plugins/VirtualProduction/RemoteControl/Source/RemoteControlUI/Private/UI/SRCPanelExposedField.h b/Engine/Plugins/VirtualProduction/RemoteControl/Source/RemoteControlUI/Private/UI/SRCPanelExposedField.h
--- a/Engine/Plugins/VirtualProduction/RemoteControl/Source/RemoteControlUI/Private/UI/SRCPanelExposedField.h
+++ b/Engine/Plugins/VirtualProduction/RemoteControl/Source/RemoteControlUI/Private/UI/SRCPanelExposedField.h
@@ -76,6 +76,10 @@ struct SRCPanelExposedField : public SRCPanelExposedEntity
 	/** Get the owner name of this field */
 	FName GetOwnerName() const;
 
+protected:
+	/** Returns populated args to display this widget. */
+	virtual FMakeNodeWidgetArgs CreateEntityWidgetInternal(TSharedPtr<SWidget> ValueWidget, TSharedPtr<SWidget> ResetWidget = SNullWidget::NullWidget, const FText& OptionalWarningMessage = FText::GetEmpty(), TSharedRef<SWidget> EditConditionWidget = SNullWidget::NullWidget) override;
+
 private:
 	/** Construct a property widget. */
 	TSharedRef<SWidget> ConstructWidget();
@@ -96,6 +100,11 @@ private:
 	/** Handles calling an exposed function.*/
 	FReply OnClickFunctionButton();
 
+	/** Return the text to be displayed in the PropertyId column */
+	FText GetPropertyIdText();
+	/** Handles committing a property id label. */
+	void OnPropertyIdTextCommitted(const FText& InText, ETextCommit::Type InCommitInfo);
+
 private:
 	/** Weak pointer to the underlying RC Field. */
 	TWeakPtr<FRemoteControlField> WeakField;
