[RemoteControl]
Fix:
- Block creation of actions that are not supported by PropertyBag
Since the function that returns the value type is private we have to use a workaround by creating a BagDesc and then retrieve the value type from that to check that it is not None.

diff --git a/Engine/Plugins/VirtualProduction/RemoteControl/Source/RemoteControlLogic/Private/Behaviour/Builtin/Bind/RCBehaviourBind.cpp b/Engine/Plugins/VirtualProduction/RemoteControl/Source/RemoteControlLogic/Private/Behaviour/Builtin/Bind/RCBehaviourBind.cpp
--- a/Engine/Plugins/VirtualProduction/RemoteControl/Source/RemoteControlLogic/Private/Behaviour/Builtin/Bind/RCBehaviourBind.cpp
+++ b/Engine/Plugins/VirtualProduction/RemoteControl/Source/RemoteControlLogic/Private/Behaviour/Builtin/Bind/RCBehaviourBind.cpp
@@ -8,6 +8,7 @@
 #include "Controller/RCCustomControllerUtilities.h"
 #include "Engine/Texture2D.h"
 #include "IRemoteControlPropertyHandle.h"
+#include "PropertyBag.h"
 #include "RCVirtualProperty.h"
 #include "RemoteControlField.h"
 
@@ -84,7 +85,11 @@ bool URCBehaviourBind::CanHaveActionForField(const TSharedPtr<FRemoteControlFiel
 		{
 			if (URCBehaviourBind::CanHaveActionForField(Controller, RCProperty.ToSharedRef(), bAllowNumericInputAsStrings))
 			{
-				return RCProperty->IsEditable();
+				if (const FProperty* Property = RCProperty->GetProperty())
+				{
+					const FPropertyBagPropertyDesc PropertyBagDesc = FPropertyBagPropertyDesc(Property->GetFName(), Property);
+					return RCProperty->IsEditable() && PropertyBagDesc.ValueType != EPropertyBagPropertyType::None;
+				}
 			}
 		}
 	}
diff --git a/Engine/Plugins/VirtualProduction/RemoteControl/Source/RemoteControlLogic/Private/Behaviour/Builtin/Conditional/RCBehaviourConditional.cpp b/Engine/Plugins/VirtualProduction/RemoteControl/Source/RemoteControlLogic/Private/Behaviour/Builtin/Conditional/RCBehaviourConditional.cpp
--- a/Engine/Plugins/VirtualProduction/RemoteControl/Source/RemoteControlLogic/Private/Behaviour/Builtin/Conditional/RCBehaviourConditional.cpp
+++ b/Engine/Plugins/VirtualProduction/RemoteControl/Source/RemoteControlLogic/Private/Behaviour/Builtin/Conditional/RCBehaviourConditional.cpp
@@ -6,6 +6,7 @@
 #include "Action/RCFunctionAction.h"
 #include "Action/RCPropertyAction.h"
 #include "Behaviour/Builtin/Conditional/RCBehaviourConditionalNode.h"
+#include "PropertyBag.h"
 #include "RCVirtualProperty.h"
 #include "RemoteControlField.h"
 #include "Controller/RCController.h"
@@ -84,7 +85,11 @@ bool URCBehaviourConditional::CanHaveActionForField(const TSharedPtr<FRemoteCont
 	{
 		if (const TSharedPtr<FRemoteControlProperty>& RCProperty = StaticCastSharedPtr<FRemoteControlProperty>(InRemoteControlField))
 		{
-			return RCProperty->IsEditable();
+			if (const FProperty* Property = RCProperty->GetProperty())
+			{
+				const FPropertyBagPropertyDesc PropertyBagDesc = FPropertyBagPropertyDesc(Property->GetFName(), Property);
+				return RCProperty->IsEditable() && PropertyBagDesc.ValueType != EPropertyBagPropertyType::None;
+			}
 		}
 	}
 
diff --git a/Engine/Plugins/VirtualProduction/RemoteControl/Source/RemoteControlLogic/Private/Behaviour/Builtin/RangeMap/RCRangeMapBehaviour.cpp b/Engine/Plugins/VirtualProduction/RemoteControl/Source/RemoteControlLogic/Private/Behaviour/Builtin/RangeMap/RCRangeMapBehaviour.cpp
--- a/Engine/Plugins/VirtualProduction/RemoteControl/Source/RemoteControlLogic/Private/Behaviour/Builtin/RangeMap/RCRangeMapBehaviour.cpp
+++ b/Engine/Plugins/VirtualProduction/RemoteControl/Source/RemoteControlLogic/Private/Behaviour/Builtin/RangeMap/RCRangeMapBehaviour.cpp
@@ -13,6 +13,7 @@
 #include "Controller/RCController.h"
 #include "IRemoteControlPropertyHandle.h"
 #include "Kismet/KismetMathLibrary.h"
+#include "PropertyBag.h"
 #include "RCVirtualProperty.h"
 #include "RCVirtualPropertyContainer.h"
 #include "RemoteControlField.h"
@@ -518,7 +519,11 @@ bool URCRangeMapBehaviour::CanHaveActionForField(const TSharedPtr<FRemoteControl
 	{
 		if (const TSharedPtr<FRemoteControlProperty>& RCProperty = StaticCastSharedPtr<FRemoteControlProperty>(InRemoteControlField))
 		{
-			return RCProperty->IsEditable();
+			if (const FProperty* Property = RCProperty->GetProperty())
+			{
+				const FPropertyBagPropertyDesc PropertyBagDesc = FPropertyBagPropertyDesc(Property->GetFName(), Property);
+				return RCProperty->IsEditable() && PropertyBagDesc.ValueType != EPropertyBagPropertyType::None;
+			}
 		}
 	}
 
diff --git a/Engine/Plugins/VirtualProduction/RemoteControl/Source/RemoteControlLogic/Private/Behaviour/RCBehaviour.cpp b/Engine/Plugins/VirtualProduction/RemoteControl/Source/RemoteControlLogic/Private/Behaviour/RCBehaviour.cpp
--- a/Engine/Plugins/VirtualProduction/RemoteControl/Source/RemoteControlLogic/Private/Behaviour/RCBehaviour.cpp
+++ b/Engine/Plugins/VirtualProduction/RemoteControl/Source/RemoteControlLogic/Private/Behaviour/RCBehaviour.cpp
@@ -8,6 +8,7 @@
 #include "Behaviour/RCBehaviourNode.h"
 #include "Controller/RCController.h"
 #include "Engine/Blueprint.h"
+#include "PropertyBag.h"
 #include "RemoteControlField.h"
 #include "RemoteControlPreset.h"
 
@@ -85,7 +86,11 @@ bool URCBehaviour::CanHaveActionForField(const TSharedPtr<FRemoteControlField> I
 	{
 		if (const TSharedPtr<FRemoteControlProperty>& RCProperty = StaticCastSharedPtr<FRemoteControlProperty>(InRemoteControlField))
 		{
-			return RCProperty->IsEditable();
+			if (const FProperty* Property = RCProperty->GetProperty())
+			{
+				const FPropertyBagPropertyDesc PropertyBagDesc = FPropertyBagPropertyDesc(Property->GetFName(), Property);
+				return RCProperty->IsEditable() && PropertyBagDesc.ValueType != EPropertyBagPropertyType::None;
+			}
 		}
 	}
 
