[RemoteControl]
Fix:
- PropId not always working correctly in Game
- Float property will be treated as Double
- Color property will be treated as LinearColor

diff --git a/Engine/Plugins/VirtualProduction/RemoteControl/Source/RemoteControl/Private/PropertyIdHandler/BasePropertyIdHandler.cpp b/Engine/Plugins/VirtualProduction/RemoteControl/Source/RemoteControl/Private/PropertyIdHandler/BasePropertyIdHandler.cpp
--- a/Engine/Plugins/VirtualProduction/RemoteControl/Source/RemoteControl/Private/PropertyIdHandler/BasePropertyIdHandler.cpp
+++ b/Engine/Plugins/VirtualProduction/RemoteControl/Source/RemoteControl/Private/PropertyIdHandler/BasePropertyIdHandler.cpp
@@ -47,9 +47,14 @@ EPropertyBagPropertyType FBasePropertyIdHandler::GetPropertyType(const FProperty
 		{
 			PropertyType = EPropertyBagPropertyType::Int64;
 		}
+		// We treat Float as Double
 		else if (PropertyName == NAME_FloatProperty)
 		{
-			PropertyType = EPropertyBagPropertyType::Float;
+			PropertyType = EPropertyBagPropertyType::Double;
+		}
+		else if (PropertyName == NAME_DoubleProperty)
+		{
+			PropertyType = EPropertyBagPropertyType::Double;
 		}
 		else if (PropertyName == NAME_NameProperty)
 		{
@@ -67,7 +72,12 @@ EPropertyBagPropertyType FBasePropertyIdHandler::GetPropertyType(const FProperty
 	return PropertyType;
 }
 
-FName FBasePropertyIdHandler::GetPropertyTypeName(const FProperty* InProperty) const
+FName FBasePropertyIdHandler::GetPropertySuperTypeName(const FProperty* InProperty) const
+{
+	return GetPropertySubTypeName(InProperty);
+}
+
+FName FBasePropertyIdHandler::GetPropertySubTypeName(const FProperty* InProperty) const
 {
 	if (InProperty)
 	{
@@ -79,9 +89,15 @@ FName FBasePropertyIdHandler::GetPropertyTypeName(const FProperty* InProperty) c
 				return ByteProperty->Enum->GetFName();
 			}
 		}
-		return Property->GetClass()->GetFName();
+		const FName& ClassName = Property->GetClass()->GetFName();
+		if (ClassName == NAME_FloatProperty)
+		{
+			return NAME_DoubleProperty;
+		}
+
+		return ClassName;
 	}
-	return FName(TEXT(""));
+	return NAME_None;
 }
 
 UObject* FBasePropertyIdHandler::GetPropertyTypeObject(const FProperty* InProperty) const
@@ -99,8 +115,3 @@ UObject* FBasePropertyIdHandler::GetPropertyTypeObject(const FProperty* InProper
 	}
 	return nullptr;
 }
-
-TObjectPtr<UObject> FBasePropertyIdHandler::GetObjectPropertyDefaultValue(const FProperty* InProperty, const UClass* InClassToCreate) const
-{
-	return nullptr;
-}
diff --git a/Engine/Plugins/VirtualProduction/RemoteControl/Source/RemoteControl/Private/PropertyIdHandler/BasePropertyIdHandler.h b/Engine/Plugins/VirtualProduction/RemoteControl/Source/RemoteControl/Private/PropertyIdHandler/BasePropertyIdHandler.h
--- a/Engine/Plugins/VirtualProduction/RemoteControl/Source/RemoteControl/Private/PropertyIdHandler/BasePropertyIdHandler.h
+++ b/Engine/Plugins/VirtualProduction/RemoteControl/Source/RemoteControl/Private/PropertyIdHandler/BasePropertyIdHandler.h
@@ -10,8 +10,8 @@ public:
 	//~ Begin IPropertyIdHandler
 	virtual bool IsPropertySupported(const FProperty* InProperty) const override;
 	virtual EPropertyBagPropertyType GetPropertyType(const FProperty* InProperty) const override;
-	virtual FName GetPropertyTypeName(const FProperty* InProperty) const override;
+	virtual FName GetPropertySuperTypeName(const FProperty* InProperty) const override;
+	virtual FName GetPropertySubTypeName(const FProperty* InProperty) const override;
 	virtual UObject* GetPropertyTypeObject(const FProperty* InProperty) const override;
-	virtual TObjectPtr<UObject> GetObjectPropertyDefaultValue(const FProperty* InProperty, const UClass* InClassToCreate) const override;
 	//~ End IPropertyIdHandler
 };
diff --git a/Engine/Plugins/VirtualProduction/RemoteControl/Source/RemoteControl/Private/PropertyIdHandler/EnumPropertyIdHandler.cpp b/Engine/Plugins/VirtualProduction/RemoteControl/Source/RemoteControl/Private/PropertyIdHandler/EnumPropertyIdHandler.cpp
--- a/Engine/Plugins/VirtualProduction/RemoteControl/Source/RemoteControl/Private/PropertyIdHandler/EnumPropertyIdHandler.cpp
+++ b/Engine/Plugins/VirtualProduction/RemoteControl/Source/RemoteControl/Private/PropertyIdHandler/EnumPropertyIdHandler.cpp
@@ -17,7 +17,12 @@ EPropertyBagPropertyType FEnumPropertyIdHandler::GetPropertyType(const FProperty
 	return EPropertyBagPropertyType::Enum;
 }
 
-FName FEnumPropertyIdHandler::GetPropertyTypeName(const FProperty* InProperty) const
+FName FEnumPropertyIdHandler::GetPropertySuperTypeName(const FProperty* InProperty) const
+{
+	return NAME_EnumProperty;
+}
+
+FName FEnumPropertyIdHandler::GetPropertySubTypeName(const FProperty* InProperty) const
 {
 	if (InProperty)
 	{
@@ -27,7 +32,7 @@ FName FEnumPropertyIdHandler::GetPropertyTypeName(const FProperty* InProperty) c
 			return EnumProperty->GetEnum()->GetFName();
 		}
 	}
-	return FName(TEXT(""));
+	return NAME_None;
 }
 
 UObject* FEnumPropertyIdHandler::GetPropertyTypeObject(const FProperty* InProperty) const
@@ -42,8 +47,3 @@ UObject* FEnumPropertyIdHandler::GetPropertyTypeObject(const FProperty* InProper
 	}
 	return nullptr;
 }
-
-TObjectPtr<UObject> FEnumPropertyIdHandler::GetObjectPropertyDefaultValue(const FProperty* InProperty, const UClass* InClassToCreate) const
-{
-	return nullptr;
-}
diff --git a/Engine/Plugins/VirtualProduction/RemoteControl/Source/RemoteControl/Private/PropertyIdHandler/EnumPropertyIdHandler.h b/Engine/Plugins/VirtualProduction/RemoteControl/Source/RemoteControl/Private/PropertyIdHandler/EnumPropertyIdHandler.h
--- a/Engine/Plugins/VirtualProduction/RemoteControl/Source/RemoteControl/Private/PropertyIdHandler/EnumPropertyIdHandler.h
+++ b/Engine/Plugins/VirtualProduction/RemoteControl/Source/RemoteControl/Private/PropertyIdHandler/EnumPropertyIdHandler.h
@@ -10,8 +10,8 @@ public:
 	//~ Begin IPropertyIdHandler
 	virtual bool IsPropertySupported(const FProperty* InProperty) const override;
 	virtual EPropertyBagPropertyType GetPropertyType(const FProperty* InProperty) const override;
-	virtual FName GetPropertyTypeName(const FProperty* InProperty) const override;
+	virtual FName GetPropertySuperTypeName(const FProperty* InProperty) const override;
+	virtual FName GetPropertySubTypeName(const FProperty* InProperty) const override;
 	virtual UObject* GetPropertyTypeObject(const FProperty* InProperty) const override;
-	virtual TObjectPtr<UObject> GetObjectPropertyDefaultValue(const FProperty* InProperty, const UClass* InClassToCreate) const override;
 	//~ End IPropertyIdHandler
 };
diff --git a/Engine/Plugins/VirtualProduction/RemoteControl/Source/RemoteControl/Private/PropertyIdHandler/ObjectPropertyIdHandler.cpp b/Engine/Plugins/VirtualProduction/RemoteControl/Source/RemoteControl/Private/PropertyIdHandler/ObjectPropertyIdHandler.cpp
--- a/Engine/Plugins/VirtualProduction/RemoteControl/Source/RemoteControl/Private/PropertyIdHandler/ObjectPropertyIdHandler.cpp
+++ b/Engine/Plugins/VirtualProduction/RemoteControl/Source/RemoteControl/Private/PropertyIdHandler/ObjectPropertyIdHandler.cpp
@@ -22,40 +22,28 @@ EPropertyBagPropertyType FObjectPropertyIdHandler::GetPropertyType(const FProper
 	return EPropertyBagPropertyType::Object;
 }
 
-FName FObjectPropertyIdHandler::GetPropertyTypeName(const FProperty* InProperty) const
+FName FObjectPropertyIdHandler::GetPropertySuperTypeName(const FProperty* InProperty) const
+{
+	return NAME_ObjectProperty;
+}
+
+FName FObjectPropertyIdHandler::GetPropertySubTypeName(const FProperty* InProperty) const
 {
 	if (InProperty)
 	{
 		const FProperty* Property = GetPropertyInsideContainer(InProperty);
 		if (const FObjectProperty* ObjectProperty = CastField<FObjectProperty>(Property))
 		{
-			return ObjectProperty->PropertyClass->GetFName();
+			if (ObjectProperty->PropertyClass)
+			{
+				return ObjectProperty->PropertyClass->GetFName();
+			}
 		}
 	}
-	return FName(TEXT(""));
+	return NAME_None;
 }
 
 UObject* FObjectPropertyIdHandler::GetPropertyTypeObject(const FProperty* InProperty) const
 {
 	return nullptr;
 }
-
-TObjectPtr<UObject> FObjectPropertyIdHandler::GetObjectPropertyDefaultValue(const FProperty* InProperty, const UClass* InClassToCreate) const
-{
-	if (InClassToCreate)
-	{
-		if (InClassToCreate->IsChildOf(UMaterialInterface::StaticClass()))
-    	{
-    		return UMaterial::GetDefaultMaterial(MD_Surface);
-    	}
-
-		if (InClassToCreate->HasAnyClassFlags(CLASS_Abstract))
-		{
-			// for now skip any AbstractClass
-			return nullptr;
-		}
-
-		return NewObject<UObject>(GetTransientPackage(), InClassToCreate);
-	}
-	return nullptr;
-}
diff --git a/Engine/Plugins/VirtualProduction/RemoteControl/Source/RemoteControl/Private/PropertyIdHandler/ObjectPropertyIdHandler.h b/Engine/Plugins/VirtualProduction/RemoteControl/Source/RemoteControl/Private/PropertyIdHandler/ObjectPropertyIdHandler.h
--- a/Engine/Plugins/VirtualProduction/RemoteControl/Source/RemoteControl/Private/PropertyIdHandler/ObjectPropertyIdHandler.h
+++ b/Engine/Plugins/VirtualProduction/RemoteControl/Source/RemoteControl/Private/PropertyIdHandler/ObjectPropertyIdHandler.h
@@ -10,8 +10,8 @@ public:
 	//~ Begin IPropertyIdHandler
 	virtual bool IsPropertySupported(const FProperty* InProperty) const override;
 	virtual EPropertyBagPropertyType GetPropertyType(const FProperty* InProperty) const override;
-	virtual FName GetPropertyTypeName(const FProperty* InProperty) const override;
+	virtual FName GetPropertySuperTypeName(const FProperty* InProperty) const override;
+	virtual FName GetPropertySubTypeName(const FProperty* InProperty) const override;
 	virtual UObject* GetPropertyTypeObject(const FProperty* InProperty) const override;
-	virtual TObjectPtr<UObject> GetObjectPropertyDefaultValue(const FProperty* InProperty, const UClass* InClassToCreate) const override;
 	//~ End IPropertyIdHandler
 };
diff --git a/Engine/Plugins/VirtualProduction/RemoteControl/Source/RemoteControl/Private/PropertyIdHandler/StructPropertyIdHandler.cpp b/Engine/Plugins/VirtualProduction/RemoteControl/Source/RemoteControl/Private/PropertyIdHandler/StructPropertyIdHandler.cpp
--- a/Engine/Plugins/VirtualProduction/RemoteControl/Source/RemoteControl/Private/PropertyIdHandler/StructPropertyIdHandler.cpp
+++ b/Engine/Plugins/VirtualProduction/RemoteControl/Source/RemoteControl/Private/PropertyIdHandler/StructPropertyIdHandler.cpp
@@ -17,21 +17,30 @@ EPropertyBagPropertyType FStructPropertyIdHandler::GetPropertyType(const FProper
 	return EPropertyBagPropertyType::Struct;
 }
 
-FName FStructPropertyIdHandler::GetPropertyTypeName(const FProperty* InProperty) const
+FName FStructPropertyIdHandler::GetPropertySuperTypeName(const FProperty* InProperty) const
+{
+	return NAME_StructProperty;
+}
+
+FName FStructPropertyIdHandler::GetPropertySubTypeName(const FProperty* InProperty) const
 {
 	if (InProperty)
 	{
 		const FProperty* Property = GetPropertyInsideContainer(InProperty);
 		if (const FStructProperty* StructProperty = CastField<FStructProperty>(Property))
 		{
-			if (StructProperty->Struct->GetFName() == NAME_LinearColor)
+			if (StructProperty->Struct)
 			{
-				return NAME_Color;
+				const FName& StructName = StructProperty->Struct->GetFName();
+				if (StructName == NAME_Color)
+				{
+					return NAME_LinearColor;
+				}
+				return StructName;
 			}
-			return StructProperty->Struct->GetFName();
 		}
 	}
-	return FName(TEXT(""));
+	return NAME_None;
 }
 
 UObject* FStructPropertyIdHandler::GetPropertyTypeObject(const FProperty* InProperty) const
@@ -41,17 +50,12 @@ UObject* FStructPropertyIdHandler::GetPropertyTypeObject(const FProperty* InProp
 		const FProperty* Property = GetPropertyInsideContainer(InProperty);
 		if (const FStructProperty* StructProperty = CastField<FStructProperty>(Property))
 		{
-			if (StructProperty->Struct->GetFName() == NAME_LinearColor)
+			if (StructProperty->Struct->GetFName() == NAME_Color)
 			{
-				return TBaseStructure<FColor>::Get();
+				return TBaseStructure<FLinearColor>::Get();
 			}
 			return StructProperty->Struct;
 		}
 	}
 	return nullptr;
 }
-
-TObjectPtr<UObject> FStructPropertyIdHandler::GetObjectPropertyDefaultValue(const FProperty* InProperty, const UClass* InClassToCreate) const
-{
-	return nullptr;
-}
diff --git a/Engine/Plugins/VirtualProduction/RemoteControl/Source/RemoteControl/Private/PropertyIdHandler/StructPropertyIdHandler.h b/Engine/Plugins/VirtualProduction/RemoteControl/Source/RemoteControl/Private/PropertyIdHandler/StructPropertyIdHandler.h
--- a/Engine/Plugins/VirtualProduction/RemoteControl/Source/RemoteControl/Private/PropertyIdHandler/StructPropertyIdHandler.h
+++ b/Engine/Plugins/VirtualProduction/RemoteControl/Source/RemoteControl/Private/PropertyIdHandler/StructPropertyIdHandler.h
@@ -10,8 +10,8 @@ public:
 	//~ Begin IPropertyIdHandler
 	virtual bool IsPropertySupported(const FProperty* InProperty) const override;
 	virtual EPropertyBagPropertyType GetPropertyType(const FProperty* InProperty) const override;
-	virtual FName GetPropertyTypeName(const FProperty* InProperty) const override;
+	virtual FName GetPropertySuperTypeName(const FProperty* InProperty) const override;
+	virtual FName GetPropertySubTypeName(const FProperty* InProperty) const override;
 	virtual UObject* GetPropertyTypeObject(const FProperty* InProperty) const override;
-	virtual TObjectPtr<UObject> GetObjectPropertyDefaultValue(const FProperty* InProperty, const UClass* InClassToCreate) const override;
 	//~ End IPropertyIdHandler
 };
diff --git a/Engine/Plugins/VirtualProduction/RemoteControl/Source/RemoteControl/Private/RCVirtualProperty.cpp b/Engine/Plugins/VirtualProduction/RemoteControl/Source/RemoteControl/Private/RCVirtualProperty.cpp
--- a/Engine/Plugins/VirtualProduction/RemoteControl/Source/RemoteControl/Private/RCVirtualProperty.cpp
+++ b/Engine/Plugins/VirtualProduction/RemoteControl/Source/RemoteControl/Private/RCVirtualProperty.cpp
@@ -210,6 +210,17 @@ bool URCVirtualPropertyBase::IsColorType() const
 	return false;
 }
 
+bool URCVirtualPropertyBase::IsLinearColorType() const
+{
+	const FProperty* Property = GetProperty();
+	if (const FStructProperty* StructProperty = CastField<FStructProperty>(Property))
+	{
+		return StructProperty->Struct == TBaseStructure<FLinearColor>::Get();
+	}
+
+	return false;
+}
+
 bool URCVirtualPropertyBase::IsRotatorType() const
 {
 	const FProperty* Property = GetProperty();
@@ -549,6 +560,11 @@ bool URCVirtualPropertyBase::GetValueColor(FColor& OutColor) const
 	return FRCVirtualPropertyCastHelpers::GetStructValue<FColor>(this, TBaseStructure<FColor>::Get(), &OutColor);
 }
 
+bool URCVirtualPropertyBase::GetValueLinearColor(FLinearColor& OutLinearColor) const
+{
+	return FRCVirtualPropertyCastHelpers::GetStructValue<FLinearColor>(this, TBaseStructure<FLinearColor>::Get(), &OutLinearColor);
+}
+
 UObject* URCVirtualPropertyBase::GetValueObject() const
 {
 	return FRCVirtualPropertyCastHelpers::GetObjectValue(this);
@@ -767,6 +783,11 @@ bool URCVirtualPropertyBase::SetValueColor(const FColor& InColor)
 	return FRCVirtualPropertyCastHelpers::SetStructValue<FColor>(this, TBaseStructure<FColor>::Get(), InColor);
 }
 
+bool URCVirtualPropertyBase::SetValueLinearColor(const FLinearColor& InLinearColor)
+{
+	return FRCVirtualPropertyCastHelpers::SetStructValue<FLinearColor>(this, TBaseStructure<FLinearColor>::Get(), InLinearColor);
+}
+
 FName URCVirtualPropertyBase::GetPropertyName() const
 {
 	return PropertyName;
diff --git a/Engine/Plugins/VirtualProduction/RemoteControl/Source/RemoteControl/Private/RemoteControlPropertyIdRegistry.cpp b/Engine/Plugins/VirtualProduction/RemoteControl/Source/RemoteControl/Private/RemoteControlPropertyIdRegistry.cpp
--- a/Engine/Plugins/VirtualProduction/RemoteControl/Source/RemoteControl/Private/RemoteControlPropertyIdRegistry.cpp
+++ b/Engine/Plugins/VirtualProduction/RemoteControl/Source/RemoteControl/Private/RemoteControlPropertyIdRegistry.cpp
@@ -2,20 +2,19 @@
 
 #include "RemoteControlPropertyIdRegistry.h"
 
+#include "Backends/CborStructDeserializerBackend.h"
+#include "Backends/CborStructSerializerBackend.h"
+#include "IPropertyIdHandler.h"
 #include "IRemoteControlModule.h"
+#include "IStructDeserializerBackend.h"
 #include "Materials/MaterialInterface.h"
 #include "RCVirtualProperty.h"
 #include "RemoteControlPreset.h"
+#include "StructSerializer.h"
 #include "UObject/Field.h"
 
-#if !WITH_EDITOR
-#include "Backends/CborStructDeserializerBackend.h"
-#include "Backends/CborStructSerializerBackend.h"
-#include "IStructDeserializerBackend.h"
-#include "StructSerializer.h"
-#endif
 
-namespace RemoteControlPropertyIdUtilities
+namespace RCPropIdUtils
 {
 	bool GetObjectRef(const TSharedPtr<FRemoteControlProperty>& Field, const ERCAccess Access, FRCObjectReference& OutObjectRef)
 	{
@@ -36,105 +35,6 @@ namespace RemoteControlPropertyIdUtilities
 		}
 		return false;
 	}
-
-	bool CopyNonUClassOwnerArray(const TObjectPtr<URCVirtualPropertySelfContainer>& InVirtualPropertySelfContainer,
-		const FArrayProperty* InArrayProperty,
-		const TSharedPtr<FRemoteControlProperty>& InTargetRCProperty,
-		FProperty* InProperty,
-		FRCFieldPathInfo InSegment,
-		const int32 InIndexSegment,
-		const uint8* InPropertyContainer,
-		const bool bCheckByteEnumComparison,
-		uint8** OutPtrContainer)
-	{
-		FScriptArrayHelper ArrayHelper(InArrayProperty, InPropertyContainer);
-		if (const TSharedPtr<FRemoteControlField> TargetRCField = StaticCastSharedPtr<FRemoteControlField>(InTargetRCProperty);
-			!TargetRCField->FieldPathInfo.Segments.IsEmpty())
-		{
-			if (!OutPtrContainer)
-            {
-            	return false;
-            }
-
-			const int32 Index = InSegment.Segments[InIndexSegment].ArrayIndex;
-			*OutPtrContainer = ArrayHelper.GetElementPtr(Index);
-			if (InIndexSegment < InSegment.GetSegmentCount() - 2)
-			{
-				return false;
-			}
-			if (const FArrayProperty* ArrayRealProperty = CastField<FArrayProperty>(InProperty))
-			{
-				return InVirtualPropertySelfContainer->CopyCompleteValue(ArrayRealProperty->Inner, *OutPtrContainer, bCheckByteEnumComparison);
-			}
-		}
-		return false;
-	}
-	
-	bool CopyNonUClassOwnerMap(const TObjectPtr<URCVirtualPropertySelfContainer>& InVirtualPropertySelfContainer,
-		const FMapProperty* InMapProperty,
-		const TSharedPtr<FRemoteControlProperty>& InTargetRCProperty,
-		FProperty* InProperty,
-		FRCFieldPathInfo InSegment,
-		const int32 InIndexSegment,
-		const uint8* InPropertyContainer,
-		const bool bCheckByteEnumComparison,
-		uint8** OutPtrContainer)
-	{
-		FScriptMapHelper MapHelper(InMapProperty, InPropertyContainer);
-		if (const TSharedPtr<FRemoteControlField> TargetRCField = StaticCastSharedPtr<FRemoteControlField>(InTargetRCProperty);
-			!TargetRCField->FieldPathInfo.Segments.IsEmpty())
-		{
-			if (!OutPtrContainer)
-			{
-				return false;
-			}
-
-			const int32 Index = InSegment.Segments[InIndexSegment].ArrayIndex;
-			*OutPtrContainer = MapHelper.GetValuePtr(Index);
-			if (InIndexSegment < InSegment.GetSegmentCount() - 2)
-			{
-				return false;
-			}
-			if (const FMapProperty* MapRealProperty = CastField<FMapProperty>(InProperty))
-			{
-				return InVirtualPropertySelfContainer->CopyCompleteValue(MapRealProperty->ValueProp, *OutPtrContainer, bCheckByteEnumComparison);
-			}
-		}
-		return false;
-	}
-
-	bool CopyNonUClassOwnerSet(const TObjectPtr<URCVirtualPropertySelfContainer>& InVirtualPropertySelfContainer,
-		const FSetProperty* InSetProperty,
-		const TSharedPtr<FRemoteControlProperty>& InTargetRCProperty,
-		FProperty* InProperty,
-		FRCFieldPathInfo InSegment,
-		const int32 InIndexSegment,
-		const uint8* InPropertyContainer,
-		const bool bCheckByteEnumComparison,
-		uint8** OutPtrContainer)
-	{
-		FScriptSetHelper SetHelper(InSetProperty, InPropertyContainer);
-		if (const TSharedPtr<FRemoteControlField> TargetRCField = StaticCastSharedPtr<FRemoteControlField>(InTargetRCProperty);
-			!TargetRCField->FieldPathInfo.Segments.IsEmpty())
-		{
-			if (!OutPtrContainer)
-			{
-				return false;
-			}
-
-			const int32 Index = InSegment.Segments[InIndexSegment].ArrayIndex;
-			*OutPtrContainer = SetHelper.GetElementPtr(Index);
-			if (InIndexSegment < InSegment.GetSegmentCount() - 2)
-			{
-				return false;
-			}
-			if (const FSetProperty* SetRealProperty = CastField<FSetProperty>(InProperty))
-			{
-				return InVirtualPropertySelfContainer->CopyCompleteValue(SetRealProperty->ElementProp, *OutPtrContainer, bCheckByteEnumComparison);
-			}
-		}
-		return false;
-	}
 }
 
 FRCPropertyIdWrapper::FRCPropertyIdWrapper(const TSharedRef<FRemoteControlProperty>& InRCProperty)
@@ -176,11 +76,6 @@ const FName& FRCPropertyIdWrapper::GetSubType() const
 	return SubType;
 }
 
-const UClass* FRCPropertyIdWrapper::GetClassToCreate() const
-{
-	return ClassToCreate;
-}
-
 bool FRCPropertyIdWrapper::IsValid() const
 {
 	return EntityId.IsValid() && (!SuperType.IsNone() || !SubType.IsNone());
@@ -216,109 +111,23 @@ bool FRCPropertyIdWrapper::operator==(const FRCPropertyIdWrapper& Other) const
 
 void FRCPropertyIdWrapper::UpdateTypes(const TSharedRef<FRemoteControlProperty>& InRCProperty, UObject* InOwner, FProperty* InProperty)
 {
-	SuperType = InProperty->GetClass()->GetFName();
+	SuperType = NAME_None;
 	SubType = NAME_None;
-	ClassToCreate = nullptr;
-	if (const FEnumProperty* EnumProperty = CastField<FEnumProperty>(InProperty))
-	{
-		if (EnumProperty->GetEnum())
-		{
-			SubType = EnumProperty->GetEnum()->GetFName();
-		}
-	}
-	else if (const FByteProperty* ByteProperty = CastField<FByteProperty>(InProperty))
-	{
-		if (ByteProperty->Enum)
-		{
-			SuperType = NAME_EnumProperty;
-			SubType = ByteProperty->Enum->GetFName();
-		}
-	}
-	else if (const FStructProperty* StructProperty = CastField<FStructProperty>(InProperty))
-	{
-		if (StructProperty->Struct->GetFName() == NAME_LinearColor)
-		{
-			SubType = NAME_Color;
-		}
-		else
-		{
-			SubType = StructProperty->Struct->GetFName();
-		}
-	}
-	else if (const FObjectProperty* ObjectProperty = CastField<FObjectProperty>(InProperty))
+
+	const TSharedPtr<IPropertyIdHandler> PropertyIdHandler = IRemoteControlModule::Get().GetPropertyIdHandlerFor(InProperty);
+	if (!PropertyIdHandler.IsValid())
 	{
-		// ContainerPtrToValuePtr will crash if the GetOwner<UClass> is nullptr so we check before.
-		if (ObjectProperty->GetOwner<UClass>())
-		{
-			if (const uint8* PropertyValuePtr = ObjectProperty->ContainerPtrToValuePtr<uint8>(InOwner))
-			{
-				if (const TObjectPtr<UObject> Object = ObjectProperty->GetPropertyValue(PropertyValuePtr))
-				{
-					if (Object && Object->IsValidLowLevel())
-					{
-						if (Object->IsA(UMaterialInterface::StaticClass()))
-						{
-							SubType = UMaterialInterface::StaticClass()->GetFName();
-							ClassToCreate = UMaterialInterface::StaticClass();
-						}
-						else
-						{
-							SubType = Object->GetClass()->GetFName();
-							ClassToCreate = Object->GetClass();
-						}
-					}
-				}
-			}
-		}
+		return;
 	}
-	else if (const FArrayProperty* ArrayProperty = CastField<FArrayProperty>(InProperty))
-	{
-		SuperType = ArrayProperty->Inner->GetClass()->GetFName();
 
-		// ContainerPtrToValuePtr will crash if the GetOwner<UClass> is nullptr so we check before.
-		if (ArrayProperty->GetOwner<UClass>())
-		{
-			if (const uint8* PropertyValuePtr = ArrayProperty->ContainerPtrToValuePtr<uint8>(InOwner))
-			{
-				if (const FObjectProperty* InnerObjectProperty = CastField<FObjectProperty>(ArrayProperty->Inner))
-				{
-					FScriptArrayHelper ArrayHelper(ArrayProperty, PropertyValuePtr);
-					
-					if (const TSharedRef<FRemoteControlField> TargetRCField = StaticCastSharedRef<FRemoteControlField>(InRCProperty); !TargetRCField->FieldPathInfo.Segments.IsEmpty())
-					{
-						const int32 ArrayIndex = TargetRCField->FieldPathInfo.Segments[0].ArrayIndex;
-						if (ArrayIndex != INDEX_NONE)
-						{
-							const uint8* ObjPtrContainer = ArrayHelper.GetRawPtr(ArrayIndex);
-							if (const UObject* CurrentObject = InnerObjectProperty->GetObjectPropertyValue(ObjPtrContainer))
-							{
-								if (CurrentObject->IsA(UMaterialInterface::StaticClass()))
-								{
-									SubType = UMaterialInterface::StaticClass()->GetFName();
-									ClassToCreate = UMaterialInterface::StaticClass();
-								}
-								else
-								{
-									SubType = CurrentObject->GetClass()->GetFName();
-									ClassToCreate = CurrentObject->GetClass();
-								}
-							}
-						}
-					}
-				}
-			}
-		}
-		else if (const FStructProperty* InnerStructProperty = CastField<FStructProperty>(ArrayProperty->Inner))
-		{
-			if (InnerStructProperty->Struct->GetFName() == NAME_LinearColor)
-			{
-				SubType = NAME_Color;
-			}
-			else
-			{
-				SubType = InnerStructProperty->Struct->GetFName();
-			}
-		}
+	SuperType = PropertyIdHandler->GetPropertySuperTypeName(InProperty);
+	const EPropertyBagPropertyType PropertyType = PropertyIdHandler->GetPropertyType(InProperty);
+
+	if (PropertyType == EPropertyBagPropertyType::Enum ||
+		PropertyType == EPropertyBagPropertyType::Object ||
+		PropertyType == EPropertyBagPropertyType::Struct)
+	{
+		SubType = PropertyIdHandler->GetPropertySubTypeName(InProperty);
 	}
 }
 
@@ -345,6 +154,14 @@ URemoteControlPreset* URemoteControlPropertyIdRegistry::GetSourcePreset() const
 	return GetTypedOuter<URemoteControlPreset>();
 }
 
+#if WITH_EDITOR
+void URemoteControlPropertyIdRegistry::PostEditChangeProperty(FPropertyChangedEvent& InPropertyChangedEvent)
+{
+	UObject::PostEditChangeProperty(InPropertyChangedEvent);
+	Initialize();
+}
+#endif
+
 void URemoteControlPropertyIdRegistry::PerformChainReaction(const FRemoteControlPropertyIdArgs& InArgs)
 {
 	URemoteControlPreset* SourcePreset = GetSourcePreset();
@@ -373,150 +190,69 @@ void URemoteControlPropertyIdRegistry::PerformChainReaction(const FRemoteControl
 	{
 		if (TSharedPtr<FRemoteControlProperty> TargetRCProperty = SourcePreset->GetExposedEntity<FRemoteControlProperty>(TargetProperty).Pin())
 		{
-			if (UObject* BoundObject = TargetRCProperty->GetBoundObject())
+			const TObjectPtr<URCVirtualPropertySelfContainer>* RealPropContainer = InArgs.RealProperties.Find(TargetRCProperty->GetId());
+			if (!RealPropContainer)
+			{
+				 continue;
+			}
+
+			TObjectPtr<URCVirtualPropertySelfContainer> VirtualProperty = InArgs.VirtualProperty;
+			if (!VirtualProperty)
 			{
-				bool bCopyComplete = false;
+				continue;
+			}
 
-				if (FProperty* Property = TargetRCProperty->GetProperty())
+			// We copy the VirtualProperty value to the RealPropContainer to make sure it is set and updated correctly through the serialization
+			if ((*RealPropContainer)->GetValueType() == VirtualProperty->GetValueType())
+			{
+				// Special case for Color
+				if ((*RealPropContainer)->IsColorType() && VirtualProperty->IsLinearColorType())
 				{
-#if WITH_EDITOR
-					BoundObject->PreEditChange(Property);
-					
-					BoundObject->Modify();
-#endif // WITH_EDITOR
-					// Note : For all other object types except materials.
-					if (FObjectProperty* ObjectProperty = CastField<FObjectProperty>(Property))
-					{
-						if (uint8* PropertyValuePtr = ObjectProperty->ContainerPtrToValuePtr<uint8>(BoundObject))
-						{
-							UObject* CurrentObject = ObjectProperty->GetObjectPropertyValue(PropertyValuePtr);
-							if (CurrentObject && CurrentObject->IsA(InArgs.SourceClass))
-							{
-								ObjectProperty->SetObjectPropertyValue(PropertyValuePtr, InArgs.SourceObject);
-								bCopyComplete = true;
-							}
-						}
-					}
-					// Note : Specialization for Materials.
-					else if (FArrayProperty* ArrayProperty = CastField<FArrayProperty>(Property))
-					{
-						if (ArrayProperty->GetOwner<UClass>())
-						{
-							if (uint8* PropertyValuePtr = ArrayProperty->ContainerPtrToValuePtr<uint8>(BoundObject))
-							{
-								if (FObjectProperty* InnerObjectProperty = CastField<FObjectProperty>(ArrayProperty->Inner))
-								{
-									FScriptArrayHelper ArrayHelper(ArrayProperty, PropertyValuePtr);
-									if (TSharedPtr<FRemoteControlField> TargetRCField = StaticCastSharedPtr<FRemoteControlField>(TargetRCProperty); !TargetRCField->FieldPathInfo.Segments.IsEmpty())
-									{
-										const int32 MaterialIndex = TargetRCField->FieldPathInfo.Segments[0].ArrayIndex;
-										uint8* ObjPtrContainer = ArrayHelper.GetRawPtr(MaterialIndex);
-										UObject* CurrentObject = InnerObjectProperty->GetObjectPropertyValue(ObjPtrContainer);
-										if (CurrentObject && CurrentObject->IsA(InArgs.SourceClass))
-										{
-											InnerObjectProperty->SetObjectPropertyValue(ObjPtrContainer, InArgs.SourceObject);
-											bCopyComplete = true;
-										}
-									}
-								}
-							}
-						}
-					}
-					// Note : For primitive types not in container and that has a valid UClass Owner.
-					// ContainerPtrToValuePtr will crash if the GetOwner<UClass> is nullptr so we check before.
-					if (!bCopyComplete && Property->GetOwner<UClass>() != nullptr &&
-						!Property->IsA<FArrayProperty>() && !Property->IsA<FSetProperty>() && !Property->IsA<FMapProperty>())
-					{
-						if (uint8* PropertyValuePtr = Property->ContainerPtrToValuePtr<uint8>(BoundObject))
-						{
-							if (const FStructProperty* StructProperty = CastField<FStructProperty>(Property))
-							{
-								// FLinearColor are treated as FColor to avoid creating 2 different widget for them and avoid confusion
-								// Instead of the normal CopyCompleteValue we use this for FLinearColor
-								if (StructProperty->Struct->GetFName() == NAME_LinearColor)
-								{
-									FColor ColorValue;
-									InArgs.VirtualProperty->GetValueColor(ColorValue);
-									const FLinearColor RealValue(ColorValue);
-									Property->CopyCompleteValue(PropertyValuePtr, &RealValue);
-									bCopyComplete = true;
-								}
-							}
-							// We do this to avoid copying it above and here
-							if (bCopyComplete == false)
-							{
-								bCopyComplete = InArgs.VirtualProperty->CopyCompleteValue(Property, PropertyValuePtr, true);
-							}
-						}
-					}
-					//Note : For all the other cases, Container and UStruct Owner.
-					else if (!bCopyComplete)
-					{
-#if !WITH_EDITOR
-						StructMemoryContainer = nullptr;
-#endif // !WITH_EDITOR
-						bCopyComplete = CopyNonUClassOwnerProperty(InArgs, TargetRCProperty, Property, BoundObject);
-					}
-					if (bCopyComplete)
-					{
-#if WITH_EDITOR
-						TargetRCProperty->FieldPathInfo.Resolve(BoundObject);
-						FPropertyChangedEvent PropertyEvent = TargetRCProperty->FieldPathInfo.ToPropertyChangedEvent(EPropertyChangeType::ValueSet);
-						FEditPropertyChain EditPropertyChain;
-						TargetRCProperty->FieldPathInfo.ToEditPropertyChain(EditPropertyChain);
-						if (EditPropertyChain.IsEmpty())
-						{
-							BoundObject->PostEditChangeProperty(PropertyEvent);
-						}
-						else
-						{
-							FPropertyChangedChainEvent ChainEvent(EditPropertyChain, PropertyEvent);
-							TArray<TMap<FString, int32>> ArrayIndicesPerObject;
-							{
-								TMap<FString, int32> ArrayIndices;
-								ArrayIndices.Reserve(TargetRCProperty->FieldPathInfo.Segments.Num());
-								for (const FRCFieldPathSegment& Segment : TargetRCProperty->FieldPathInfo.Segments)
-								{
-									ArrayIndices.Add(Segment.Name.ToString(), Segment.ArrayIndex);
-								}
-								ArrayIndicesPerObject.Add(MoveTemp(ArrayIndices));
-							}
-							ChainEvent.ObjectIteratorIndex = 0;
-							ChainEvent.SetArrayIndexPerObject(ArrayIndicesPerObject);
-							BoundObject->PostEditChangeChainProperty(ChainEvent);
-						}
-						
-#else // NOTE: During runtime use Serialization API to update the property properly.
-						TArray<uint8> Buffer;
-						FMemoryWriter Writer(Buffer);
-						FCborStructSerializerBackend WriterBackend(Writer, EStructSerializerBackendFlags::Default);
-						FStructSerializerPolicies Policies;
-						Policies.MapSerialization = EStructSerializerMapPolicies::Array;
-
-						// We do this at runtime because otherwise Property Owned by UStruct won't update.
-						// So during the copy we take the StructContainer and use that to serialize the property.
-						if (StructMemoryContainer != nullptr)
-						{
-							FStructSerializer::SerializeElement(StructMemoryContainer, Property, INDEX_NONE, WriterBackend, Policies);
-						}
-						else
-						{
-							FStructSerializer::SerializeElement(BoundObject, Property, INDEX_NONE, WriterBackend, Policies);
-						}
-						// Deserialization
-						FMemoryReader Reader(Buffer);
-						FCborStructDeserializerBackend ReaderBackend(Reader);
-						FRCObjectReference TargetObjectRef;
-						if (!RemoteControlPropertyIdUtilities::GetObjectRef(TargetRCProperty, ERCAccess::WRITE_ACCESS, TargetObjectRef))
-						{
-							continue;
-						}
-						if (!IRemoteControlModule::Get().SetObjectProperties(TargetObjectRef, ReaderBackend, ERCPayloadType::Cbor, Buffer))
-						{
-							continue;
-						}
-#endif // WITH_EDITOR
-					}
+					FLinearColor LinearColor;
+					VirtualProperty->GetValueLinearColor(LinearColor);
+					FColor ValueToAssign = LinearColor.ToFColor(true);
+					(*RealPropContainer)->SetValueColor(ValueToAssign);
+				}
+				else
+				{
+					(*RealPropContainer)->UpdateValueWithProperty(InArgs.VirtualProperty);
+				}
+			}
+			else if ((*RealPropContainer)->GetValueType() == EPropertyBagPropertyType::Float &&
+				VirtualProperty->GetValueType() == EPropertyBagPropertyType::Double)
+			{
+				// Special case for Float
+				double Double;
+				VirtualProperty->GetValueDouble(Double);
+				float ValueToAssign = Double;
+				(*RealPropContainer)->SetValueFloat(ValueToAssign);
+			}
+
+			// Update the property correctly through the serialization
+			FRCObjectReference ObjectRef;
+			ObjectRef.Property = TargetRCProperty->GetProperty();
+			ObjectRef.Access = ERCAccess::WRITE_ACCESS;
+			ObjectRef.PropertyPathInfo = TargetRCProperty->FieldPathInfo.ToString();
+
+			for (UObject* Object : TargetRCProperty->GetBoundObjects())
+			{
+				if (IRemoteControlModule::Get().ResolveObjectProperty(ObjectRef.Access, Object, ObjectRef.PropertyPathInfo, ObjectRef))
+				{
+					TArray<uint8> Buffer;
+					FMemoryWriter Writer(Buffer);
+
+					FCborStructSerializerBackend SerializerBackend(Writer, EStructSerializerBackendFlags::Default);
+
+					FStructSerializerPolicies Policies; 
+					Policies.MapSerialization = EStructSerializerMapPolicies::Array;
+
+					(*RealPropContainer)->SerializeToBackend(SerializerBackend);
+
+					// Deserialization
+					FMemoryReader Reader(Buffer);
+					FCborStructDeserializerBackend DeserializerBackend(Reader);
+
+					IRemoteControlModule::Get().SetObjectProperties(ObjectRef, DeserializerBackend, ERCPayloadType::Cbor, Buffer);
 				}
 			}
 		}
@@ -534,6 +270,9 @@ void URemoteControlPropertyIdRegistry::AddIdentifiedField(const TSharedRef<FRemo
 				FRCPropertyIdWrapper Wrapper{ RCProperty.ToSharedRef() };
 				if (Wrapper.IsValid())
 				{
+#if WITH_EDITOR
+					Modify();
+#endif
 					IdentifiedFields.Add(MoveTemp(Wrapper));
 				}
 			}
@@ -565,21 +304,13 @@ void URemoteControlPropertyIdRegistry::RemoveIdentifiedField(const FGuid& InEnti
 
 	if (IdentifiedFields.ContainsByHash(Hash, InEntityId))
 	{
+#if WITH_EDITOR
+		Modify();
+#endif
 		IdentifiedFields.RemoveByHash(Hash, InEntityId);
 	}
 }
 
-const UClass* URemoteControlPropertyIdRegistry::GetClassByEntityId(const FGuid& InEntityId) const
-{
-	const uint32 Hash = GetTypeHash(InEntityId);
-
-	if (const FRCPropertyIdWrapper* Wrapper = IdentifiedFields.FindByHash(Hash, InEntityId))
-	{
-		return Wrapper->GetClassToCreate();
-	}
-	return nullptr;
-}
-
 TSet<FName> URemoteControlPropertyIdRegistry::GetFieldIdsNameList() const
 {
 	TSet<FName> OutIds;
@@ -696,139 +427,3 @@ void URemoteControlPropertyIdRegistry::UpdateEntityIds(const TMap<FGuid, FGuid>&
 		IdentifiedFields = MoveTemp(RehashedIdentifiedFields);
 	}
 }
-
-bool URemoteControlPropertyIdRegistry::CopyNonUClassOwnerProperty(const FRemoteControlPropertyIdArgs& InArgs, TSharedPtr<FRemoteControlProperty> TargetRCProperty, FProperty* InProperty, UObject* BoundObject)
-{
-	bool bCheckByteEnumComparison = false;
-	if (InArgs.VirtualProperty->GetProperty())
-	{
-		bCheckByteEnumComparison = InArgs.VirtualProperty->GetProperty()->IsA<FEnumProperty>();
-	}
-	TargetRCProperty->FieldPathInfo.Resolve(BoundObject);
-	return TryCopyNonUClassOwnerProperty(InArgs.VirtualProperty, TargetRCProperty, InProperty, BoundObject, TargetRCProperty->FieldPathInfo, 0, nullptr, bCheckByteEnumComparison);
-}
-
-bool URemoteControlPropertyIdRegistry::TryCopyNonUClassOwnerProperty(const TObjectPtr<URCVirtualPropertySelfContainer>& InVirtualPropertySelfContainer, const TSharedPtr<FRemoteControlProperty>& InTargetRCProperty, FProperty* InProperty, UObject* InBoundObject, FRCFieldPathInfo InSegment, const int32 InIndexSegment, uint8* InPropertyContainer, const bool bCheckByteEnumComparison)
-{
-	if (InIndexSegment >= InSegment.GetSegmentCount())
-	{
-		return false;
-	}
-
-	if (const FArrayProperty* ArrayProperty = CastField<FArrayProperty>(InSegment.Segments[InIndexSegment].ResolvedData.Field))
-	{
-		if (ArrayProperty->GetOwner<UClass>())
-		{
-			if (const uint8* PropertyValuePtr = ArrayProperty->ContainerPtrToValuePtr<uint8>(InBoundObject))
-			{
-				uint8* PtrContainer = nullptr;
-				if (RemoteControlPropertyIdUtilities::CopyNonUClassOwnerArray(InVirtualPropertySelfContainer, ArrayProperty, InTargetRCProperty, InProperty, InSegment, InIndexSegment, PropertyValuePtr, bCheckByteEnumComparison, &PtrContainer))
-				{
-					return true;
-				}
-				return TryCopyNonUClassOwnerProperty(InVirtualPropertySelfContainer, InTargetRCProperty, InProperty, InBoundObject, InSegment, InIndexSegment + 1, PtrContainer, bCheckByteEnumComparison);
-			}
-		}
-		else if (InPropertyContainer)
-		{
-			uint8* PtrContainer = nullptr;
-			if (RemoteControlPropertyIdUtilities::CopyNonUClassOwnerArray(InVirtualPropertySelfContainer, ArrayProperty, InTargetRCProperty, InProperty, InSegment, InIndexSegment, InPropertyContainer, bCheckByteEnumComparison, &PtrContainer))
-			{
-				return true;
-			}
-			return TryCopyNonUClassOwnerProperty(InVirtualPropertySelfContainer, InTargetRCProperty, InProperty, InBoundObject, InSegment, InIndexSegment + 1, PtrContainer, bCheckByteEnumComparison);
-		}
-	}
-	else if (const FSetProperty* SetProperty = CastField<FSetProperty>(InSegment.Segments[InIndexSegment].ResolvedData.Field))
-	{
-		if (SetProperty->GetOwner<UClass>())
-		{
-			if (const uint8* PropertyValuePtr = SetProperty->ContainerPtrToValuePtr<uint8>(InBoundObject))
-			{
-				uint8* PtrContainer = nullptr;
-				if (RemoteControlPropertyIdUtilities::CopyNonUClassOwnerSet(InVirtualPropertySelfContainer, SetProperty, InTargetRCProperty, InProperty, InSegment, InIndexSegment, PropertyValuePtr, bCheckByteEnumComparison, &PtrContainer))
-				{
-					return true;
-				}
-				return TryCopyNonUClassOwnerProperty(InVirtualPropertySelfContainer, InTargetRCProperty, InProperty, InBoundObject, InSegment, InIndexSegment + 1, PtrContainer, bCheckByteEnumComparison);
-			}
-		}
-		else if (InPropertyContainer)
-		{
-			uint8* PtrContainer = nullptr;
-			if (RemoteControlPropertyIdUtilities::CopyNonUClassOwnerSet(InVirtualPropertySelfContainer, SetProperty, InTargetRCProperty, InProperty, InSegment, InIndexSegment, InPropertyContainer, bCheckByteEnumComparison, &PtrContainer))
-			{
-				return true;
-			}
-			return TryCopyNonUClassOwnerProperty(InVirtualPropertySelfContainer, InTargetRCProperty, InProperty, InBoundObject, InSegment, InIndexSegment + 1, PtrContainer, bCheckByteEnumComparison);
-		}
-	}
-	else if (const FMapProperty* MapProperty = CastField<FMapProperty>(InSegment.Segments[InIndexSegment].ResolvedData.Field))
-	{
-		if (MapProperty->GetOwner<UClass>())
-		{
-			if (const uint8* PropertyValuePtr = MapProperty->ContainerPtrToValuePtr<uint8>(InBoundObject))
-			{
-				uint8* PtrContainer = nullptr;
-				if (RemoteControlPropertyIdUtilities::CopyNonUClassOwnerMap(InVirtualPropertySelfContainer, MapProperty, InTargetRCProperty, InProperty, InSegment, InIndexSegment, PropertyValuePtr, bCheckByteEnumComparison, &PtrContainer))
-				{
-					return true;
-				}
-				return TryCopyNonUClassOwnerProperty(InVirtualPropertySelfContainer, InTargetRCProperty, InProperty, InBoundObject, InSegment, InIndexSegment + 1, PtrContainer, bCheckByteEnumComparison);
-			}
-		}
-		else if (InPropertyContainer)
-		{
-			uint8* PtrContainer = nullptr;
-			if (RemoteControlPropertyIdUtilities::CopyNonUClassOwnerMap(InVirtualPropertySelfContainer, MapProperty, InTargetRCProperty, InProperty, InSegment, InIndexSegment, InPropertyContainer, bCheckByteEnumComparison, &PtrContainer))
-			{
-				return true;
-			}
-			return TryCopyNonUClassOwnerProperty(InVirtualPropertySelfContainer, InTargetRCProperty, InProperty, InBoundObject, InSegment, InIndexSegment + 1, PtrContainer, bCheckByteEnumComparison);
-		}
-	}
-	else if (const FStructProperty* StructProperty = CastField<FStructProperty>(InSegment.Segments[InIndexSegment].ResolvedData.Field))
-	{
-		uint8* InnerStructPtrContainer;
-		if (InPropertyContainer != nullptr)
-		{
-			InnerStructPtrContainer = StructProperty->ContainerPtrToValuePtr<uint8>(InPropertyContainer);
-		}
-		else
-		{
-			InnerStructPtrContainer = StructProperty->ContainerPtrToValuePtr<uint8>(InBoundObject);
-		}
-		if (uint8* PtrContainer = InProperty->ContainerPtrToValuePtr<uint8>(InnerStructPtrContainer))
-		{
-			if (InIndexSegment < InSegment.GetSegmentCount() - 2)
-			{
-				return TryCopyNonUClassOwnerProperty(InVirtualPropertySelfContainer, InTargetRCProperty, InProperty, InBoundObject, InSegment, InIndexSegment + 1, InnerStructPtrContainer, bCheckByteEnumComparison);
-			}
-
-			// FLinearColor are treated as FColor to avoid creating 2 different widget for them and avoid confusion
-			// Instead of the normal CopyCompleteValue we use this for FLinearColor
-			if (const FStructProperty* StructRealProperty = CastField<FStructProperty>(InProperty))
-			{
-#if !WITH_EDITOR
-				StructMemoryContainer = InnerStructPtrContainer;
-#endif // !WITH_EDITOR
-				if (StructRealProperty->Struct->GetFName() == NAME_LinearColor)
-				{
-					FColor ColorValue;
-					InVirtualPropertySelfContainer->GetValueColor(ColorValue);
-					const FLinearColor RealValue(ColorValue);
-					InProperty->CopyCompleteValue(PtrContainer, &RealValue);
-					return true;
-				}
-			}
-
-			if (InProperty->IsA<FArrayProperty>() || InProperty->IsA<FSetProperty>() || InProperty->IsA<FMapProperty>())
-			{
-				return TryCopyNonUClassOwnerProperty(InVirtualPropertySelfContainer, InTargetRCProperty, InProperty, InBoundObject, InSegment, InIndexSegment + 1, InnerStructPtrContainer, bCheckByteEnumComparison);
-			}
-			return InVirtualPropertySelfContainer->CopyCompleteValue(InProperty, PtrContainer, bCheckByteEnumComparison);
-		}
-	}
-
-	return false;
-}
diff --git a/Engine/Plugins/VirtualProduction/RemoteControl/Source/RemoteControl/Public/IPropertyIdHandler.h b/Engine/Plugins/VirtualProduction/RemoteControl/Source/RemoteControl/Public/IPropertyIdHandler.h
--- a/Engine/Plugins/VirtualProduction/RemoteControl/Source/RemoteControl/Public/IPropertyIdHandler.h
+++ b/Engine/Plugins/VirtualProduction/RemoteControl/Source/RemoteControl/Public/IPropertyIdHandler.h
@@ -22,15 +22,21 @@ public:
 	/** Get the PropertyBag type of the given property */
 	virtual EPropertyBagPropertyType GetPropertyType(const FProperty* InProperty) const = 0;
 
-	/** Get the Property type FName */
-	virtual FName GetPropertyTypeName(const FProperty* InProperty) const = 0;
+	/**
+	 * Get the Property super type FName\n
+	 * ex: Object/Struct/Enum or the primitive type if it is a primitive property (float/bool/etc...)
+	 */
+	virtual FName GetPropertySuperTypeName(const FProperty* InProperty) const = 0;
+
+	/**
+	 * Get the Property sub type FName\n
+	 * ex: StaticMesh/Material or the primitive type if it is a primitive property (float/bool/etc...)
+	 */
+	virtual FName GetPropertySubTypeName(const FProperty* InProperty) const = 0;
 
 	/** Get the Property type Object */
 	virtual UObject* GetPropertyTypeObject(const FProperty* InProperty) const = 0;
 
-	/** Get the Property default value (Used for UObject to create a new one for the ValueWidget) */
-	virtual TObjectPtr<UObject> GetObjectPropertyDefaultValue(const FProperty* InProperty, const UClass* InClassToCreate) const = 0;
-
 	/** Get the property inside the container if not inside any container return the property passed */
 	const FProperty* GetPropertyInsideContainer(const FProperty* InProperty) const
 	{
diff --git a/Engine/Plugins/VirtualProduction/RemoteControl/Source/RemoteControl/Public/RCVirtualProperty.h b/Engine/Plugins/VirtualProduction/RemoteControl/Source/RemoteControl/Public/RCVirtualProperty.h
--- a/Engine/Plugins/VirtualProduction/RemoteControl/Source/RemoteControl/Public/RCVirtualProperty.h
+++ b/Engine/Plugins/VirtualProduction/RemoteControl/Source/RemoteControl/Public/RCVirtualProperty.h
@@ -98,31 +98,34 @@ public:
 	/** Whether this Virtual Property is a numeric type, i.e. an Integral or Floating point type */
 	bool IsNumericType() const;
 
-	/** Whether this Virtual Property represents an FVector*/
+	/** Whether this Virtual Property represents an FVector */
 	bool IsVectorType() const;
 
-	/** Whether this Virtual Property represents an FVector2D*/
+	/** Whether this Virtual Property represents an FVector2D */
 	bool IsVector2DType() const;
 
-	/** Whether this Virtual Property represents an FColor*/
+	/** Whether this Virtual Property represents an FColor */
 	bool IsColorType() const;
 
-	/** Whether this Virtual Property represents an FRotator*/
+	/** Whether this Virtual Property represents an FColor */
+	bool IsLinearColorType() const;
+
+	/** Whether this Virtual Property represents an FRotator */
 	bool IsRotatorType() const;
 
 	/** Compare this virtual property value with given property value */
 	bool IsValueEqual(URCVirtualPropertyBase* InVirtualProperty) const;
 
-	/** Operator > comparison of self with a given virtual property*/
+	/** Operator > comparison of self with a given virtual property */
 	bool IsValueGreaterThan(URCVirtualPropertyBase* InVirtualProperty) const;
 
-	/** Operator >= comparison of self with a given virtual property*/
+	/** Operator >= comparison of self with a given virtual property */
 	bool IsValueGreaterThanOrEqualTo(URCVirtualPropertyBase* InVirtualProperty) const;
 
-	/** Operator < comparison of self with a given virtual property*/
+	/** Operator < comparison of self with a given virtual property */
 	bool IsValueLesserThan(URCVirtualPropertyBase* InVirtualProperty) const;
 
-	/** Operator <= comparison of self with a given virtual property*/
+	/** Operator <= comparison of self with a given virtual property */
 	bool IsValueLesserThanOrEqualTo(URCVirtualPropertyBase* InVirtualProperty) const;
 
 	/** Copy this virtual property's data onto a given FProperty
@@ -228,6 +231,10 @@ public:
 	UFUNCTION(BlueprintCallable, Category = "Remote Control Behaviour")
 	bool GetValueColor(FColor& OutColor) const;
 
+	/** Get LinearColor value from Virtual Property */
+	UFUNCTION(BlueprintCallable, Category = "Remote Control Behaviour")
+	bool GetValueLinearColor(FLinearColor& OutLinearColor) const;
+
 	/** Get Object value from Virtual Property */
 	UFUNCTION(BlueprintCallable, Category = "Remote Control Behaviour")
 	UObject* GetValueObject() const;
@@ -313,15 +320,19 @@ public:
 	/** Set Color value from Virtual Property */
 	UFUNCTION(BlueprintCallable, Category = "Remote Control Behaviour")
 	bool SetValueColor(const FColor& InColor);
-	
+
+	/** Set LinearColor value from Virtual Property */
+	UFUNCTION(BlueprintCallable, Category = "Remote Control Behaviour")
+	bool SetValueLinearColor(const FLinearColor& InLinearColor);
+
 	/** Get FProperty Name */
 	UFUNCTION(BlueprintPure, Category = "Remote Control Behaviour")
 	FName GetPropertyName() const;
 
-	/** Fetches a user-friendly representation of the Virtual Property's type name (i.e. Controller type), given a Property Bag type and value object (for structures)*/
+	/** Fetches a user-friendly representation of the Virtual Property's type name (i.e. Controller type), given a Property Bag type and value object (for structures) */
 	static FName GetVirtualPropertyTypeDisplayName(const EPropertyBagPropertyType InValueType, UObject* InValueTypeObject);
 
-	/** Updates the display index of the virtual property (transaction aware)*/
+	/** Updates the display index of the virtual property (transaction aware) */
 	void SetDisplayIndex(const int32 InDisplayIndex);
 
 public:
diff --git a/Engine/Plugins/VirtualProduction/RemoteControl/Source/RemoteControl/Public/RemoteControlPreset.h b/Engine/Plugins/VirtualProduction/RemoteControl/Source/RemoteControl/Public/RemoteControlPreset.h
--- a/Engine/Plugins/VirtualProduction/RemoteControl/Source/RemoteControl/Public/RemoteControlPreset.h
+++ b/Engine/Plugins/VirtualProduction/RemoteControl/Source/RemoteControl/Public/RemoteControlPreset.h
@@ -77,6 +77,9 @@ public:
 
 	/** (Optional) The virtual property to use for all the cases. */
 	TObjectPtr<URCVirtualPropertySelfContainer> VirtualProperty;
+
+	/** Map with the Guid of the RCField and as its value the real PropContainer of the real property */
+	TMap<FGuid, TObjectPtr<URCVirtualPropertySelfContainer>> RealProperties;
 };
 
 /**
diff --git a/Engine/Plugins/VirtualProduction/RemoteControl/Source/RemoteControl/Public/RemoteControlPropertyIdRegistry.h b/Engine/Plugins/VirtualProduction/RemoteControl/Source/RemoteControl/Public/RemoteControlPropertyIdRegistry.h
--- a/Engine/Plugins/VirtualProduction/RemoteControl/Source/RemoteControl/Public/RemoteControlPropertyIdRegistry.h
+++ b/Engine/Plugins/VirtualProduction/RemoteControl/Source/RemoteControl/Public/RemoteControlPropertyIdRegistry.h
@@ -56,11 +56,6 @@ struct FRCPropertyIdWrapper
 	 */
 	const FName& GetSubType() const;
 
-	/** Get the object class to create of this wrapped entity.
-	 * @return The Object UClass to create for this property (used for ObjectProperty).
-	 */
-	const UClass* GetClassToCreate() const;
-
 	/** Returns whether the type and the underlying data is valid.
 	 * @return True if EntityId and either SuperType or SubType is valid, false otherwise.
 	 */
@@ -98,10 +93,6 @@ private:
 	UPROPERTY()
 	FName SubType;
 
-	/** Holds the type to be created. */
-	UPROPERTY()
-	UClass* ClassToCreate = nullptr;
-
 	friend class URemoteControlPropertyIdRegistry;
 };
 
@@ -130,6 +121,10 @@ public:
 
 	URemoteControlPreset* GetSourcePreset() const;
 
+#if WITH_EDITOR
+	virtual void PostEditChangeProperty(FPropertyChangedEvent& InPropertyChangedEvent) override;
+#endif
+
 	/**
 	 * Update the value(s) of the property(ies) that are bound to a PropertyIdAction.
 	 * @param InArgs Argument used to update the property(ies) value.
@@ -160,13 +155,6 @@ public:
 	 */
 	void RemoveIdentifiedField(const FGuid& InEntityId);
 
-	/**
-	 * Get the Object UClass to create for the property with the given Id.
-	 * @param InEntityId Id of the Property to get the class from.
-	 * @return The UClass to create for this property.
-	 */
-	const UClass* GetClassByEntityId(const FGuid& InEntityId) const;
-	
 	/**
 	 * Get the list of existing Property Ids.
 	 * @return The list of existing Property Ids that has a valid FieldId.
@@ -227,19 +215,6 @@ public:
 	 */
 	void UpdateEntityIds(const TMap<FGuid, FGuid>& InEntityIdMap);
 
-private:
-	/** Call the recursive function TryCopyNonUClassOwnerProperty for the last try on copying the value to the property
-	 *	@param InArgs Argument used to expose a property.
-	 *	@param TargetRCProperty The RC property that should change.
-	 *	@param InProperty The real property to assign the new value.
-	 *	@param BoundObject The Object that has the property that need to change.
-	 *	@return True if it was able to copy the value into the property, false otherwise.
-	 */
-	bool CopyNonUClassOwnerProperty(const FRemoteControlPropertyIdArgs& InArgs, TSharedPtr<FRemoteControlProperty> TargetRCProperty, FProperty* InProperty, UObject* BoundObject);
-
-	/** Try to copy the Virtual Property value into the real property, last try for property inside container and non UClass Owner */
-	bool TryCopyNonUClassOwnerProperty(const TObjectPtr<URCVirtualPropertySelfContainer>& InVirtualPropertySelfContainer, const TSharedPtr<FRemoteControlProperty>& TargetRCProperty, FProperty* InProperty, UObject* BoundObject, FRCFieldPathInfo Segment, int32 IndexSegment, uint8* PropertyContainer,bool bCheckByteEnumComparison);
-
 private:
 	/** Holds the identified fields. */
 	UPROPERTY(Transient)
@@ -252,8 +227,4 @@ private:
 	/** Delegate triggered when PropertyID UI needs to refresh. */
 	FOnPropertyIdActionNeedsRefresh OnPropertyIdActionNeedsRefreshDelegate;
 #endif // WITH_EDITOR
-
-#if !WITH_EDITOR
-	uint8* StructMemoryContainer;
-#endif // !WITH_EDITOR
 };
diff --git a/Engine/Plugins/VirtualProduction/RemoteControl/Source/RemoteControlLogic/Private/Action/RCPropertyIdAction.cpp b/Engine/Plugins/VirtualProduction/RemoteControl/Source/RemoteControlLogic/Private/Action/RCPropertyIdAction.cpp
--- a/Engine/Plugins/VirtualProduction/RemoteControl/Source/RemoteControlLogic/Private/Action/RCPropertyIdAction.cpp
+++ b/Engine/Plugins/VirtualProduction/RemoteControl/Source/RemoteControlLogic/Private/Action/RCPropertyIdAction.cpp
@@ -38,59 +38,24 @@ void URCPropertyIdAction::Execute() const
 			FRemoteControlPropertyIdArgs PropertyIdArgs;
 			PropertyIdArgs.VirtualProperty = PropertyContainer.Value;
 			PropertyIdArgs.PropertyId = PropertyContainer.Key.PropertyId;
-			PropertyIdArgs.SuperType = Property->GetClass()->GetFName();
+			PropertyIdArgs.RealProperties = RealPropertySelfContainer;
 
 			TSharedPtr<IPropertyIdHandler> PropertyIdHandler = IRemoteControlModule::Get().GetPropertyIdHandlerFor(Property);
 			if (!PropertyIdHandler.IsValid())
 			{
 				continue;
 			}
+
+			PropertyIdArgs.SuperType = PropertyIdHandler->GetPropertySuperTypeName(Property);
 			const EPropertyBagPropertyType PropertyBagType = PropertyIdHandler->GetPropertyType(Property);
 
-			if (PropertyBagType == EPropertyBagPropertyType::Enum)
-			{
-				PropertyIdArgs.SuperType = NAME_EnumProperty;
-				PropertyIdArgs.SubType = PropertyIdHandler->GetPropertyTypeName(Property);
-			}
-			else if (PropertyBagType == EPropertyBagPropertyType::Object)
-			{
-				if (const TSharedPtr<FStructOnScope> StructOnScope = PropertyContainer.Value->CreateStructOnScope())
-				{
-					if (const FObjectProperty* ObjectProperty = CastField<FObjectProperty>(PropertyContainer.Value->GetProperty()))
-					{
-						if (const uint8* PropertyValuePtr = ObjectProperty->ContainerPtrToValuePtr<uint8>(StructOnScope->GetStructMemory()))
-						{
-							PropertyIdArgs.SourceObject = ObjectProperty->GetObjectPropertyValue(PropertyValuePtr);
-							if (PropertyIdArgs.SourceObject)
-							{
-								if (PropertyIdArgs.SourceObject->IsA(UMaterialInterface::StaticClass()))
-								{
-									PropertyIdArgs.SourceClass = UMaterialInterface::StaticClass();
-								}
-								else
-								{
-									PropertyIdArgs.SourceClass = PropertyIdArgs.SourceObject->GetClass();
-								}
-								PropertyIdArgs.SubType = PropertyIdArgs.SourceClass->GetFName();
-							}
-						}
-					}
-				}
-			}
-			else if (PropertyBagType == EPropertyBagPropertyType::Struct)
+			if (PropertyBagType == EPropertyBagPropertyType::Enum ||
+				PropertyBagType == EPropertyBagPropertyType::Object ||
+				PropertyBagType == EPropertyBagPropertyType::Struct)
 			{
-				if (const FStructProperty* StructProperty = CastField<FStructProperty>(Property))
-				{
-					if (StructProperty->Struct->GetFName() == NAME_LinearColor)
-					{
-						PropertyIdArgs.SubType = NAME_Color;
-					}
-					else
-					{
-						PropertyIdArgs.SubType = StructProperty->Struct->GetFName();
-					}
-				}
+				PropertyIdArgs.SubType = PropertyIdHandler->GetPropertySubTypeName(Property);
 			}
+
 			PresetWeakPtr->PerformChainReaction(PropertyIdArgs);
 		}
 	}
@@ -142,6 +107,8 @@ void URCPropertyIdAction::UpdatePropertyId()
 	if (URemoteControlPreset* Preset = PresetWeakPtr.Get())
 	{
 		PropertySelfContainer.Empty();
+		// todo: This can be improved to not Empty it everytime
+		RealPropertySelfContainer.Empty();
 		const TObjectPtr<URemoteControlPropertyIdRegistry> PropertyIdRegistry = Preset->GetPropertyIdRegistry();
 
 		for (const FGuid& TargetProperty : PropertyIdRegistry->GetEntityIdsList())
@@ -159,11 +126,37 @@ void URCPropertyIdAction::UpdatePropertyId()
 						}
 
 						FString PrefixPropertyClassName = TargetRCProperty->PropertyId.ToString() + TEXT(".");
-						FName PropertyClassName = FName(PrefixPropertyClassName + PropertyIdHandler->GetPropertyTypeName(Property).ToString());
+						FName PropertyClassName = FName(PrefixPropertyClassName + PropertyIdHandler->GetPropertySubTypeName(Property).ToString());
 						const FName& NewPropertyIdName = *(TargetRCProperty->PropertyId.ToString() + TEXT(".") + PropertyClassName.ToString());
+
+						FPropertyIdContainerKey CurrentKey = FPropertyIdContainerKey(TargetRCProperty->PropertyId, PropertyClassName);
+
+						TArray<UObject*> BoundObjects = TargetRCProperty->GetBoundObjects();
+						if (BoundObjects.IsEmpty() || !BoundObjects[0])
+						{
+							continue;
+						}
+
+						FRCObjectReference ObjectRefReading;
+						const bool bResolveForReading = IRemoteControlModule::Get().ResolveObjectProperty(ERCAccess::READ_ACCESS, BoundObjects[0], TargetRCProperty->FieldPathInfo.ToString(), ObjectRefReading);
 						const FProperty* PropToDuplicate = PropertyIdHandler->GetPropertyInsideContainer(Property);
-						FPropertyIdContainerKey CurrentKey = { TargetRCProperty->PropertyId, PropertyClassName };
 
+						if (!bResolveForReading)
+						{
+							continue;
+						}
+
+						const FName& PropertyName = TargetRCProperty->GetProperty()->GetFName();
+						RealPropertySelfContainer.Add(TargetRCProperty->GetId(), NewObject<URCVirtualPropertySelfContainer>(this));
+
+						const uint8* RealPropAddress = (uint8*)ObjectRefReading.ContainerAdress;
+						if (PropToDuplicate != Property)
+						{
+							RealPropAddress = PropToDuplicate->ContainerPtrToValuePtr<uint8>(ObjectRefReading.ContainerAdress);
+						}
+
+						RealPropertySelfContainer[TargetRCProperty->GetId()]->DuplicatePropertyWithCopy(PropertyName, PropToDuplicate, RealPropAddress);
+						
 						if (CachedPropertySelfContainer.Contains(CurrentKey))
 						{
 							PropertySelfContainer.Add(CurrentKey, CachedPropertySelfContainer[CurrentKey]);
@@ -172,20 +165,28 @@ void URCPropertyIdAction::UpdatePropertyId()
 						{
 							CachedPropertySelfContainer.Add(CurrentKey, NewObject<URCVirtualPropertySelfContainer>(this));
 
-							// Special case for LinearColor, since we want to base the copy on Color struct here we add the property instead of duplicating it
-							bool bColorAdded = false;
+							bool bIsSpecialCase = false;
 
-							if (const FStructProperty* StructProp = CastField<FStructProperty>(PropToDuplicate))
+							// float property and double ones are treated as one the same goes for FLinearColor and FColor
+							if (PropToDuplicate->IsA<FFloatProperty>())
+							{
+								CachedPropertySelfContainer[CurrentKey]->AddProperty(NewPropertyIdName,
+								PropertyIdHandler->GetPropertyType(Property),
+								PropertyIdHandler->GetPropertyTypeObject(Property));
+
+								bIsSpecialCase = true;
+							}
+							else if (const FStructProperty* StructProp = CastField<FStructProperty>(PropToDuplicate))
 							{
 								if (StructProp->Struct)
 								{
-									if (StructProp->Struct->GetFName() == NAME_LinearColor ||
-										StructProp->Struct->GetFName() == NAME_Color)
+									const FName& StructName = StructProp->Struct->GetFName();
+									if (StructName == NAME_LinearColor ||
+										StructName == NAME_Color)
 									{
 										CachedPropertySelfContainer[CurrentKey]->AddProperty(NewPropertyIdName,
 										PropertyIdHandler->GetPropertyType(Property),
 										PropertyIdHandler->GetPropertyTypeObject(Property));
-
 #if WITH_EDITORONLY_DATA
 										FProperty* BagProperty = CachedPropertySelfContainer[CurrentKey]->GetProperty();
 										FStructProperty* StructProperty = CastField<FStructProperty>(BagProperty);
@@ -195,16 +196,23 @@ void URCPropertyIdAction::UpdatePropertyId()
 											StructProperty->AppendMetaData({{FName("OnlyUpdateOnInteractionEnd"), TEXT("true")}});
 										}
 #endif
-										bColorAdded = true;
+										bIsSpecialCase = true;
 									}
 								}
 							}
 
-							if (!bColorAdded)
+							if (!bIsSpecialCase)
 							{
 								CachedPropertySelfContainer[CurrentKey]->DuplicateProperty(NewPropertyIdName, PropToDuplicate);
 							}
 
+							// Do this the first time it is created so that it will have a better default value except for Object
+							// Some ObjectProperty won't work correctly with this copy for example the Material one, so we skip it
+							const bool bIsObject = PropToDuplicate->IsA<FObjectProperty>();
+							if (!bIsObject)
+							{
+								CachedPropertySelfContainer[CurrentKey]->UpdateValueWithProperty(PropToDuplicate, RealPropAddress);
+							}
 							CachedPropertySelfContainer[CurrentKey]->PresetWeakPtr = PresetWeakPtr;
 							PropertySelfContainer.Add(CurrentKey, CachedPropertySelfContainer[CurrentKey]);
 						}
diff --git a/Engine/Plugins/VirtualProduction/RemoteControl/Source/RemoteControlLogic/Public/Action/RCPropertyIdAction.h b/Engine/Plugins/VirtualProduction/RemoteControl/Source/RemoteControlLogic/Public/Action/RCPropertyIdAction.h
--- a/Engine/Plugins/VirtualProduction/RemoteControl/Source/RemoteControlLogic/Public/Action/RCPropertyIdAction.h
+++ b/Engine/Plugins/VirtualProduction/RemoteControl/Source/RemoteControlLogic/Public/Action/RCPropertyIdAction.h
@@ -2,8 +2,8 @@
 
 #pragma once
 
-#include "CoreMinimal.h"
 #include "Action/RCAction.h"
+#include "RemoteControlFieldPath.h"
 #include "RCPropertyIdAction.generated.h"
 
 class URCVirtualPropertySelfContainer;
@@ -16,6 +16,13 @@ struct FPropertyIdContainerKey
 {
 	GENERATED_BODY()
 
+	FPropertyIdContainerKey() {}
+
+	FPropertyIdContainerKey(FName InPropertyId, FName InContainerName)
+		: PropertyId(InPropertyId)
+		, ContainerName(InContainerName)
+	{}
+
 	/** PropertyId */
 	UPROPERTY()
 	FName PropertyId;
@@ -82,4 +89,8 @@ public:
 	/** Cached Virtual Property Container */
 	UPROPERTY()
 	TMap<FPropertyIdContainerKey, TObjectPtr<URCVirtualPropertySelfContainer>> CachedPropertySelfContainer;
+
+	/** Contains the real property container */
+	UPROPERTY()
+	TMap<FGuid, TObjectPtr<URCVirtualPropertySelfContainer>> RealPropertySelfContainer;
 };
diff --git a/Engine/Plugins/VirtualProduction/RemoteControl/Source/RemoteControlUI/Private/UI/Action/Conditional/RCActionConditionalModel.cpp b/Engine/Plugins/VirtualProduction/RemoteControl/Source/RemoteControlUI/Private/UI/Action/Conditional/RCActionConditionalModel.cpp
--- a/Engine/Plugins/VirtualProduction/RemoteControl/Source/RemoteControlUI/Private/UI/Action/Conditional/RCActionConditionalModel.cpp
+++ b/Engine/Plugins/VirtualProduction/RemoteControl/Source/RemoteControlUI/Private/UI/Action/Conditional/RCActionConditionalModel.cpp
@@ -103,12 +103,12 @@ TSharedPtr<SHeaderRow> FRCActionConditionalModel::GetHeaderRow()
 
 		+ SHeaderRow::Column(UE::RCActionPanelConditionalList::Columns::Description)
 		.DefaultLabel(LOCTEXT("RCActionDescColumnHeader", "PropertyID"))
-		.FillWidth(0.4f)
+		.FillWidth(0.2f)
 		.HeaderContentPadding(RCPanelStyle->HeaderRowPadding)
 
 		+ SHeaderRow::Column(UE::RCActionPanelConditionalList::Columns::Value)
 		.DefaultLabel(LOCTEXT("RCActionValueColumnHeader", "Value"))
-		.FillWidth(0.4f)
+		.FillWidth(0.6f)
 		.HeaderContentPadding(RCPanelStyle->HeaderRowPadding);
 }
 
diff --git a/Engine/Plugins/VirtualProduction/RemoteControl/Source/RemoteControlUI/Private/UI/Action/RCActionModel.cpp b/Engine/Plugins/VirtualProduction/RemoteControl/Source/RemoteControlUI/Private/UI/Action/RCActionModel.cpp
--- a/Engine/Plugins/VirtualProduction/RemoteControl/Source/RemoteControlUI/Private/UI/Action/RCActionModel.cpp
+++ b/Engine/Plugins/VirtualProduction/RemoteControl/Source/RemoteControlUI/Private/UI/Action/RCActionModel.cpp
@@ -400,8 +400,12 @@ TSharedRef<SWidget> FRCPropertyIdActionType::GetPropertyIdValueWidget() const
 	SortedTreeNodesKeys.Sort();
 
 	FName LastPropId = NAME_None;
+	constexpr float Offset = 5.f;
+
 	FMargin TitlePadding = FMargin(0.f);
 	FMargin ValuePadding = FMargin(0.f);
+	ValuePadding.Right = Offset;
+
 	int32 OriginalPropIdLength = 0;
 	int32 OriginalDotCount = 0;
 
@@ -418,8 +422,6 @@ TSharedRef<SWidget> FRCPropertyIdActionType::GetPropertyIdValueWidget() const
 
 		if (LastPropId != TreeNodeKey.PropertyId)
 		{
-			constexpr float Offset = 5.f;
-
 			if (LastPropId != NAME_None)
 			{
 				TitlePadding.Top = Offset;
@@ -446,17 +448,20 @@ TSharedRef<SWidget> FRCPropertyIdActionType::GetPropertyIdValueWidget() const
 
 		if (const TWeakPtr<IDetailTreeNode>* Node = ValueTreeNodeWeakPtr.Find(TreeNodeKey))
 		{
-			VerticalBox->AddSlot()
-				.Padding(ValuePadding)
-				.AutoHeight()
-				[
-					SNew(SHorizontalBox)
-					+ SHorizontalBox::Slot()
-					.AutoWidth()
+			if (Node->IsValid())
+			{
+				VerticalBox->AddSlot()
+					.Padding(ValuePadding)
+					.AutoHeight()
 					[
-						UE::RCUIHelpers::GetGenericFieldWidget(Node->Pin())
-					]
-				];
+						SNew(SHorizontalBox)
+						+ SHorizontalBox::Slot()
+						.FillWidth(0.7f)
+						[
+							UE::RCUIHelpers::GetGenericFieldWidget(Node->Pin())
+						]
+					];
+			}
 		}
 	}
 
