[RemoteControl]
Fix:
- Instead of Emptying the RealProperties in propId now we only add newer ones
- Crash when creating a behaviour
- Crash when copying some Struct
- Blocked copy of ObjectProperty value in PropId Action to avoid crash (For example with the Material)
- Added copy for Struct only if Allowed (Color/LinearColor/Vector/Rotation)
- Crash when calling OnObjectPropertyChanged in Preset because recursive in some case
- Changed from AddLambda to AddSPLambda

diff --git a/Engine/Plugins/VirtualProduction/RemoteControl/Source/RemoteControl/Private/RemoteControlPreset.cpp b/Engine/Plugins/VirtualProduction/RemoteControl/Source/RemoteControl/Private/RemoteControlPreset.cpp
--- a/Engine/Plugins/VirtualProduction/RemoteControl/Source/RemoteControl/Private/RemoteControlPreset.cpp
+++ b/Engine/Plugins/VirtualProduction/RemoteControl/Source/RemoteControl/Private/RemoteControlPreset.cpp
@@ -2039,6 +2039,13 @@ void URemoteControlPreset::OnObjectPropertyChanged(UObject* Object, struct FProp
 	// Objects modified should have run through the preobjectmodified. If interesting, they will be cached
 	TRACE_CPUPROFILER_EVENT_SCOPE(URemoteControlPreset::OnObjectPropertyChanged);
 
+	// ResetToDefaultValue will recall this function but re-executing it will cause a crash so we guard it to avoid this scenario
+	if (bReentryGuard)
+	{
+		return;
+	}
+	TGuardValue<bool> Guard(bReentryGuard, true);
+
 	// Handle enter / exit Multi-User session
 	if (Object == this->GetPackage() && Event.ChangeType == EPropertyChangeType::Redirected)
 	{
diff --git a/Engine/Plugins/VirtualProduction/RemoteControl/Source/RemoteControl/Public/RemoteControlPreset.h b/Engine/Plugins/VirtualProduction/RemoteControl/Source/RemoteControl/Public/RemoteControlPreset.h
--- a/Engine/Plugins/VirtualProduction/RemoteControl/Source/RemoteControl/Public/RemoteControlPreset.h
+++ b/Engine/Plugins/VirtualProduction/RemoteControl/Source/RemoteControl/Public/RemoteControlPreset.h
@@ -987,6 +987,9 @@ private:
 	/** Whether there is an ongoing remote modification happening. */
 	bool bOngoingRemoteModification = false;
 
+	/** Used for OnObjectPropertyChanged to avoid being called by itself. */
+	bool bReentryGuard = false;
+
 	/** Holds manager that handles rebinding unbound entities upon load or map change. */
 	TPimplPtr<FRemoteControlPresetRebindingManager> RebindingManager;
 
diff --git a/Engine/Plugins/VirtualProduction/RemoteControl/Source/RemoteControlLogic/Private/Action/RCPropertyIdAction.cpp b/Engine/Plugins/VirtualProduction/RemoteControl/Source/RemoteControlLogic/Private/Action/RCPropertyIdAction.cpp
--- a/Engine/Plugins/VirtualProduction/RemoteControl/Source/RemoteControlLogic/Private/Action/RCPropertyIdAction.cpp
+++ b/Engine/Plugins/VirtualProduction/RemoteControl/Source/RemoteControlLogic/Private/Action/RCPropertyIdAction.cpp
@@ -10,12 +10,25 @@
 #include "RemoteControlPreset.h"
 #include "RemoteControlPropertyIdRegistry.h"
 
+TSet<FName> URCPropertyIdAction::AllowedStructNameToCopy
+{
+	NAME_Color,
+	NAME_LinearColor,
+	NAME_Vector,
+	NAME_Vector2D,
+	NAME_Rotation,
+	NAME_Rotator
+};
+
 URCPropertyIdAction::~URCPropertyIdAction()
 {
 	if (PresetWeakPtr.IsValid())
 	{
 		PresetWeakPtr->OnEntityUnexposed().RemoveAll(this);
-		PresetWeakPtr->GetPropertyIdRegistry()->OnPropertyIdUpdated().RemoveAll(this);
+		if (const TObjectPtr<URemoteControlPropertyIdRegistry> Registry = PresetWeakPtr->GetPropertyIdRegistry())
+		{
+			Registry->OnPropertyIdUpdated().RemoveAll(this);
+		}
 	}
 }
 
@@ -107,8 +120,7 @@ void URCPropertyIdAction::UpdatePropertyId()
 	if (URemoteControlPreset* Preset = PresetWeakPtr.Get())
 	{
 		PropertySelfContainer.Empty();
-		// todo: This can be improved to not Empty it everytime
-		RealPropertySelfContainer.Empty();
+		TArray<FGuid> CurrentRCGuidsInAction;
 		const TObjectPtr<URemoteControlPropertyIdRegistry> PropertyIdRegistry = Preset->GetPropertyIdRegistry();
 
 		for (const FGuid& TargetProperty : PropertyIdRegistry->GetEntityIdsList())
@@ -147,16 +159,27 @@ void URCPropertyIdAction::UpdatePropertyId()
 						}
 
 						const FName& PropertyName = TargetRCProperty->GetProperty()->GetFName();
-						RealPropertySelfContainer.Add(TargetRCProperty->GetId(), NewObject<URCVirtualPropertySelfContainer>(this));
-
 						const uint8* RealPropAddress = (uint8*)ObjectRefReading.ContainerAdress;
-						if (PropToDuplicate != Property)
+
+						if (RealPropAddress && PropToDuplicate != Property)
 						{
 							RealPropAddress = PropToDuplicate->ContainerPtrToValuePtr<uint8>(ObjectRefReading.ContainerAdress);
 						}
 
-						RealPropertySelfContainer[TargetRCProperty->GetId()]->DuplicatePropertyWithCopy(PropertyName, PropToDuplicate, RealPropAddress);
-						
+						if (!RealPropAddress)
+						{
+							continue;
+						}
+
+						FGuid RCGuid = TargetRCProperty->GetId();
+						CurrentRCGuidsInAction.Add(RCGuid);
+						if (!RealPropertySelfContainer.Contains(RCGuid))
+						{
+							RealPropertySelfContainer.Add(RCGuid, NewObject<URCVirtualPropertySelfContainer>(this));
+							RealPropertySelfContainer[RCGuid]->DuplicateProperty(PropertyName, PropToDuplicate);
+							RealPropertySelfContainer[RCGuid]->PresetWeakPtr = PresetWeakPtr;
+						}
+
 						if (CachedPropertySelfContainer.Contains(CurrentKey))
 						{
 							PropertySelfContainer.Add(CurrentKey, CachedPropertySelfContainer[CurrentKey]);
@@ -166,6 +189,7 @@ void URCPropertyIdAction::UpdatePropertyId()
 							CachedPropertySelfContainer.Add(CurrentKey, NewObject<URCVirtualPropertySelfContainer>(this));
 
 							bool bIsSpecialCase = false;
+							FName StructName = NAME_None;
 
 							// float property and double ones are treated as one the same goes for FLinearColor and FColor
 							if (PropToDuplicate->IsA<FFloatProperty>())
@@ -180,7 +204,7 @@ void URCPropertyIdAction::UpdatePropertyId()
 							{
 								if (StructProp->Struct)
 								{
-									const FName& StructName = StructProp->Struct->GetFName();
+									StructName = StructProp->Struct->GetFName();
 									if (StructName == NAME_LinearColor ||
 										StructName == NAME_Color)
 									{
@@ -208,8 +232,22 @@ void URCPropertyIdAction::UpdatePropertyId()
 
 							// Do this the first time it is created so that it will have a better default value except for Object
 							// Some ObjectProperty won't work correctly with this copy for example the Material one, so we skip it
-							const bool bIsObject = PropToDuplicate->IsA<FObjectProperty>();
-							if (!bIsObject)
+							// The same for some struct so we copy only the AllowedOnes
+							bool bCanCopyValue = true;
+							if (PropToDuplicate->IsA<FObjectProperty>())
+							{
+								bCanCopyValue = false;
+							}
+
+							if (StructName != NAME_None)
+							{
+								if (!AllowedStructNameToCopy.Contains(StructName))
+								{
+									bCanCopyValue = false;
+								}
+							}
+
+							if (bCanCopyValue)
 							{
 								CachedPropertySelfContainer[CurrentKey]->UpdateValueWithProperty(PropToDuplicate, RealPropAddress);
 							}
@@ -220,6 +258,17 @@ void URCPropertyIdAction::UpdatePropertyId()
 				}
 			}
 		}
+
+		// Remove RealContainers that are not anymore part of this action
+		TArray<FGuid> CurrentRealPropertyGuids;
+		RealPropertySelfContainer.GetKeys(CurrentRealPropertyGuids);
+
+		CurrentRealPropertyGuids.RemoveAll([CurrentRCGuidsInAction] (const FGuid& InGuid) { return CurrentRCGuidsInAction.Contains(InGuid); });
+		for (const FGuid& Guid : CurrentRealPropertyGuids)
+		{
+			RealPropertySelfContainer.Remove(Guid);
+		}
+
 #if WITH_EDITOR
 		if (PresetWeakPtr.IsValid())
 		{
@@ -234,7 +283,10 @@ void URCPropertyIdAction::Initialize()
 	if (PresetWeakPtr.IsValid())
 	{
 		PresetWeakPtr->OnEntityUnexposed().AddUObject(this, &URCPropertyIdAction::OnEntityUnexposed);
-		PresetWeakPtr->GetPropertyIdRegistry()->OnPropertyIdUpdated().AddUObject(this, &URCPropertyIdAction::UpdatePropertyId);
+		if (const TObjectPtr<URemoteControlPropertyIdRegistry> Registry = PresetWeakPtr->GetPropertyIdRegistry())
+		{
+			Registry->OnPropertyIdUpdated().AddUObject(this, &URCPropertyIdAction::UpdatePropertyId);
+		}
 	}
 }
 
diff --git a/Engine/Plugins/VirtualProduction/RemoteControl/Source/RemoteControlLogic/Public/Action/RCPropertyIdAction.h b/Engine/Plugins/VirtualProduction/RemoteControl/Source/RemoteControlLogic/Public/Action/RCPropertyIdAction.h
--- a/Engine/Plugins/VirtualProduction/RemoteControl/Source/RemoteControlLogic/Public/Action/RCPropertyIdAction.h
+++ b/Engine/Plugins/VirtualProduction/RemoteControl/Source/RemoteControlLogic/Public/Action/RCPropertyIdAction.h
@@ -93,4 +93,7 @@ public:
 	/** Contains the real property container */
 	UPROPERTY()
 	TMap<FGuid, TObjectPtr<URCVirtualPropertySelfContainer>> RealPropertySelfContainer;
+
+private:
+	static TSet<FName> AllowedStructNameToCopy;
 };
diff --git a/Engine/Plugins/VirtualProduction/RemoteControl/Source/RemoteControlUI/Private/UI/Action/SRCActionPanel.cpp b/Engine/Plugins/VirtualProduction/RemoteControl/Source/RemoteControlUI/Private/UI/Action/SRCActionPanel.cpp
--- a/Engine/Plugins/VirtualProduction/RemoteControl/Source/RemoteControlUI/Private/UI/Action/SRCActionPanel.cpp
+++ b/Engine/Plugins/VirtualProduction/RemoteControl/Source/RemoteControlUI/Private/UI/Action/SRCActionPanel.cpp
@@ -79,8 +79,11 @@ void SRCActionPanel::Construct(const FArguments& InArgs, const TSharedRef<SRemot
 	{
 		Preset->Layout.OnFieldAdded().AddSP(this, &SRCActionPanel::OnRemoteControlFieldAdded);
 		Preset->Layout.OnFieldDeleted().AddSP(this, &SRCActionPanel::OnRemoteControlFieldDeleted);
-		Preset->GetPropertyIdRegistry()->OnPropertyIdUpdated().AddLambda([this](){ bAddActionMenuNeedsRefresh = true; });
-		Preset->GetPropertyIdRegistry()->OnPropertyIdActionNeedsRefresh().AddLambda([this](){ bAddActionMenuNeedsRefresh = true; });
+		if (const TObjectPtr<URemoteControlPropertyIdRegistry> Registry = Preset->GetPropertyIdRegistry())
+		{
+			Registry->OnPropertyIdUpdated().AddSPLambda(this, [this](){ bAddActionMenuNeedsRefresh = true; });
+			Registry->OnPropertyIdActionNeedsRefresh().AddSPLambda(this, [this](){ bAddActionMenuNeedsRefresh = true; });
+		}
 	}
 }
 
diff --git a/Engine/Plugins/VirtualProduction/RemoteControl/Source/RemoteControlUI/Private/UI/Behaviour/SRCBehaviourPanel.cpp b/Engine/Plugins/VirtualProduction/RemoteControl/Source/RemoteControlUI/Private/UI/Behaviour/SRCBehaviourPanel.cpp
--- a/Engine/Plugins/VirtualProduction/RemoteControl/Source/RemoteControlUI/Private/UI/Behaviour/SRCBehaviourPanel.cpp
+++ b/Engine/Plugins/VirtualProduction/RemoteControl/Source/RemoteControlUI/Private/UI/Behaviour/SRCBehaviourPanel.cpp
@@ -153,7 +153,7 @@ TSharedRef<SWidget> SRCBehaviourPanel::GetBehaviourMenuContentWidget()
 					}
 
 					FUIAction Action(FExecuteAction::CreateSP(this, &SRCBehaviourPanel::OnAddBehaviourClicked, Class));
-					Action.CanExecuteAction.BindSP(this, &SRCBehaviourPanel::CanExecuteAddBehaviour, Class, Behaviour);
+					Action.CanExecuteAction.BindSP(this, &SRCBehaviourPanel::CanExecuteAddBehaviour, Class, Controller);
 
 					MenuBuilder.AddMenuEntry(
 						FText::Format(LOCTEXT("AddBehaviourNode", "{0}"), Behaviour->GetDisplayName()),
@@ -187,10 +187,19 @@ void SRCBehaviourPanel::OnAddBehaviourClicked(UClass* InClass)
 	}
 }
 
-bool SRCBehaviourPanel::CanExecuteAddBehaviour(UClass* InClass, URCBehaviour* InBehaviour) const
+bool SRCBehaviourPanel::CanExecuteAddBehaviour(UClass* InClass, URCController* InController) const
 {
+	if (!InClass || !InController)
+	{
+		return false;
+	}
 	const URCBehaviourNode* DefaultBehaviourNode = Cast<URCBehaviourNode>(InClass->GetDefaultObject());
-	return DefaultBehaviourNode->IsSupported(InBehaviour);
+	TObjectPtr<URCBehaviour> Behaviour = InController->CreateBehaviourWithoutCheck(InClass);
+	if (!Behaviour)
+	{
+		return false;
+	}
+	return DefaultBehaviourNode->IsSupported(Behaviour);
 }
 
 FReply SRCBehaviourPanel::OnClickEmptyButton()
diff --git a/Engine/Plugins/VirtualProduction/RemoteControl/Source/RemoteControlUI/Private/UI/Behaviour/SRCBehaviourPanel.h b/Engine/Plugins/VirtualProduction/RemoteControl/Source/RemoteControlUI/Private/UI/Behaviour/SRCBehaviourPanel.h
--- a/Engine/Plugins/VirtualProduction/RemoteControl/Source/RemoteControlUI/Private/UI/Behaviour/SRCBehaviourPanel.h
+++ b/Engine/Plugins/VirtualProduction/RemoteControl/Source/RemoteControlUI/Private/UI/Behaviour/SRCBehaviourPanel.h
@@ -88,7 +88,7 @@ private:
 	void OnAddBehaviourClicked(UClass* InClass);
 
 	/** Check if the controller can add the behaviour requested */
-	bool CanExecuteAddBehaviour(UClass* InClass, URCBehaviour* InBehaviour) const;
+	bool CanExecuteAddBehaviour(UClass* InClass, URCController* InController) const;
 
 	/** Handles click event for "Empty" button; clears all Behaviours from the panel*/
 	FReply OnClickEmptyButton();
