[RemoteControl]
FindSetter function will not return early in Editor if not found in the first tries.

diff --git a/Engine/Plugins/VirtualProduction/RemoteControl/Source/RemoteControlCommon/Public/RCPropertyUtilities.inl b/Engine/Plugins/VirtualProduction/RemoteControl/Source/RemoteControlCommon/Public/RCPropertyUtilities.inl
--- a/Engine/Plugins/VirtualProduction/RemoteControl/Source/RemoteControlCommon/Public/RCPropertyUtilities.inl
+++ b/Engine/Plugins/VirtualProduction/RemoteControl/Source/RemoteControlCommon/Public/RCPropertyUtilities.inl
@@ -506,11 +506,6 @@ namespace RemoteControlPropertyUtilities
 			{
 				SetterFunction = OwnerClass->FindFunctionByName(*CommonPropertiesToFunctions[Property->GetFName()].ToString());
 			}
-
-			if (!SetterFunction)
-			{
-				return nullptr;
-			}
 		}
 #endif
 
