[RemoteControl]
Rename controllers column from "Name" to "Controller Id"

diff --git a/Engine/Plugins/VirtualProduction/RemoteControl/Source/RemoteControlUI/Private/UI/Controller/SRCControllerPanelList.cpp b/Engine/Plugins/VirtualProduction/RemoteControl/Source/RemoteControlUI/Private/UI/Controller/SRCControllerPanelList.cpp
--- a/Engine/Plugins/VirtualProduction/RemoteControl/Source/RemoteControlUI/Private/UI/Controller/SRCControllerPanelList.cpp
+++ b/Engine/Plugins/VirtualProduction/RemoteControl/Source/RemoteControlUI/Private/UI/Controller/SRCControllerPanelList.cpp
@@ -52,7 +52,7 @@ namespace UE::RCControllerPanelList
 	namespace Columns
 	{
 		const FName TypeColor = TEXT("TypeColor");
-		const FName Name = TEXT("Controller Name");
+		const FName ControllerId = TEXT("Controller Id");
 		const FName Description = TEXT("Controller Description");
 		const FName Value = TEXT("Controller Value");
 		const FName DragHandle = TEXT("Drag Handle");
@@ -136,7 +136,7 @@ namespace UE::RCControllerPanelList
 			{
 				return WrapWithDropTarget(ControllerItem->GetTypeSelectionWidget());
 			}
-			else if (ColumnName == UE::RCControllerPanelList::Columns::Name)
+			else if (ColumnName == UE::RCControllerPanelList::Columns::ControllerId)
 			{
 				return WrapWithDropTarget(ControllerItem->GetNameWidget());
 			}
@@ -368,8 +368,8 @@ void SRCControllerPanelList::Construct(const FArguments& InArgs, const TSharedRe
 			.FixedWidth(15)
 			.HeaderContentPadding(RCPanelStyle->HeaderRowPadding)
 
-			+ SHeaderRow::Column(UE::RCControllerPanelList::Columns::Name)
-			.DefaultLabel(LOCTEXT("ControllerNameColumnName", "Name"))
+			+ SHeaderRow::Column(UE::RCControllerPanelList::Columns::ControllerId)
+			.DefaultLabel(LOCTEXT("ControllerIdColumnName", "Controller Id"))
 			.FillWidth(0.2f)
 			.HeaderContentPadding(RCPanelStyle->HeaderRowPadding)
 
