[RemoteControl]
Make controller Id widget to not scale in height but always be centered

diff --git a/Engine/Plugins/VirtualProduction/RemoteControl/Source/RemoteControlUI/Private/UI/Controller/RCControllerModel.cpp b/Engine/Plugins/VirtualProduction/RemoteControl/Source/RemoteControlUI/Private/UI/Controller/RCControllerModel.cpp
--- a/Engine/Plugins/VirtualProduction/RemoteControl/Source/RemoteControlUI/Private/UI/Controller/RCControllerModel.cpp
+++ b/Engine/Plugins/VirtualProduction/RemoteControl/Source/RemoteControlUI/Private/UI/Controller/RCControllerModel.cpp
@@ -28,6 +28,8 @@ FRCControllerModel::FRCControllerModel(URCVirtualPropertyBase* InVirtualProperty
 		}
 
 		SAssignNew(ControllerNameTextBox, SEditableTextBox)
+			.RevertTextOnEscape(true)
+			.SelectAllTextWhenFocused(true)
 			.Text(FText::FromName(InVirtualProperty->DisplayName))
 			.OnTextCommitted_Raw(this, &FRCControllerModel::OnControllerNameCommitted);
 
@@ -103,7 +105,9 @@ TSharedRef<SWidget> FRCControllerModel::GetWidget() const
 
 TSharedRef<SWidget> FRCControllerModel::GetNameWidget() const
 {
-	return SNew(SBox).Padding(10.f, 2.f)
+	return SNew(SBox)
+		.VAlign(VAlign_Center)
+		.Padding(10.f, 2.f)
 		[
 			ControllerNameTextBox.ToSharedRef()
 		];
diff --git a/Engine/Plugins/VirtualProduction/RemoteControl/Source/RemoteControlUI/Private/UI/Controller/RCControllerModel.h b/Engine/Plugins/VirtualProduction/RemoteControl/Source/RemoteControlUI/Private/UI/Controller/RCControllerModel.h
--- a/Engine/Plugins/VirtualProduction/RemoteControl/Source/RemoteControlUI/Private/UI/Controller/RCControllerModel.h
+++ b/Engine/Plugins/VirtualProduction/RemoteControl/Source/RemoteControlUI/Private/UI/Controller/RCControllerModel.h
@@ -125,10 +125,10 @@ private:
 	/** The currently selected Behaviour (UI model) */
 	TWeakPtr<FRCBehaviourModel>  SelectedBehaviourModelWeakPtr;
 
-	/** Controller name - editable text box */
+	/** Controller id - editable text box */
 	TSharedPtr<SEditableTextBox> ControllerNameTextBox;
 
-	/** Controller name - editable text box */
+	/** Controller description - editable text box */
 	TSharedPtr<SInlineEditableTextBlock> ControllerDescriptionTextBox;
 
 	/** Controller Field Id - editable text box */
