[RemoteControl]
While in the Operation Mode you are unable to toggle on the Logic Panel.
Switched position of the Logic Panel button and the Protocol button.

diff --git a/Engine/Plugins/VirtualProduction/RemoteControl/Source/RemoteControlUI/Private/UI/SRemoteControlPanel.cpp b/Engine/Plugins/VirtualProduction/RemoteControl/Source/RemoteControlUI/Private/UI/SRemoteControlPanel.cpp
--- a/Engine/Plugins/VirtualProduction/RemoteControl/Source/RemoteControlUI/Private/UI/SRemoteControlPanel.cpp
+++ b/Engine/Plugins/VirtualProduction/RemoteControl/Source/RemoteControlUI/Private/UI/SRemoteControlPanel.cpp
@@ -582,24 +582,46 @@ void SRemoteControlPanel::Construct(const FArguments& InArgs, URemoteControlPres
 
 				+ SWidgetSwitcher::Slot()
 				[
-					EntityList.ToSharedRef()
+					SNew(SWidgetSwitcher)
+					.WidgetIndex_Lambda([this]() { return !bIsInLiveMode ? 0 : 1; })
+
+					+ SWidgetSwitcher::Slot()
+                    [
+						EntityList.ToSharedRef()
+					]
+
+					+ SWidgetSwitcher::Slot()
+					[
+						SNullWidget::NullWidget
+					]
 				]
 
 				+ SWidgetSwitcher::Slot()
 				[
-					SNew(SSplitter)
-					.Orientation(Orient_Horizontal)
+					SNew(SWidgetSwitcher)
+					.WidgetIndex_Lambda([this]() { return bIsInLiveMode ? 0 : 1; })
 
-					+SSplitter::Slot()
-					.Value(0.4)
+					+ SWidgetSwitcher::Slot()
 					[
-						EntityList.ToSharedRef()
+						LogicPanel
 					]
 
-					+SSplitter::Slot()
-					.Value(0.6)
+					+ SWidgetSwitcher::Slot()
 					[
-						LogicPanel
+						SNew(SSplitter)
+						.Orientation(Orient_Horizontal)
+
+						+SSplitter::Slot()
+						.Value(0.4)
+						[
+							EntityList.ToSharedRef()
+						]
+
+						+SSplitter::Slot()
+						.Value(0.6)
+						[
+							LogicPanel
+						]
 					]
 				]
 			]
@@ -2047,6 +2069,7 @@ void SRemoteControlPanel::GenerateToolbar()
 					if (NewMode.ModeId == TEXT("Operation"))
 					{
 						bIsInLiveMode = true;
+						bIsLogicPanelEnabled = true;
 					}
 					else if (NewMode.ModeId == TEXT("Setup"))
 					{
@@ -2181,20 +2204,6 @@ void SRemoteControlPanel::RegisterAuxiliaryToolBar()
 
 			const FRemoteControlCommands& Commands = FRemoteControlCommands::Get();
 
-			ToolsSection.AddEntry(FToolMenuEntry::InitWidget("Protocols"
-				, SNew(SVerticalBox)
-					+ SVerticalBox::Slot()
-					.AutoHeight()
-					[
-						SNew(SAutoResizeButton)
-							.UICommand(FRemoteControlCommands::Get().ToggleProtocolMappings)
-							.ForceSmallIcons_Static(SRemoteControlPanel::ShouldForceSmallIcons)
-							.IconOverride(FSlateIcon(FAppStyle::GetAppStyleSetName(), "LevelEditor.Tabs.StatsViewer"))
-					]
-					, Commands.ToggleProtocolMappings->GetLabel()
-			)
-			);
-
 			ToolsSection.AddEntry(FToolMenuEntry::InitWidget("Logic"
 			, SNew(SVerticalBox)
 				+ SVerticalBox::Slot()
@@ -2208,6 +2217,20 @@ void SRemoteControlPanel::RegisterAuxiliaryToolBar()
 				, Commands.ToggleLogicEditor->GetLabel()
 			)
 			);
+
+			ToolsSection.AddEntry(FToolMenuEntry::InitWidget("Protocols"
+			, SNew(SVerticalBox)
+				+ SVerticalBox::Slot()
+				.AutoHeight()
+				[
+					SNew(SAutoResizeButton)
+					.UICommand(FRemoteControlCommands::Get().ToggleProtocolMappings)
+					.ForceSmallIcons_Static(SRemoteControlPanel::ShouldForceSmallIcons)
+					.IconOverride(FSlateIcon(FAppStyle::GetAppStyleSetName(), "LevelEditor.Tabs.StatsViewer"))
+				]
+				, Commands.ToggleProtocolMappings->GetLabel()
+			)
+			);
 		}
 	}
 }
@@ -2378,12 +2401,12 @@ void SRemoteControlPanel::ToggleLogicEditor_Execute()
 
 bool SRemoteControlPanel::CanToggleLogicPanel() const
 {
-	return !bIsInLiveMode;
+	return true;
 }
 
 bool SRemoteControlPanel::IsLogicPanelEnabled() const
 {
-	return !bIsInLiveMode && bIsLogicPanelEnabled;
+	return bIsLogicPanelEnabled;
 }
 
 void SRemoteControlPanel::OnRCPanelToggled(ERCPanels InPanelID)
