[RemoteControl] Alternative API to prevent remote control presets from renewing their Guids on duplicate.

The Motion Design Rundown editor is loading and duplicating level assets containing embedded remote control presets. This is done as part of the Rundown editor's work flow in order to have a copy of the RCP that will be used to run the controllers and capture the resulting values. We want the duplicated RCP to keep it's guids the same as the original.

The original method for doing this was to use EDuplicateMode::PIE. However, since there are may be no PIE instances running, this method hits some check failures. One such failure in RemapGuid (LazyObjectPtr.cpp) seems to indicate EDuplicateMode::PIE can't be used if there are no PIE instances running.

To resolve this issue for the moment, we will add an API to explicitly disable the guid renewal.

diff --git a/Engine/Plugins/Experimental/Avalanche/Source/AvalancheMedia/Private/Playlist/AvalancheManagedInstanceLevel.cpp b/Engine/Plugins/Experimental/Avalanche/Source/AvalancheMedia/Private/Playlist/AvalancheManagedInstanceLevel.cpp
--- a/Engine/Plugins/Experimental/Avalanche/Source/AvalancheMedia/Private/Playlist/AvalancheManagedInstanceLevel.cpp
+++ b/Engine/Plugins/Experimental/Avalanche/Source/AvalancheMedia/Private/Playlist/AvalancheManagedInstanceLevel.cpp
@@ -52,8 +52,11 @@ FAvalancheManagedInstanceLevel::FAvalancheManagedInstanceLevel(FAvalancheManaged
 		return;
 	}
 
-	// Remark: using PIE duplication mode to avoid entity and controller Ids from being renewed.
-	ManagedAvalancheLevel = Cast<UWorld>(StaticDuplicateObject(SourceAvalancheLevel, ManagedAvalancheLevelPackage.Get(), NAME_None, RF_NoFlags, nullptr, EDuplicateMode::PIE));
+	{
+		FRCPresetGuidRenewGuard PresetGuidRenewGuard;
+		ManagedAvalancheLevel = Cast<UWorld>(StaticDuplicateObject(SourceAvalancheLevel, ManagedAvalancheLevelPackage.Get()));
+	}
+	
 	if (!ManagedAvalancheLevel)
 	{
 		UE_LOG(LogAvaMedia, Error, TEXT("Unable to duplicate Source Motion Design Level: %s"), *InAssetPath.ToString());
diff --git a/Engine/Plugins/VirtualProduction/RemoteControl/Source/RemoteControl/Private/RemoteControlPreset.cpp b/Engine/Plugins/VirtualProduction/RemoteControl/Source/RemoteControl/Private/RemoteControlPreset.cpp
--- a/Engine/Plugins/VirtualProduction/RemoteControl/Source/RemoteControl/Private/RemoteControlPreset.cpp
+++ b/Engine/Plugins/VirtualProduction/RemoteControl/Source/RemoteControl/Private/RemoteControlPreset.cpp
@@ -56,6 +56,7 @@ namespace
 	FName NAME_DefaultLayoutGroup = FName("All");
 	FName NAME_DefaultNewGroup = FName("New Group");
 	const FString DefaultObjectPrefix = TEXT("Default__");
+	bool bAllowPresetGuidRenewal = true; 
 
 	UClass* FindCommonBase(const TArray<UObject*>& ObjectsToTest)
 	{
@@ -180,6 +181,17 @@ FRemoteControlPresetExposeArgs::FRemoteControlPresetExposeArgs(FString InLabel,
 {
 }
 
+FRCPresetGuidRenewGuard::FRCPresetGuidRenewGuard()
+{
+	bPreviousValue = bAllowPresetGuidRenewal;
+	bAllowPresetGuidRenewal = false;
+}
+
+FRCPresetGuidRenewGuard::~FRCPresetGuidRenewGuard()
+{
+	bAllowPresetGuidRenewal = bPreviousValue;
+}
+
 const TArray<FGuid>& FRemoteControlPresetGroup::GetFields() const
 {
 	return Fields;
@@ -710,7 +722,7 @@ void URemoteControlPreset::PostLoad()
 void URemoteControlPreset::PostDuplicate(bool bDuplicateForPIE)
 {
 	Super::PostDuplicate(bDuplicateForPIE);
-	if (!bDuplicateForPIE)
+	if (!bDuplicateForPIE && bAllowPresetGuidRenewal)
 	{
 		PresetId = FGuid::NewGuid();
 
diff --git a/Engine/Plugins/VirtualProduction/RemoteControl/Source/RemoteControl/Public/RemoteControlPreset.h b/Engine/Plugins/VirtualProduction/RemoteControl/Source/RemoteControl/Public/RemoteControlPreset.h
--- a/Engine/Plugins/VirtualProduction/RemoteControl/Source/RemoteControl/Public/RemoteControlPreset.h
+++ b/Engine/Plugins/VirtualProduction/RemoteControl/Source/RemoteControl/Public/RemoteControlPreset.h
@@ -130,6 +130,18 @@ struct REMOTECONTROL_API FExposedFunction
 	TArray<UObject*> OwnerObjects;
 };
 
+/**
+ * Temporarily prevent Remote Control Presets from renewing it's guids during duplicate operation.
+ */
+struct REMOTECONTROL_API FRCPresetGuidRenewGuard
+{
+	FRCPresetGuidRenewGuard();
+	~FRCPresetGuidRenewGuard();
+		
+private:
+	bool bPreviousValue;
+};
+
 /**
  * Represents a group of field and offers operations to operate on the fields inside of that group.
  */
