[RemoteControl]
Array Customization for Path behaviour to add Input button

diff --git a/Engine/Plugins/VirtualProduction/RemoteControl/Source/RemoteControlLogic/Private/Behaviour/Builtin/Path/RCSetAssetByPathBehaviour.cpp b/Engine/Plugins/VirtualProduction/RemoteControl/Source/RemoteControlLogic/Private/Behaviour/Builtin/Path/RCSetAssetByPathBehaviour.cpp
--- a/Engine/Plugins/VirtualProduction/RemoteControl/Source/RemoteControlLogic/Private/Behaviour/Builtin/Path/RCSetAssetByPathBehaviour.cpp
+++ b/Engine/Plugins/VirtualProduction/RemoteControl/Source/RemoteControlLogic/Private/Behaviour/Builtin/Path/RCSetAssetByPathBehaviour.cpp
@@ -50,6 +50,43 @@ void URCSetAssetByPathBehaviour::UpdateEntityIds(const TMap<FGuid, FGuid>& InEnt
 	Super::UpdateEntityIds(InEntityIdMap);
 }
 
+void URCSetAssetByPathBehaviour::PostLoad()
+{
+	Super::PostLoad();
+	PRAGMA_DISABLE_DEPRECATION_WARNINGS
+	// support for older version, loading the data of the deprecated property into the new one
+	if (!PathStruct.PathArray_DEPRECATED.IsEmpty())
+	{
+		PathStruct.AssetPath.Empty();
+	}
+	for (const FString& Path : PathStruct.PathArray_DEPRECATED)
+	{
+		bool bIsInput = false;
+		FString PathToCopy = Path;
+		if (PathToCopy.Contains(SetAssetByPathBehaviourHelpers::InputToken)
+			&& PathToCopy.Find(SetAssetByPathBehaviourHelpers::InputToken) == 0)
+		{
+			bIsInput = true;
+			PathToCopy.RemoveFromStart(SetAssetByPathBehaviourHelpers::InputToken);
+			int32 CharIndexStart;
+			int32 CharIndexEnd;
+			PathToCopy.FindChar('(', CharIndexStart);
+			PathToCopy.FindChar(')', CharIndexEnd);
+			if (CharIndexStart != INDEX_NONE)
+			{
+				PathToCopy.RemoveAt(CharIndexStart);
+			}
+			if (CharIndexEnd != INDEX_NONE)
+			{
+				PathToCopy.RemoveAt(CharIndexEnd-1);
+			}
+		}
+		PathStruct.AssetPath.Add(FRCAssetPathElement(bIsInput, PathToCopy));
+	}
+	PathStruct.PathArray_DEPRECATED.Empty();
+	PRAGMA_ENABLE_DEPRECATION_WARNINGS
+}
+
 bool URCSetAssetByPathBehaviour::SetAssetByPath(const FString& AssetPath, const FString& DefaultString)
 {
 	const URCController* Controller = ControllerWeakPtr.Get();
@@ -302,16 +339,14 @@ FString URCSetAssetByPathBehaviour::GetCurrentPath()
 
 	// Add Path String Concat
 	CurrentPath = bInternal ? SetAssetByPathBehaviourHelpers::ContentFolder : FString("");
-	for (FString PathPart : PathStruct.PathArray)
+	for (FRCAssetPathElement PathPart : PathStruct.AssetPath)
 	{
-		if (PathPart.IsEmpty())
+		if (PathPart.Path.IsEmpty())
 		{
 			continue;
 		}
-		
-		// Change *INPUT String Token to the correct Controller
-		if (PathPart.Contains(SetAssetByPathBehaviourHelpers::InputToken, ESearchCase::CaseSensitive)
-			&& PathPart.Find(SetAssetByPathBehaviourHelpers::InputToken) == 0)
+
+		if (PathPart.bIsInput)
 		{
 			const URemoteControlPreset* RemoteControlPreset = RCController->PresetWeakPtr.Get();
 			if (!RemoteControlPreset)
@@ -319,45 +354,32 @@ FString URCSetAssetByPathBehaviour::GetCurrentPath()
 				ensureMsgf(false, TEXT("Remote Control Preset is invalid"));
 				continue;
 			}
-			
-			PathPart.RemoveFromStart(SetAssetByPathBehaviourHelpers::InputToken);
-			int32 CharIndexStart;
-			int32 CharIndexEnd;
-			PathPart.FindChar('(', CharIndexStart);
-			PathPart.FindChar(')', CharIndexEnd);
-			if (CharIndexStart == INDEX_NONE || CharIndexEnd == INDEX_NONE)
-			{
-				// TODO: Add Message that Input token has been done wrongly.
-				ensureMsgf(false, TEXT("Token has been input wrongly"));
-				continue;
-			}
 
-			PathPart.RemoveAt(CharIndexStart);
-			PathPart.RemoveAt(CharIndexEnd-1);
-			const URCVirtualPropertyBase* TokenController = RemoteControlPreset->GetControllerByDisplayName(FName(PathPart));
+			const URCVirtualPropertyBase* TokenController = RemoteControlPreset->GetControllerByDisplayName(FName(PathPart.Path));
 			if (!TokenController)
 			{
-				ensureMsgf(false, TEXT("No Controller with given name found."));
-				continue;
+				PathPart.Path = TEXT("InvalidControllerName");
+			}
+			else
+			{
+				PathPart.Path = TokenController->GetDisplayValueAsString();
 			}
-
-			PathPart = TokenController->GetDisplayValueAsString();
 		}
 
 		// Check for certain chars.
 		int32 IndexFound;
-		if (PathPart.FindLastChar('_', IndexFound) && IndexFound == PathPart.Len() - 1)
+		if (PathPart.Path.FindLastChar('_', IndexFound) && IndexFound == PathPart.Path.Len() - 1)
 		{
 			// In the case there's underscore char in the at the end of one of the Path Strings, do nothing to facilitate more complex pathing behaviours.
 		}
-		else if (!PathPart.FindLastChar('/', IndexFound) || IndexFound < PathPart.Len() - 1)
+		else if (!PathPart.Path.FindLastChar('/', IndexFound) || IndexFound < PathPart.Path.Len() - 1)
 		{
-			PathPart = PathPart.AppendChar('/');
+			PathPart.Path = PathPart.Path.AppendChar('/');
 		}
 
-		PathPart.ReplaceCharInline(TCHAR('\\'), TCHAR('/'), ESearchCase::IgnoreCase);
+		PathPart.Path.ReplaceCharInline(TCHAR('\\'), TCHAR('/'), ESearchCase::IgnoreCase);
 		
-		CurrentPath += PathPart;
+		CurrentPath += PathPart.Path;
 	}
 
 	return CurrentPath;
@@ -494,9 +516,9 @@ void URCSetAssetByPathBehaviour::UpdateTargetEntity()
 
 void URCSetAssetByPathBehaviour::RefreshPathArray()
 {
-	if (PathStruct.PathArray.Num() < 1)
+	if (PathStruct.AssetPath.Num() < 1)
 	{
-		PathStruct.PathArray.Add("");
+		PathStruct.AssetPath.AddDefaulted();
 	}
 }
 
diff --git a/Engine/Plugins/VirtualProduction/RemoteControl/Source/RemoteControlLogic/Public/Behaviour/Builtin/Path/RCSetAssetByPathBehaviour.h b/Engine/Plugins/VirtualProduction/RemoteControl/Source/RemoteControlLogic/Public/Behaviour/Builtin/Path/RCSetAssetByPathBehaviour.h
--- a/Engine/Plugins/VirtualProduction/RemoteControl/Source/RemoteControlLogic/Public/Behaviour/Builtin/Path/RCSetAssetByPathBehaviour.h
+++ b/Engine/Plugins/VirtualProduction/RemoteControl/Source/RemoteControlLogic/Public/Behaviour/Builtin/Path/RCSetAssetByPathBehaviour.h
@@ -45,21 +45,56 @@ struct FPendingFileRequest
 	FString FileName;
 };
 
-/** Struct to help generate Widgts for the DetailsPanel of the Bahviour */
+USTRUCT()
+struct FRCAssetPathElement
+{
+	GENERATED_BODY()
+
+	FRCAssetPathElement()
+		: bIsInput(false)
+		, Path(TEXT(""))
+	{}
+
+	FRCAssetPathElement(bool bInIsInput, const FString& InPath)
+		: bIsInput(bInIsInput)
+		, Path(InPath)
+	{}
+
+	UPROPERTY(EditAnywhere, Category="Path")
+	bool bIsInput;
+
+	UPROPERTY(EditAnywhere, Category="Path")
+	FString Path;
+};
+
+/** Struct to help generate Widgets for the DetailsPanel of the Behaviour */
 USTRUCT()
 struct FRCSetAssetPath
 {
 	GENERATED_BODY()
 
+	PRAGMA_DISABLE_DEPRECATION_WARNINGS
+	// TODO [UE_DEPRECATED(5.4)] this rule of 5 is obligatory to avoid compile-time warning, delete them when the property is removed
 	FRCSetAssetPath()
 	{
-		// Add Empty index 0 when creating 
-		PathArray.Add("");
+		// Add Empty index 0 when creating
+		AssetPath.AddDefaulted();
 	}
-	
-	/** An Array of Strings holding the Path of an Asset, seperated in several String. Will concated back together later. */
+	~FRCSetAssetPath() = default;
+	FRCSetAssetPath(const FRCSetAssetPath&) = default;
+	FRCSetAssetPath(FRCSetAssetPath&&) = default;
+	FRCSetAssetPath& operator=(const FRCSetAssetPath&) = default;
+	FRCSetAssetPath& operator=(FRCSetAssetPath&&) = default;
+	PRAGMA_ENABLE_DEPRECATION_WARNINGS
+
+	/** An Array of Strings holding the Path of an Asset, seperated in several String. Will concatenated back together later. */
+	UPROPERTY(EditAnywhere, Category="Path")
+	TArray<FRCAssetPathElement> AssetPath;
+
+	/** An Array of Strings holding the Path of an Asset, seperated in several String. Will concatenated back together later. */
+	UE_DEPRECATED(5.4, "This property is deprecated please use AssetPathArray instead")
 	UPROPERTY()
-	TArray<FString> PathArray;
+	TArray<FString> PathArray_DEPRECATED;
 };
 
 /**
@@ -77,6 +112,8 @@ public:
 	virtual void UpdateEntityIds(const TMap<FGuid, FGuid>& InEntityIdMap) override;
 	//~ End URCBehaviour interface
 
+	virtual void PostLoad() override;
+
 	/** Given an Input Path, sets the Target Exposed Property to the Asset. */
 	bool SetAssetByPath(const FString& AssetPath, const FString& DefaultString);
 
diff --git a/Engine/Plugins/VirtualProduction/RemoteControl/Source/RemoteControlUI/Private/RemoteControlUIModule.cpp b/Engine/Plugins/VirtualProduction/RemoteControl/Source/RemoteControlUI/Private/RemoteControlUIModule.cpp
--- a/Engine/Plugins/VirtualProduction/RemoteControl/Source/RemoteControlUI/Private/RemoteControlUIModule.cpp
+++ b/Engine/Plugins/VirtualProduction/RemoteControl/Source/RemoteControlUI/Private/RemoteControlUIModule.cpp
@@ -5,6 +5,7 @@
 #include "AssetEditor/RemoteControlPresetEditorToolkit.h"
 #include "AssetTools/RemoteControlPresetActions.h"
 #include "AssetToolsModule.h"
+#include "Behaviour/Builtin/Path/RCSetAssetByPathBehaviour.h"
 #include "Commands/RemoteControlCommands.h"
 #include "Components/SkeletalMeshComponent.h"
 #include "Components/StaticMeshComponent.h"
@@ -39,6 +40,7 @@
 #include "UI/Controller/CustomControllers/SCustomTextureControllerWidget.h"
 #include "UI/Customizations/FPassphraseCustomization.h"
 #include "UI/Customizations/NetworkAddressCustomization.h"
+#include "UI/Customizations/RCAssetPathElementCustomization.h"
 #include "UI/Customizations/RemoteControlEntityCustomization.h"
 #include "UI/RemoteControlExposeMenuStyle.h"
 #include "UI/RemoteControlPanelStyle.h"
@@ -909,6 +911,7 @@ void FRemoteControlUIModule::RegisterStructCustomizations()
 		PropertyEditorModule.RegisterCustomClassLayout(Name, FOnGetDetailCustomizationInstance::CreateStatic(&FRemoteControlEntityCustomization::MakeInstance));
 	}
 
+	PropertyEditorModule.RegisterCustomPropertyTypeLayout(FRCAssetPathElement::StaticStruct()->GetFName(), FOnGetPropertyTypeCustomizationInstance::CreateStatic(&FRCAssetPathElementCustomization::MakeInstance));
 	PropertyEditorModule.RegisterCustomPropertyTypeLayout(FRCNetworkAddress::StaticStruct()->GetFName(), FOnGetPropertyTypeCustomizationInstance::CreateStatic(&FNetworkAddressCustomization::MakeInstance));
 	PropertyEditorModule.RegisterCustomPropertyTypeLayout(FRCPassphrase::StaticStruct()->GetFName(), FOnGetPropertyTypeCustomizationInstance::CreateStatic(&FPassphraseCustomization::MakeInstance));
 }
@@ -925,6 +928,7 @@ void FRemoteControlUIModule::UnregisterStructCustomizations()
 
 	if (UObjectInitialized())
 	{
+		PropertyEditorModule.UnregisterCustomPropertyTypeLayout(FRCAssetPathElement::StaticStruct()->GetFName());
 		PropertyEditorModule.UnregisterCustomPropertyTypeLayout(FRCNetworkAddress::StaticStruct()->GetFName());
 		PropertyEditorModule.UnregisterCustomPropertyTypeLayout(FRCPassphrase::StaticStruct()->GetFName());
 	}
diff --git a/Engine/Plugins/VirtualProduction/RemoteControl/Source/RemoteControlUI/Private/UI/Behaviour/Builtin/Path/RCBehaviourSetAssetByPathModel.cpp b/Engine/Plugins/VirtualProduction/RemoteControl/Source/RemoteControlUI/Private/UI/Behaviour/Builtin/Path/RCBehaviourSetAssetByPathModel.cpp
--- a/Engine/Plugins/VirtualProduction/RemoteControl/Source/RemoteControlUI/Private/UI/Behaviour/Builtin/Path/RCBehaviourSetAssetByPathModel.cpp
+++ b/Engine/Plugins/VirtualProduction/RemoteControl/Source/RemoteControlUI/Private/UI/Behaviour/Builtin/Path/RCBehaviourSetAssetByPathModel.cpp
@@ -25,10 +25,12 @@ FRCSetAssetByPathBehaviourModel::FRCSetAssetByPathBehaviourModel(URCSetAssetByPa
 	: FRCBehaviourModel(SetAssetByPathBehaviour)
 	, SetAssetByPathBehaviourWeakPtr(SetAssetByPathBehaviour)
 {
-	FPropertyRowGeneratorArgs Args;
-	Args.bShouldShowHiddenProperties = true;
-	PropertyRowGenerator = FModuleManager::GetModuleChecked<FPropertyEditorModule>("PropertyEditor").CreatePropertyRowGenerator(Args);
-	PropertyRowGeneratorArray = FModuleManager::GetModuleChecked<FPropertyEditorModule>("PropertyEditor").CreatePropertyRowGenerator(Args);
+	FPropertyRowGeneratorArgs ArgsRowGenerator;
+	ArgsRowGenerator.bShouldShowHiddenProperties = true;
+	PropertyRowGenerator = FModuleManager::GetModuleChecked<FPropertyEditorModule>("PropertyEditor").CreatePropertyRowGenerator(ArgsRowGenerator);
+
+	const FPropertyRowGeneratorArgs ArgsRowGeneratorArray;
+	PropertyRowGeneratorArray = FModuleManager::GetModuleChecked<FPropertyEditorModule>("PropertyEditor").CreatePropertyRowGenerator(ArgsRowGeneratorArray);
 
 	if (SetAssetByPathBehaviour)
 	{
@@ -48,12 +50,15 @@ FRCSetAssetByPathBehaviourModel::FRCSetAssetByPathBehaviourModel(URCSetAssetByPa
 		PropertyRowGeneratorArray->SetStructure(MakeShareable(new FStructOnScope(FRCSetAssetPath::StaticStruct(), (uint8*) &SetAssetByPathBehaviour->PathStruct)));
 		PropertyRowGeneratorArray->OnFinishedChangingProperties().AddLambda([this](const FPropertyChangedEvent& InEvent)
 		{
-			if (InEvent.ChangeType != EPropertyChangeType::ValueSet)
+			if (InEvent.ChangeType == EPropertyChangeType::ValueSet)
 			{
-				return;
+				// Just refresh the preview without reconstructing the widget if we just set values
+				RefreshPreview();
+			}
+			else
+			{
+				RefreshPathAndPreview();
 			}
-
-			RefreshPathAndPreview();
 		});
 		PropertyRowGeneratorArray->OnRowsRefreshed().AddLambda([this]()
 		{
@@ -240,71 +245,31 @@ void FRCSetAssetByPathBehaviourModel::RefreshPathAndPreview()
 		
 		for (uint8 Counter = 0; Counter < Children.Num(); Counter++)
 		{
-			TSharedPtr<SWidget> Button = SNullWidget::NullWidget;
-			if (Counter == 0)
-			{
-				Button = SNew(SBox)
-					.HAlign(HAlign_Center)
-					.VAlign(VAlign_Center)
-					.WidthOverride(22)
-					.HeightOverride(22)
-					[
-						SNew(SButton)
-						.ButtonStyle( FAppStyle::Get(), "SimpleButton" )
-						.OnClicked_Lambda([this]()
-						{
-							if (URCSetAssetByPathBehaviour* Behaviour = Cast<URCSetAssetByPathBehaviour>(GetBehaviour()))
-							{
-								Behaviour->Modify();
-								CreateAssetPathFromSelection();
-								RefreshPathAndPreview();
-							}
-
-							return FReply::Handled();
-						})
-						.IsFocusable(false)
-						.ContentPadding(0)
-						.Visibility_Lambda([this]()
-						{
-							URCSetAssetByPathBehaviour* Behaviour = Cast<URCSetAssetByPathBehaviour>(GetBehaviour());
-							if (!Behaviour)
-							{
-								return EVisibility::Collapsed;
-							}
-
-							return Behaviour->bInternal ? EVisibility::Visible : EVisibility::Collapsed;
-						})
-						.Content()
-						[
-							SNew(SImage)
-							.Image( FAppStyle::GetBrush("Icons.Use") )
-							.ColorAndOpacity( FSlateColor::UseForeground() )
-						]
-					];
-			}
-			
 			const FNodeWidgets NodeWidgets = Children[Counter]->CreateNodeWidgets();
 			if (NodeWidgets.ValueWidget)
 			{
-				FieldArrayWidget->AddSlot()
-					.Padding(FMargin(3.0f, 2.0f))
-					.AutoHeight()
+				TSharedPtr<SHorizontalBox> CurrentHorizontalBox = SNew(SHorizontalBox);
+				if (Counter == 0)
+				{
+					// Add NameWidget only for the ArrayProperty not for its entries
+					CurrentHorizontalBox->AddSlot()
 					[
-						SNew(SHorizontalBox)
-						+ SHorizontalBox::Slot()
-						[
-							NodeWidgets.NameWidget.ToSharedRef()
-						]
-						+ SHorizontalBox::Slot()
-						[
-							NodeWidgets.ValueWidget.ToSharedRef()
-						]
-						+ SHorizontalBox::Slot()
-						.AutoWidth()
-						[
-							Button.ToSharedRef()
-						]
+						NodeWidgets.NameWidget.ToSharedRef()
 					];
+				}
+				CurrentHorizontalBox->AddSlot()
+				.FillWidth(1.f)
+				[
+					NodeWidgets.ValueWidget.ToSharedRef()
+				];
+
+				FieldArrayWidget->AddSlot()
+				.Padding(FMargin(3.0f, 2.0f))
+				.AutoHeight()
+				.VAlign(VAlign_Center)
+				[
+					CurrentHorizontalBox.ToSharedRef()
+				];
 			}
 			else if (NodeWidgets.WholeRowWidget)
 			{
@@ -318,13 +283,17 @@ void FRCSetAssetByPathBehaviourModel::RefreshPathAndPreview()
 		}
 	}
 
+	RefreshPreview();
+	PathArrayWidget->SetContent(FieldArrayWidget.ToSharedRef());
+}
+
+void FRCSetAssetByPathBehaviourModel::RefreshPreview() const
+{
 	FString PreviewText = TEXT("Path is invalid!");
 	if (URCSetAssetByPathBehaviour* SetAssetByPathBehaviour = SetAssetByPathBehaviourWeakPtr.Get())
 	{
 		PreviewText = SetAssetByPathBehaviour->GetCurrentPath();
 	}
-	
-	PathArrayWidget->SetContent(FieldArrayWidget.ToSharedRef());
 	PreviewPathWidget->SetText(FText::FromString(PreviewText));
 }
 
@@ -396,40 +365,4 @@ FText FRCSetAssetByPathBehaviourModel::GetSelectedEntityText() const
 	return ReturnText;
 }
 
-
-void FRCSetAssetByPathBehaviourModel::CreateAssetPathFromSelection() const
-{
-	URCSetAssetByPathBehaviour* SetAssetByPathBehaviour = Cast<URCSetAssetByPathBehaviour>(GetBehaviour());
-	if (!SetAssetByPathBehaviour)
-	{
-		return;
-	}
-	
-	TArray<FAssetData> AssetData;
-	GEditor->GetContentBrowserSelections(AssetData);
-	if (AssetData.Num() != 1)
-	{
-		// Only works with one selected Content
-		return;
-	}
-	
-	// Clear it, in case it is an already used one.
-	TArray<FString>& PathArray = SetAssetByPathBehaviour->PathStruct.PathArray;
-	PathArray.Empty();
-
-	int32 IndexOfLast;
-	const UObject* SelectedAsset = AssetData[0].GetAsset();
-	FString PathString = SelectedAsset->GetPathName();
-
-	// Remove the initial Game
-	PathString.RemoveFromStart("/Game/");
-
-	// Remove anything after the last /
-	PathString.FindLastChar('/', IndexOfLast);
-	PathString.RemoveAt(IndexOfLast, PathString.Len() - IndexOfLast);
-
-	PathArray.Add(PathString);
-}
-
-
 #undef LOCTEXT_NAMESPACE
diff --git a/Engine/Plugins/VirtualProduction/RemoteControl/Source/RemoteControlUI/Private/UI/Behaviour/Builtin/Path/RCBehaviourSetAssetByPathModel.h b/Engine/Plugins/VirtualProduction/RemoteControl/Source/RemoteControlUI/Private/UI/Behaviour/Builtin/Path/RCBehaviourSetAssetByPathModel.h
--- a/Engine/Plugins/VirtualProduction/RemoteControl/Source/RemoteControlUI/Private/UI/Behaviour/Builtin/Path/RCBehaviourSetAssetByPathModel.h
+++ b/Engine/Plugins/VirtualProduction/RemoteControl/Source/RemoteControlUI/Private/UI/Behaviour/Builtin/Path/RCBehaviourSetAssetByPathModel.h
@@ -80,6 +80,5 @@ private:
 	/** Regenerates and creates a new PathArray Widget if changed */
 	void RefreshPathAndPreview();
 
-	/** Creates a asset path, based on the selected Asset in the content browser. */
-	void CreateAssetPathFromSelection() const;
+	void RefreshPreview() const;
 };
diff --git a/Engine/Plugins/VirtualProduction/RemoteControl/Source/RemoteControlUI/Private/UI/Customizations/RCAssetPathElementCustomization.cpp b/Engine/Plugins/VirtualProduction/RemoteControl/Source/RemoteControlUI/Private/UI/Customizations/RCAssetPathElementCustomization.cpp
new file mode 100644
--- /dev/null
+++ b/Engine/Plugins/VirtualProduction/RemoteControl/Source/RemoteControlUI/Private/UI/Customizations/RCAssetPathElementCustomization.cpp
@@ -0,0 +1,162 @@
+// Copyright Epic Games, Inc. All Rights Reserved.
+
+#include "RCAssetPathElementCustomization.h"
+
+#include "Behaviour/Builtin/Path/RCSetAssetByPathBehaviour.h"
+#include "DetailLayoutBuilder.h"
+#include "Editor.h"
+#include "Framework/Application/SlateApplication.h"
+#include "IDetailChildrenBuilder.h"
+#include "PropertyCustomizationHelpers.h"
+#include "UI/RemoteControlPanelStyle.h"
+#include "Widgets/Images/SImage.h"
+#include "Widgets/Input/SButton.h"
+#include "Widgets/Input/SCheckBox.h"
+#include "Widgets/Input/SEditableTextBox.h"
+#include "Widgets/SToolTip.h"
+#include "Widgets/Text/STextBlock.h"
+
+#define LOCTEXT_NAMESPACE "RCAssetPathElementCustomization"
+
+TSharedRef<IPropertyTypeCustomization> FRCAssetPathElementCustomization::MakeInstance()
+{
+	return MakeShared<FRCAssetPathElementCustomization>();
+}
+
+void FRCAssetPathElementCustomization::CustomizeHeader(TSharedRef<IPropertyHandle> InPropertyHandle, FDetailWidgetRow& InHeaderRow, IPropertyTypeCustomizationUtils& InCustomizationUtils)
+{
+	ArrayEntryHandle = InPropertyHandle;
+	IsInputHandle = ArrayEntryHandle->GetChildHandle(GET_MEMBER_NAME_CHECKED(FRCAssetPathElement, bIsInput));
+	PathHandle = ArrayEntryHandle->GetChildHandle(GET_MEMBER_NAME_CHECKED(FRCAssetPathElement, Path));
+
+	if (!IsInputHandle.IsValid() || !PathHandle.IsValid())
+	{
+		return;
+	}
+
+	const TSharedRef<SToolTip> GetAssetPathToolTipWidget = SNew(SToolTip)
+		.Text(LOCTEXT("RCGetAssetPathButton_Tooltip", "Get the path of the currently first selected asset in the content browser and set it to the current path"));
+
+	InHeaderRow.NameContent()
+		[
+			InPropertyHandle->CreatePropertyNameWidget()
+		];
+
+	InHeaderRow.ValueContent()
+	[
+		SNew(SHorizontalBox)
+		+SHorizontalBox::Slot()
+		.AutoWidth()
+		.Padding(5.f, 0.f)
+		.VAlign(VAlign_Center)
+		[
+			SNew(SCheckBox)
+			.Style(&FRemoteControlPanelStyle::Get()->GetWidgetStyle<FCheckBoxStyle>("RemoteControlPathBehaviour.AssetCheckBox"))
+			.IsChecked(this, &FRCAssetPathElementCustomization::IsChecked)
+			.OnCheckStateChanged(this, &FRCAssetPathElementCustomization::OnCheckStateChanged)
+			.IsFocusable(false)
+			[
+				SNew(STextBlock)
+				.Text(LOCTEXT("RCInputButtonAssetPath", "RCInput"))
+				.Font(IDetailLayoutBuilder::GetDetailFont())
+			]
+		]
+
+		+SHorizontalBox::Slot()
+		.FillWidth(1.f)
+		.VAlign(VAlign_Center)
+		[
+			PathHandle->CreatePropertyValueWidget(false)
+		]
+
+		+SHorizontalBox::Slot()
+		.AutoWidth()
+		.VAlign(VAlign_Center)
+		[
+			SNew(SButton)
+			.ButtonStyle(FAppStyle::Get(), "SimpleButton")
+			.OnClicked(this, &FRCAssetPathElementCustomization::OnGetAssetFromSelectionClicked)
+			.ToolTip(GetAssetPathToolTipWidget)
+			.IsFocusable(false)
+			.ContentPadding(0)
+			.Content()
+			[
+				SNew(SImage)
+				.Image(FAppStyle::GetBrush("Icons.Use"))
+				.ColorAndOpacity(FSlateColor::UseForeground())
+			]
+		]
+	];
+}
+
+void FRCAssetPathElementCustomization::CustomizeChildren(TSharedRef<IPropertyHandle> InPropertyHandle, IDetailChildrenBuilder& InChildBuilder, IPropertyTypeCustomizationUtils& InCustomizationUtils)
+{
+}
+
+ECheckBoxState FRCAssetPathElementCustomization::IsChecked() const
+{
+	ECheckBoxState ReturnValue = ECheckBoxState::Undetermined;
+	if (IsInputHandle.IsValid())
+	{
+		bool bIsInput;
+		const FPropertyAccess::Result Result = IsInputHandle->GetValue(bIsInput);
+		if (Result == FPropertyAccess::Success)
+		{
+			ReturnValue = bIsInput ? ECheckBoxState::Checked : ECheckBoxState::Unchecked;
+		}
+	}
+	return ReturnValue;
+}
+
+void FRCAssetPathElementCustomization::OnCheckStateChanged(ECheckBoxState InNewState) const
+{
+	if (!IsInputHandle.IsValid())
+	{
+		return;
+	}
+	IsInputHandle->SetValue(InNewState == ECheckBoxState::Checked ? true : false);
+}
+
+FReply FRCAssetPathElementCustomization::OnGetAssetFromSelectionClicked() const
+{
+	TArray<FAssetData> AssetData;
+	GEditor->GetContentBrowserSelections(AssetData);
+	if (AssetData.Num() <= 0)
+	{
+		return FReply::Handled();
+	}
+				
+	// Clear it, in case it is an already used one.
+	// Use the first one in the Array
+	int32 IndexOfLast;
+	const UObject* SelectedAsset = AssetData[0].GetAsset();
+	FString PathString = SelectedAsset->GetPathName();
+
+	// Remove the initial Game
+	PathString.RemoveFromStart("/Game/");
+
+	// Remove anything after the last /
+	PathString.FindLastChar('/', IndexOfLast);
+	if (IndexOfLast != INDEX_NONE)
+	{
+		PathString.RemoveAt(IndexOfLast, PathString.Len() - IndexOfLast);
+	}
+	else
+	{
+		// if the Index is -1 then it means that we are selecting an asset already in the topmost folder
+		// So we clear the string since it will just contains the AssetName
+		PathString.Empty();
+	}
+
+	if (IsInputHandle.IsValid())
+	{
+		IsInputHandle->SetValue(false);
+	}
+	if (PathHandle.IsValid())
+	{
+		PathHandle->SetValue(PathString);
+	}
+	return FReply::Handled();
+}
+
+#undef LOCTEXT_NAMESPACE
diff --git a/Engine/Plugins/VirtualProduction/RemoteControl/Source/RemoteControlUI/Private/UI/Customizations/RCAssetPathElementCustomization.h b/Engine/Plugins/VirtualProduction/RemoteControl/Source/RemoteControlUI/Private/UI/Customizations/RCAssetPathElementCustomization.h
new file mode 100644
--- /dev/null
+++ b/Engine/Plugins/VirtualProduction/RemoteControl/Source/RemoteControlUI/Private/UI/Customizations/RCAssetPathElementCustomization.h
@@ -0,0 +1,36 @@
+// Copyright Epic Games, Inc. All Rights Reserved.
+#pragma once
+
+#include "IPropertyTypeCustomization.h"
+#include "Input/Reply.h"
+#include "Templates/SharedPointer.h"
+#include "Styling/SlateTypes.h"
+
+class FDetailWidgetRow;
+class IDetailChildrenBuilder;
+class IPropertyHandle;
+
+class FRCAssetPathElementCustomization : public IPropertyTypeCustomization
+{
+public:
+	static TSharedRef<IPropertyTypeCustomization> MakeInstance();
+
+	//~ Begin IPropertyTypeCustomization
+	virtual void CustomizeHeader(TSharedRef<IPropertyHandle> InPropertyHandle, FDetailWidgetRow& InHeaderRow, IPropertyTypeCustomizationUtils& InCustomizationUtils) override;
+	virtual void CustomizeChildren(TSharedRef<IPropertyHandle> InPropertyHandle, IDetailChildrenBuilder& InChildBuilder, IPropertyTypeCustomizationUtils& InCustomizationUtils) override;
+	//~ End IPropertyTypeCustomization
+
+private:
+	/** Returns if the CheckBox is checked or not */
+	ECheckBoxState IsChecked() const;
+	/** Callback called when you click on the CheckBox to enable/disable it */
+	void OnCheckStateChanged(ECheckBoxState InNewState) const;
+	/** Callback called when you click on the arrow button next to the entries to retrieve the path of the selected Asset */
+	FReply OnGetAssetFromSelectionClicked() const;
+
+private:
+	TSharedPtr<SWidget> PathWidget;
+	TSharedPtr<IPropertyHandle> ArrayEntryHandle;
+	TSharedPtr<IPropertyHandle> IsInputHandle;
+	TSharedPtr<IPropertyHandle> PathHandle;
+};
diff --git a/Engine/Plugins/VirtualProduction/RemoteControl/Source/RemoteControlUI/Private/UI/RemoteControlPanelStyle.cpp b/Engine/Plugins/VirtualProduction/RemoteControl/Source/RemoteControlUI/Private/UI/RemoteControlPanelStyle.cpp
--- a/Engine/Plugins/VirtualProduction/RemoteControl/Source/RemoteControlUI/Private/UI/RemoteControlPanelStyle.cpp
+++ b/Engine/Plugins/VirtualProduction/RemoteControl/Source/RemoteControlUI/Private/UI/RemoteControlPanelStyle.cpp
@@ -192,6 +192,11 @@ void FRemoteControlPanelStyle::SetupPanelStyles(TSharedRef<FSlateStyleSet> InSty
 	SwitchButtonStyle.SetCheckBoxType(ESlateCheckBoxType::ToggleButton);
 	SwitchButtonStyle.SetPadding(FMargin(4.f, 2.f));
 
+	FCheckBoxStyle AssetPathToggleButtonStyle = AppStyle.GetWidgetStyle<FCheckBoxStyle>("ToggleButtonCheckbox");
+	AssetPathToggleButtonStyle.SetCheckBoxType(ESlateCheckBoxType::ToggleButton);
+	AssetPathToggleButtonStyle.SetUncheckedImage(BOX_BRUSH("Common/RoundedSelection_16x", 4.0f / 16.0f, FStyleColors::Background));
+
+	StyleSet->Set("RemoteControlPathBehaviour.AssetCheckBox", AssetPathToggleButtonStyle);
 	// Text Styles
 	const FStyleFonts& StyleFonts = FStyleFonts::Get();
 
