[RemoteControl]
Visual errors in viewport when using float controller or rotators to rotate object

diff --git a/Engine/Plugins/VirtualProduction/RemoteControl/Source/RemoteControl/Private/RemoteControlPropertyHandle.cpp b/Engine/Plugins/VirtualProduction/RemoteControl/Source/RemoteControl/Private/RemoteControlPropertyHandle.cpp
--- a/Engine/Plugins/VirtualProduction/RemoteControl/Source/RemoteControl/Private/RemoteControlPropertyHandle.cpp
+++ b/Engine/Plugins/VirtualProduction/RemoteControl/Source/RemoteControl/Private/RemoteControlPropertyHandle.cpp
@@ -52,6 +52,14 @@ namespace
 		return nullptr;
 	}
 
+	/** Add the given property and the given value to the Writer */
+	template<typename ValueType>
+	void AddPropertyToWriter(FCborWriter& OutWriter, const FString InPropName, const ValueType& InValue)
+	{
+		OutWriter.WriteValue(InPropName);
+		OutWriter.WriteValue(InValue);
+	}
+
 	/**
 	 * Get the property value for given property handle
 	 */
@@ -211,6 +219,71 @@ namespace
 
 		return bSuccess;
 	}
+
+	/**
+	 * Set property value for the given StructHandle
+	 * Setter following standard SetObjectProperties path from IRemoteControlModule
+	 * And follow the replication path
+	 */
+	template<typename InValueType>
+	static bool SetStructPropertyValues(const FRemoteControlPropertyHandle& InPropertyHandle, const TMap<FString, InValueType>& InPropNameToValues)
+	{
+		FRCFieldPathInfo RCFieldPathInfo(InPropertyHandle.GetFieldPath());
+		const TSharedPtr<FRemoteControlProperty> RCProperty = InPropertyHandle.GetRCProperty();
+		if (!ensure(RCProperty.IsValid()))
+		{
+			return false;
+		}
+	
+		const FProperty* Property = InPropertyHandle.GetProperty();
+		if (!ensure(Property))
+		{
+			return false;
+		}
+	
+		if (!ensure(Property->IsA<FStructProperty>()))
+		{
+			return false;
+		}
+
+		TArray<uint8> ApiBuffer;
+		FMemoryWriter MemoryWriter(ApiBuffer);
+		FMemoryReader MemoryReader(ApiBuffer);
+		FCborWriter CborWriter(&MemoryWriter);
+
+		// START: Struct Writer
+		CborWriter.WriteContainerStart(ECborCode::Map, -1);
+		CborWriter.WriteValue(Property->GetName());
+
+		// START: Struct properties Writer
+		CborWriter.WriteContainerStart(ECborCode::Map, -1);
+		for (const TPair<FString, InValueType>& PropToValue : InPropNameToValues)
+		{
+			AddPropertyToWriter(CborWriter, PropToValue.Key, PropToValue.Value);
+		}
+		// END: Struct properties Writer
+		CborWriter.WriteContainerEnd();
+
+		// END: Struct Writer
+		CborWriter.WriteContainerEnd();
+
+		FCborStructDeserializerBackend CborStructDeserializerBackend(MemoryReader);
+		// Set object reference
+		FRCObjectReference ObjectRef;
+		ObjectRef.Property = Property;
+		ObjectRef.Access = InPropertyHandle.ShouldGenerateTransaction() ? ERCAccess::WRITE_TRANSACTION_ACCESS : ERCAccess::WRITE_ACCESS;
+		ObjectRef.PropertyPathInfo = RCFieldPathInfo;
+
+		bool bSuccess = true;
+		for (UObject* Object : RCProperty->GetBoundObjects())
+		{
+			IRemoteControlModule::Get().ResolveObjectProperty(ObjectRef.Access, Object, ObjectRef.PropertyPathInfo, ObjectRef);
+			// Set object property and follow the replication path if there are any replicators
+			bSuccess &= IRemoteControlModule::Get().SetObjectProperties(ObjectRef, CborStructDeserializerBackend, ERCPayloadType::Cbor, ApiBuffer);
+		}
+
+		return bSuccess;
+	}
 }
 
 TSharedPtr<IRemoteControlPropertyHandle> IRemoteControlPropertyHandle::GetPropertyHandle(FName PresetName, FGuid PropertyId)
@@ -1514,12 +1587,21 @@ bool FRemoteControlPropertyHandleRotator::SetValue(FRotator InValue)
 		return false;
 	}
 
-	// To set the value from the rotator we set each child. 
-	bool ResR = RollValue.Pin()->SetValue(InValue.Roll);
-	bool ResP = PitchValue.Pin()->SetValue(InValue.Pitch);
-	bool ResY = YawValue.Pin()->SetValue(InValue.Yaw);
+	const FProperty* RollProperty = RollValue.Pin()->GetProperty();
+	const FProperty* PitchProperty = PitchValue.Pin()->GetProperty();
+	const FProperty* YawProperty = YawValue.Pin()->GetProperty();
+	if (!(RollProperty && PitchProperty && YawProperty))
+	{
+		return false;
+	}
 
-	if (ResR == false || ResP == false || ResY == false)
+	TMap<FString, double> PropertyToValues;
+	PropertyToValues.Add(PitchProperty->GetName(), InValue.Pitch);
+	PropertyToValues.Add(YawProperty->GetName(), InValue.Yaw);
+	PropertyToValues.Add(RollProperty->GetName(), InValue.Roll);
+
+	bool Res = SetStructPropertyValues(*this, MoveTemp(PropertyToValues));
+	if (Res == false)
 	{
 		return false;
 	}
