[RemoteControl]
Fix automation Test failing

diff --git a/Engine/Plugins/VirtualProduction/RemoteControl/Source/RemoteControl/Private/RemoteControlModule.cpp b/Engine/Plugins/VirtualProduction/RemoteControl/Source/RemoteControl/Private/RemoteControlModule.cpp
--- a/Engine/Plugins/VirtualProduction/RemoteControl/Source/RemoteControl/Private/RemoteControlModule.cpp
+++ b/Engine/Plugins/VirtualProduction/RemoteControl/Source/RemoteControl/Private/RemoteControlModule.cpp
@@ -2730,7 +2730,7 @@ bool FRemoteControlModule::CanInterceptFunction(const FRCCall& RCCall) const
 #else
 				FString FunctionName = RCCall.CallRef.Function->GetName();
 #endif
-				REMOTE_CONTROL_LOG_ONCE(Warning, TEXT("Function \"%s\" on object \"%s\" could not be intercepted because it contains out parameters."), *FunctionName, *RCCall.CallRef.Object->GetName());
+				REMOTE_CONTROL_LOG_ONCE(Log, TEXT("Function \"%s\" on object \"%s\" could not be intercepted because it contains out parameters."), *FunctionName, *RCCall.CallRef.Object->GetName());
 				return false;
 			}
 		}
diff --git a/Engine/Plugins/VirtualProduction/RemoteControl/Source/RemoteControlLogic/Private/Tests/RemoteControlLogicTests.cpp b/Engine/Plugins/VirtualProduction/RemoteControl/Source/RemoteControlLogic/Private/Tests/RemoteControlLogicTests.cpp
--- a/Engine/Plugins/VirtualProduction/RemoteControl/Source/RemoteControlLogic/Private/Tests/RemoteControlLogicTests.cpp
+++ b/Engine/Plugins/VirtualProduction/RemoteControl/Source/RemoteControlLogic/Private/Tests/RemoteControlLogicTests.cpp
@@ -71,7 +71,7 @@ bool FRemoteControlLogicTest::RunTest(const FString& Parameters)
 	// 4. Add Behaviour To Properties
 	// 4.1 Add Conditional Behaviour
 	URCBehaviour* FloatControllerBehaviour = FloatController->AddBehaviour(URCBehaviourConditionalNode::StaticClass());
-	const URCBehaviourConditional* ConditionalBehaviour = Cast<URCBehaviourConditional>(FloatControllerBehaviour);
+	URCBehaviourConditional* ConditionalBehaviour = Cast<URCBehaviourConditional>(FloatControllerBehaviour);
 	float IsEqualBehaviourFloatValue = 0.f;
  
 	// @todo: Add Conditional Behaviour Unit Tests here.
@@ -98,8 +98,21 @@ bool FRemoteControlLogicTest::RunTest(const FString& Parameters)
 	URCPropertyAction* StrControllerBehaviourAction = Cast<URCPropertyAction> (StrControllerBehaviour->ActionContainer->AddAction(RCProp1));
 	FColor StringControllerColorValue(7,8,9,10);
 	TestTrue(TEXT("Should Set Color"), StrControllerBehaviourAction->PropertySelfContainer->SetValueColor(StringControllerColorValue));
-	
- 
+
+	// 5.41 Add Condition to behaviour
+	FRCBehaviourCondition ConditionProp;
+	ConditionProp.ConditionType = ERCBehaviourConditionType::IsEqual;
+	ConditionProp.Comparand = NewObject<URCVirtualPropertySelfContainer>();
+	ConditionProp.Comparand->AddProperty(FName(TEXT("Prop_Condition")), EPropertyBagPropertyType::Float);
+	ConditionProp.Comparand->SetValueFloat(0.65f);
+	ConditionalBehaviour->Conditions.Add(FloatControllerBehaviourAction, ConditionProp);
+	FRCBehaviourCondition ConditionFunc;
+	ConditionFunc.ConditionType = ERCBehaviourConditionType::IsEqual;
+	ConditionFunc.Comparand = NewObject<URCVirtualPropertySelfContainer>();
+	ConditionFunc.Comparand->AddProperty(FName(TEXT("Func_Condition")), EPropertyBagPropertyType::Float);
+	ConditionFunc.Comparand->SetValueFloat(0.65f);
+	ConditionalBehaviour->Conditions.Add(FloatControllerBehaviourAction1, ConditionFunc);
+
 	// 5.5 It should not add same action for second time
 	URCAction* StrControllerBehaviourAction1 = StrControllerBehaviour->ActionContainer->AddAction(RCProp1);
 	
