[RemoteControl]
Controller Description placeholder to understand its writeable

diff --git a/Engine/Plugins/VirtualProduction/RemoteControl/Source/RemoteControlUI/Private/UI/Controller/RCControllerModel.cpp b/Engine/Plugins/VirtualProduction/RemoteControl/Source/RemoteControlUI/Private/UI/Controller/RCControllerModel.cpp
--- a/Engine/Plugins/VirtualProduction/RemoteControl/Source/RemoteControlUI/Private/UI/Controller/RCControllerModel.cpp
+++ b/Engine/Plugins/VirtualProduction/RemoteControl/Source/RemoteControlUI/Private/UI/Controller/RCControllerModel.cpp
@@ -117,7 +117,22 @@ TSharedRef<SWidget> FRCControllerModel::GetDescriptionWidget() const
 {
 	return SNew(SBox).Padding(10.f, 2.f)
 		[
-			ControllerDescriptionTextBox.ToSharedRef()
+			SNew(SOverlay)
+			+SOverlay::Slot()
+			.VAlign(VAlign_Center)
+			[
+				ControllerDescriptionTextBox.ToSharedRef()
+			]
+
+			+SOverlay::Slot()
+			.HAlign(HAlign_Left)
+			.VAlign(VAlign_Center)
+			[
+				SNew(STextBlock)
+				.Text(LOCTEXT("RemoteControlDescriptionPlaceholder", "Double click to change description"))
+				.IsEnabled(false)
+				.Visibility(this, &FRCControllerModel::GetPlaceholderVisibility)
+			]
 		];
 }
 
@@ -235,15 +250,12 @@ void FRCControllerModel::OnControllerNameCommitted(const FText& InNewControllerN
 
 void FRCControllerModel::OnControllerDescriptionCommitted(const FText& InNewControllerDescription, ETextCommit::Type InCommitInfo)
 {
-	if (URemoteControlPreset* Preset = GetPreset())
+	if (URCVirtualPropertyBase* Controller = GetVirtualProperty())
 	{
-		if (URCVirtualPropertyBase* Controller = GetVirtualProperty())
-		{
-			FScopedTransaction Transaction(LOCTEXT("ChangedControllerDescription", "Update controller description"));
-			Controller->Modify();
-			Controller->Description = InNewControllerDescription;
-			ControllerDescriptionTextBox->SetText(InNewControllerDescription);
-		}
+		FScopedTransaction Transaction(LOCTEXT("ChangedControllerDescription", "Update controller description"));
+		Controller->Modify();
+		Controller->Description = InNewControllerDescription;
+		ControllerDescriptionTextBox->SetText(InNewControllerDescription);
 	}
 }
 
@@ -297,6 +309,19 @@ void FRCControllerModel::OnPropertyValueChanged(const FPropertyChangedEvent& InP
 	}
 }
 
+EVisibility FRCControllerModel::GetPlaceholderVisibility() const
+{
+	if (ControllerDescriptionTextBox.IsValid())
+	{
+		if (ControllerDescriptionTextBox->IsInEditMode())
+		{
+			return EVisibility::Collapsed;
+		}
+		return ControllerDescriptionTextBox->GetText().IsEmpty() ? EVisibility::Visible : EVisibility:: Collapsed;
+	}
+	return EVisibility::Collapsed;
+}
+
 void FRCControllerModel::InitControlledTypes()
 {
 	ControlledTypesAsStrings.Empty();
diff --git a/Engine/Plugins/VirtualProduction/RemoteControl/Source/RemoteControlUI/Private/UI/Controller/RCControllerModel.h b/Engine/Plugins/VirtualProduction/RemoteControl/Source/RemoteControlUI/Private/UI/Controller/RCControllerModel.h
--- a/Engine/Plugins/VirtualProduction/RemoteControl/Source/RemoteControlUI/Private/UI/Controller/RCControllerModel.h
+++ b/Engine/Plugins/VirtualProduction/RemoteControl/Source/RemoteControlUI/Private/UI/Controller/RCControllerModel.h
@@ -113,6 +113,9 @@ private:
 	/** Value change event */
 	void OnPropertyValueChanged(const FPropertyChangedEvent& InPropertyChangedEvent);
 
+	/** Get the current visibility of the placeholder TextBlock widget */
+	EVisibility GetPlaceholderVisibility() const;
+
 	/** Initializes the list of controlled types. Used by MultiControllers */
 	void InitControlledTypes();
 
