[RemoteControl]
Select first behaviour when selecting controller

diff --git a/Engine/Plugins/VirtualProduction/RemoteControl/Source/RemoteControlUI/Private/UI/Behaviour/SRCBehaviourPanel.cpp b/Engine/Plugins/VirtualProduction/RemoteControl/Source/RemoteControlUI/Private/UI/Behaviour/SRCBehaviourPanel.cpp
--- a/Engine/Plugins/VirtualProduction/RemoteControl/Source/RemoteControlUI/Private/UI/Behaviour/SRCBehaviourPanel.cpp
+++ b/Engine/Plugins/VirtualProduction/RemoteControl/Source/RemoteControlUI/Private/UI/Behaviour/SRCBehaviourPanel.cpp
@@ -65,6 +65,10 @@ void SRCBehaviourPanel::OnControllerSelectionChanged(TSharedPtr<FRCControllerMod
 {
 	SelectedControllerItemWeakPtr = InControllerItem;
 	UpdateWrappedWidget(InControllerItem);
+	if (BehaviourPanelList.IsValid())
+	{
+		BehaviourPanelList->SelectFirstItem();
+	}
 }
 
 void SRCBehaviourPanel::UpdateWrappedWidget(TSharedPtr<FRCControllerModel> InControllerItem)
diff --git a/Engine/Plugins/VirtualProduction/RemoteControl/Source/RemoteControlUI/Private/UI/Behaviour/SRCBehaviourPanelList.cpp b/Engine/Plugins/VirtualProduction/RemoteControl/Source/RemoteControlUI/Private/UI/Behaviour/SRCBehaviourPanelList.cpp
--- a/Engine/Plugins/VirtualProduction/RemoteControl/Source/RemoteControlUI/Private/UI/Behaviour/SRCBehaviourPanelList.cpp
+++ b/Engine/Plugins/VirtualProduction/RemoteControl/Source/RemoteControlUI/Private/UI/Behaviour/SRCBehaviourPanelList.cpp
@@ -113,6 +113,14 @@ void SRCBehaviourPanelList::AddSpecialContextMenuOptions(FMenuBuilder& MenuBuild
 	MenuBuilder.AddMenuEntry(DisableLabel, DisableTooltip, FSlateIcon(), DisableAction);
 }
 
+void SRCBehaviourPanelList::SelectFirstItem()
+{
+	if (!BehaviourItems.IsEmpty())
+	{
+		ListView->SetSelection(BehaviourItems[0]);
+	}
+}
+
 void SRCBehaviourPanelList::SetIsBehaviourEnabled(const bool bIsEnabled)
 {
 	// Disable all selected behaviour here
diff --git a/Engine/Plugins/VirtualProduction/RemoteControl/Source/RemoteControlUI/Private/UI/Behaviour/SRCBehaviourPanelList.h b/Engine/Plugins/VirtualProduction/RemoteControl/Source/RemoteControlUI/Private/UI/Behaviour/SRCBehaviourPanelList.h
--- a/Engine/Plugins/VirtualProduction/RemoteControl/Source/RemoteControlUI/Private/UI/Behaviour/SRCBehaviourPanelList.h
+++ b/Engine/Plugins/VirtualProduction/RemoteControl/Source/RemoteControlUI/Private/UI/Behaviour/SRCBehaviourPanelList.h
@@ -65,6 +65,9 @@ public:
 	/** Allows Logic panels to add special functionality to the Context Menu based on context */
 	virtual void AddSpecialContextMenuOptions(FMenuBuilder& MenuBuilder) override;
 
+	/** Will set the selection to the first item */
+	void SelectFirstItem();
+
 private:
 
 	/** Enables or Disables the currently selected behaviour */
