[RemoteControl]
Add Group Id Drag and change it to create PropId Action

diff --git a/Engine/Plugins/VirtualProduction/RemoteControl/Source/RemoteControlLogic/Public/Behaviour/Builtin/Bind/RCBehaviourBind.h b/Engine/Plugins/VirtualProduction/RemoteControl/Source/RemoteControlLogic/Public/Behaviour/Builtin/Bind/RCBehaviourBind.h
--- a/Engine/Plugins/VirtualProduction/RemoteControl/Source/RemoteControlLogic/Public/Behaviour/Builtin/Bind/RCBehaviourBind.h
+++ b/Engine/Plugins/VirtualProduction/RemoteControl/Source/RemoteControlLogic/Public/Behaviour/Builtin/Bind/RCBehaviourBind.h
@@ -35,6 +35,8 @@ public:
 	/** Whether we can create an action pertaining to a given remote control field for the current behaviour */
 	virtual bool CanHaveActionForField(const TSharedPtr<FRemoteControlField> InRemoteControlField) const override;
 
+	/** Whether the behaviour support property id */
+	virtual bool SupportPropertyId() const override { return false; }
 	//~ End URCBehaviour interface
 
 	/** Contains behaviour-independent validation logic for determining whether a Remote Controller property can be bound to a given Controller*/
diff --git a/Engine/Plugins/VirtualProduction/RemoteControl/Source/RemoteControlLogic/Public/Behaviour/Builtin/Conditional/RCBehaviourConditional.h b/Engine/Plugins/VirtualProduction/RemoteControl/Source/RemoteControlLogic/Public/Behaviour/Builtin/Conditional/RCBehaviourConditional.h
--- a/Engine/Plugins/VirtualProduction/RemoteControl/Source/RemoteControlLogic/Public/Behaviour/Builtin/Conditional/RCBehaviourConditional.h
+++ b/Engine/Plugins/VirtualProduction/RemoteControl/Source/RemoteControlLogic/Public/Behaviour/Builtin/Conditional/RCBehaviourConditional.h
@@ -83,6 +83,9 @@ public:
 	/** Whether we can create an action pertaining to a given remote control field for the current behaviour */
 	virtual bool CanHaveActionForField(const TSharedPtr<FRemoteControlField> InRemoteControlField) const override;
 
+	/** Whether the behaviour support property id */
+	virtual bool SupportPropertyId() const override { return true; }
+
 	/** User-friendly text representation of a condition enum. Used to display comparator info in the Actions table */
 	FText GetConditionTypeAsText(ERCBehaviourConditionType ConditionType) const;
 
diff --git a/Engine/Plugins/VirtualProduction/RemoteControl/Source/RemoteControlLogic/Public/Behaviour/Builtin/Path/RCSetAssetByPathBehaviour.h b/Engine/Plugins/VirtualProduction/RemoteControl/Source/RemoteControlLogic/Public/Behaviour/Builtin/Path/RCSetAssetByPathBehaviour.h
--- a/Engine/Plugins/VirtualProduction/RemoteControl/Source/RemoteControlLogic/Public/Behaviour/Builtin/Path/RCSetAssetByPathBehaviour.h
+++ b/Engine/Plugins/VirtualProduction/RemoteControl/Source/RemoteControlLogic/Public/Behaviour/Builtin/Path/RCSetAssetByPathBehaviour.h
@@ -109,6 +109,7 @@ public:
 
 	//~ Begin URCBehaviour interface
 	virtual void Initialize() override;
+	virtual bool SupportPropertyId() const override { return false; }
 	virtual void UpdateEntityIds(const TMap<FGuid, FGuid>& InEntityIdMap) override;
 	//~ End URCBehaviour interface
 
diff --git a/Engine/Plugins/VirtualProduction/RemoteControl/Source/RemoteControlLogic/Public/Behaviour/Builtin/RangeMap/RCRangeMapBehaviour.h b/Engine/Plugins/VirtualProduction/RemoteControl/Source/RemoteControlLogic/Public/Behaviour/Builtin/RangeMap/RCRangeMapBehaviour.h
--- a/Engine/Plugins/VirtualProduction/RemoteControl/Source/RemoteControlLogic/Public/Behaviour/Builtin/RangeMap/RCRangeMapBehaviour.h
+++ b/Engine/Plugins/VirtualProduction/RemoteControl/Source/RemoteControlLogic/Public/Behaviour/Builtin/RangeMap/RCRangeMapBehaviour.h
@@ -77,6 +77,9 @@ public:
 	/** Checks whether or not one of the following fields can be added to the AddAction Menu. Makes sure they're unique. */
 	virtual bool CanHaveActionForField(const TSharedPtr<FRemoteControlField> InRemoteControlField) const override;
 
+	/** Whether the behaviour support property id */
+	virtual bool SupportPropertyId() const override { return false; }
+
 	/** Returns the Step value associated with a given Action*/
 	bool GetValueForAction(const URCAction* InAction, double& OutValue);
 
diff --git a/Engine/Plugins/VirtualProduction/RemoteControl/Source/RemoteControlLogic/Public/Behaviour/RCBehaviour.h b/Engine/Plugins/VirtualProduction/RemoteControl/Source/RemoteControlLogic/Public/Behaviour/RCBehaviour.h
--- a/Engine/Plugins/VirtualProduction/RemoteControl/Source/RemoteControlLogic/Public/Behaviour/RCBehaviour.h
+++ b/Engine/Plugins/VirtualProduction/RemoteControl/Source/RemoteControlLogic/Public/Behaviour/RCBehaviour.h
@@ -60,6 +60,9 @@ public:
 	/** Whether we can create an action pertaining to a given remote control field for the current behaviour */
 	virtual bool CanHaveActionForField(const TSharedPtr<FRemoteControlField> InRemoteControlField) const;
 
+	/** Whether the behaviour support property id, by default true because of the OnModify Behaviour */
+	virtual bool SupportPropertyId() const { return true; }
+
 	/**
 	 * Return blueprint class associated with behaviour if exists
 	 */
diff --git a/Engine/Plugins/VirtualProduction/RemoteControl/Source/RemoteControlUI/Private/UI/Action/SRCActionPanel.cpp b/Engine/Plugins/VirtualProduction/RemoteControl/Source/RemoteControlUI/Private/UI/Action/SRCActionPanel.cpp
--- a/Engine/Plugins/VirtualProduction/RemoteControl/Source/RemoteControlUI/Private/UI/Action/SRCActionPanel.cpp
+++ b/Engine/Plugins/VirtualProduction/RemoteControl/Source/RemoteControlUI/Private/UI/Action/SRCActionPanel.cpp
@@ -518,14 +518,25 @@ FReply SRCActionPanel::OnAddAllSelectedFields()
 				{
 					if (const TSharedPtr<SRCPanelExposedEntitiesGroup> RCFieldGroup = StaticCastSharedPtr<SRCPanelExposedEntitiesGroup>(RCEntity))
 					{
-						TArray<TSharedPtr<SRCPanelTreeNode>> RCFieldGroupEntities;
-						RCFieldGroup->GetNodeChildren(RCFieldGroupEntities);
-
-						for (const TSharedPtr<SRCPanelTreeNode>& RCFieldGroupEntity : RCFieldGroupEntities)
+						if (RCFieldGroup->GetGroupType() == EFieldGroupType::PropertyId)
+						{
+							const URCBehaviour* Behaviour = BehaviourItem->GetBehaviour();
+							if (Behaviour && Behaviour->SupportPropertyId())
+							{
+								AddAction(RCFieldGroup->GetFieldKey());
+							}
+						}
+						else
 						{
-							if (RCFieldGroupEntity->GetRCId().IsValid())
+							TArray<TSharedPtr<SRCPanelTreeNode>> RCFieldGroupEntities;
+							RCFieldGroup->GetNodeChildren(RCFieldGroupEntities);
+
+							for (const TSharedPtr<SRCPanelTreeNode>& RCFieldGroupEntity : RCFieldGroupEntities)
 							{
-								AddSelectedActionLambda(RCFieldGroupEntity);
+								if (RCFieldGroupEntity->GetRCId().IsValid())
+								{
+									AddSelectedActionLambda(RCFieldGroupEntity);
+								}
 							}
 						}
 					}
diff --git a/Engine/Plugins/VirtualProduction/RemoteControl/Source/RemoteControlUI/Private/UI/Action/SRCActionPanelList.h b/Engine/Plugins/VirtualProduction/RemoteControl/Source/RemoteControlUI/Private/UI/Action/SRCActionPanelList.h
--- a/Engine/Plugins/VirtualProduction/RemoteControl/Source/RemoteControlUI/Private/UI/Action/SRCActionPanelList.h
+++ b/Engine/Plugins/VirtualProduction/RemoteControl/Source/RemoteControlUI/Private/UI/Action/SRCActionPanelList.h
@@ -18,6 +18,7 @@
 #include "UI/Behaviour/RCBehaviourModel.h"
 #include "UI/RCUIHelpers.h"
 #include "UI/RemoteControlPanelStyle.h"
+#include "UI/SRCPanelExposedEntitiesGroup.h"
 #include "UI/SRCPanelExposedField.h"
 #include "UI/SRCPanelFieldGroup.h"
 #include "UI/SRemoteControlPanel.h"
@@ -204,6 +205,16 @@ public:
 		return nullptr;
 	}
 
+	/** Adds a PropertyId action with the given PropertyId */
+	URCAction* AddAction(const FName InPropertyId)
+	{
+		if (const TSharedPtr<SRCActionPanel> ActionPanel = GetActionPanel())
+		{
+			return ActionPanel->AddAction(InPropertyId);
+		}
+		return nullptr;
+	}
+
 	virtual void AddNewLogicItem(UObject* InLogicItem) override
 	{
 		AddRowToList(Cast<URCAction>(InLogicItem));
@@ -348,17 +359,65 @@ private:
 		{
 			if (DragDropOperation->IsOfType<FExposedEntityDragDrop>())
 			{
-				if (TSharedPtr<FExposedEntityDragDrop> DragDropOp = StaticCastSharedPtr<FExposedEntityDragDrop>(DragDropOperation))
+				if (const TSharedPtr<FExposedEntityDragDrop> DragDropOp = StaticCastSharedPtr<FExposedEntityDragDrop>(DragDropOperation))
 				{
-					// Fetch the Exposed Entity
-					for (const FGuid& ExposedEntityId : DragDropOp->GetSelectedIds())
+					// Check this beforehand to avoid calling all this for each PropertyId group later on
+					bool bBehaviorSupportPropertyId = false;
+					if (const TSharedPtr<FRCBehaviourModel>& BehaviorModel = GetBehaviourItem())
 					{
-						if (TSharedPtr<SRCActionPanel> ActionPanel = GetActionPanel())
+						if (const URCBehaviour* Behavior = BehaviorModel->GetBehaviour())
+						{
+							bBehaviorSupportPropertyId = Behavior->SupportPropertyId();
+						}
+					}
+
+					// Fetch the Exposed Entities
+					for (const TSharedPtr<SRCPanelTreeNode>& ExposedEntity : DragDropOp->GetSelectedEntities())
+					{
+						const SRCPanelTreeNode::ENodeType NodeType = ExposedEntity->GetRCType();
+						if (NodeType == SRCPanelTreeNode::FieldGroup)
+						{
+							if (const TSharedPtr<SRCPanelExposedEntitiesGroup>& FieldGroup = StaticCastSharedPtr<SRCPanelExposedEntitiesGroup>(ExposedEntity))
+							{
+								const EFieldGroupType GroupType = FieldGroup->GetGroupType();
+								// Property Id group
+								if (GroupType == EFieldGroupType::PropertyId)
+								{
+									if (bBehaviorSupportPropertyId)
+									{
+										AddAction(FieldGroup->GetFieldKey());
+									}
+								}
+								// Owner group
+								else if (GroupType == EFieldGroupType::Owner)
+								{
+									TArray<TSharedPtr<SRCPanelTreeNode>> OwnerGroupChildren;
+									FieldGroup->GetNodeChildren(OwnerGroupChildren);
+									for (const TSharedPtr<SRCPanelTreeNode>& Child : OwnerGroupChildren)
+									{
+										if (Child->GetRCType() == SRCPanelTreeNode::Field)
+										{
+											if (const TSharedPtr<SRCActionPanel> ActionPanel = GetActionPanel())
+											{
+												if (ActionPanel->CanHaveActionForField(Child->GetRCId()))
+												{
+													AddAction(Child->GetRCId());
+												}
+											}
+										}
+									}
+								}
+							}
+						}
+						// Fields
+						else if (NodeType == SRCPanelTreeNode::Field)
 						{
-							// Add Action
-							if (ActionPanel->CanHaveActionForField(ExposedEntityId))
+							if (const TSharedPtr<SRCActionPanel> ActionPanel = GetActionPanel())
 							{
-								AddAction(ExposedEntityId);
+								if (ActionPanel->CanHaveActionForField(ExposedEntity->GetRCId()))
+								{
+									AddAction(ExposedEntity->GetRCId());
+								}
 							}
 						}
 					}
@@ -379,9 +438,16 @@ private:
 							const TArray<FGuid> GroupFields = Group->GetFields();
 
 							// Add Action for all fields in the Group
-							for (const FGuid RemoteControlFieldId : GroupFields)
+							if (const TSharedPtr<SRCActionPanel> ActionPanel = GetActionPanel())
 							{
-								AddAction(RemoteControlFieldId);
+								for (const FGuid RemoteControlFieldId : GroupFields)
+								{
+									// Add Action
+									if (ActionPanel->CanHaveActionForField(RemoteControlFieldId))
+									{
+										AddAction(RemoteControlFieldId);
+									}
+								}
 							}
 						}
 					}
@@ -399,17 +465,67 @@ private:
 		{
 			if (DragDropOperation->IsOfType<FExposedEntityDragDrop>())
 			{
-				if (TSharedPtr<FExposedEntityDragDrop> DragDropOp = StaticCastSharedPtr<FExposedEntityDragDrop>(DragDropOperation))
+				if (const TSharedPtr<FExposedEntityDragDrop> DragDropOp = StaticCastSharedPtr<FExposedEntityDragDrop>(DragDropOperation))
 				{
+					// Check this beforehand to avoid calling all this for each PropertyId group later on
+					bool bBehaviorSupportPropertyId = false;
+					if (const TSharedPtr<FRCBehaviourModel>& BehaviorModel = GetBehaviourItem())
+					{
+						if (const URCBehaviour* Behavior = BehaviorModel->GetBehaviour())
+						{
+							bBehaviorSupportPropertyId = Behavior->SupportPropertyId();
+						}
+					}
+
 					// Fetch the Exposed Entity
-					for (const FGuid& ExposedEntityId : DragDropOp->GetSelectedIds())
+					for (const TSharedPtr<SRCPanelTreeNode>& ExposedEntity : DragDropOp->GetSelectedEntities())
 					{
-						if (TSharedPtr<SRCActionPanel> ActionPanel = GetActionPanel())
+						const SRCPanelTreeNode::ENodeType NodeType = ExposedEntity->GetRCType();
+						if (NodeType == SRCPanelTreeNode::FieldGroup)
 						{
-							// Add Action
-							if (ActionPanel->CanHaveActionForField(ExposedEntityId))
+							if (const TSharedPtr<SRCPanelExposedEntitiesGroup>& FieldGroup = StaticCastSharedPtr<SRCPanelExposedEntitiesGroup>(ExposedEntity))
 							{
-								return true;
+								const EFieldGroupType GroupType = FieldGroup->GetGroupType();
+								// Property Id Group
+								if (GroupType == EFieldGroupType::PropertyId)
+								{
+									if (bBehaviorSupportPropertyId)
+									{
+										return true;
+									}
+								}
+								// Owner Group
+								else if (GroupType == EFieldGroupType::Owner)
+								{
+									TArray<TSharedPtr<SRCPanelTreeNode>> OwnerGroupChildren;
+									FieldGroup->GetNodeChildren(OwnerGroupChildren);
+									for (const TSharedPtr<SRCPanelTreeNode>& Child : OwnerGroupChildren)
+									{
+										if (Child->GetRCType() == SRCPanelTreeNode::Field)
+										{
+											if (const TSharedPtr<SRCActionPanel> ActionPanel = GetActionPanel())
+											{
+												// if at least 1 can be created then enable it
+												if (ActionPanel->CanHaveActionForField(Child->GetRCId()))
+												{
+													return true;
+												}
+											}
+										}
+									}
+								}
+							}
+						}
+						// Fields
+						else if (NodeType == SRCPanelTreeNode::Field)
+						{
+							if (const TSharedPtr<SRCActionPanel> ActionPanel = GetActionPanel())
+							{
+								// if at least 1 can be created then enable it
+								if (ActionPanel->CanHaveActionForField(ExposedEntity->GetRCId()))
+								{
+									return true;
+								}
 							}
 						}
 					}
diff --git a/Engine/Plugins/VirtualProduction/RemoteControl/Source/RemoteControlUI/Private/UI/Behaviour/RCBehaviourModel.cpp b/Engine/Plugins/VirtualProduction/RemoteControl/Source/RemoteControlUI/Private/UI/Behaviour/RCBehaviourModel.cpp
--- a/Engine/Plugins/VirtualProduction/RemoteControl/Source/RemoteControlUI/Private/UI/Behaviour/RCBehaviourModel.cpp
+++ b/Engine/Plugins/VirtualProduction/RemoteControl/Source/RemoteControlUI/Private/UI/Behaviour/RCBehaviourModel.cpp
@@ -148,6 +148,15 @@ TSharedPtr<SRCLogicPanelListBase> FRCBehaviourModel::GetActionsListWidget(TShare
 	return SNew(SRCActionPanelList<FRCActionModel>, InActionPanel, SharedThis(this));
 }
 
+bool FRCBehaviourModel::SupportPropertyId() const
+{
+	if (const URCBehaviour* Behavior = BehaviourWeakPtr.Get())
+	{
+		return Behavior->SupportPropertyId();
+	}
+	return false;
+}
+
 URCBehaviour* FRCBehaviourModel::GetBehaviour() const
 {
 	return BehaviourWeakPtr.Get();
diff --git a/Engine/Plugins/VirtualProduction/RemoteControl/Source/RemoteControlUI/Private/UI/Behaviour/RCBehaviourModel.h b/Engine/Plugins/VirtualProduction/RemoteControl/Source/RemoteControlUI/Private/UI/Behaviour/RCBehaviourModel.h
--- a/Engine/Plugins/VirtualProduction/RemoteControl/Source/RemoteControlUI/Private/UI/Behaviour/RCBehaviourModel.h
+++ b/Engine/Plugins/VirtualProduction/RemoteControl/Source/RemoteControlUI/Private/UI/Behaviour/RCBehaviourModel.h
@@ -64,6 +64,9 @@ public:
 	*/
 	virtual TSharedPtr<SRCLogicPanelListBase> GetActionsListWidget(TSharedRef<SRCActionPanel> InActionPanel);
 
+	/** Whether or not the behaviour support PropertyId action */
+	bool SupportPropertyId() const;
+
 protected:
 	/** Invoked after an action has been added for this Behaviour in the actions panel */
 	virtual void OnActionAdded(URCAction* Action) {}
diff --git a/Engine/Plugins/VirtualProduction/RemoteControl/Source/RemoteControlUI/Private/UI/Controller/SRCControllerPanelList.cpp b/Engine/Plugins/VirtualProduction/RemoteControl/Source/RemoteControlUI/Private/UI/Controller/SRCControllerPanelList.cpp
--- a/Engine/Plugins/VirtualProduction/RemoteControl/Source/RemoteControlUI/Private/UI/Controller/SRCControllerPanelList.cpp
+++ b/Engine/Plugins/VirtualProduction/RemoteControl/Source/RemoteControlUI/Private/UI/Controller/SRCControllerPanelList.cpp
@@ -241,7 +241,7 @@ namespace UE::RCControllerPanelList
 				if (TSharedPtr<FExposedEntityDragDrop> DragDropOp = StaticCastSharedPtr<FExposedEntityDragDrop>(DragDropOperation))
 				{
 					// Fetch the Exposed Entity
-					const TArray<FGuid>& ExposedEntitiesIds = DragDropOp->GetSelectedIds();
+					const TArray<FGuid>& ExposedEntitiesIds = DragDropOp->GetSelectedFieldsId();
 
 					if (ExposedEntitiesIds.Num() == 1)
 					{
@@ -293,7 +293,7 @@ namespace UE::RCControllerPanelList
 						if (TSharedPtr<FExposedEntityDragDrop> DragDropOp = StaticCastSharedPtr<FExposedEntityDragDrop>(DragDropOperation))
 						{
 							// Fetch the Exposed Entity
-							const TArray<FGuid>& ExposedEntitiesIds = DragDropOp->GetSelectedIds();
+							const TArray<FGuid>& ExposedEntitiesIds = DragDropOp->GetSelectedFieldsId();
 
 							if (ExposedEntitiesIds.Num() == 1)
 							{
@@ -886,7 +886,7 @@ bool SRCControllerPanelList::OnAllowDrop(TSharedPtr<FDragDropOperation> DragDrop
 	if (TSharedPtr<FExposedEntityDragDrop> DragDropOp = GetExposedEntityDragDrop(DragDropOperation))
 	{
 		// Fetch the Exposed Entity
-		const TArray<FGuid>& ExposedEntitiesIds = DragDropOp->GetSelectedIds();
+		const TArray<FGuid>& ExposedEntitiesIds = DragDropOp->GetSelectedFieldsId();
 
 		// Check if Entity is supported by controllers and currently only 1 dragged entity dragged is supported
 		return ExposedEntitiesIds.Num() == 1 && IsEntitySupported(ExposedEntitiesIds[0]);
diff --git a/Engine/Plugins/VirtualProduction/RemoteControl/Source/RemoteControlUI/Private/UI/SRCPanelExposedEntitiesGroup.cpp b/Engine/Plugins/VirtualProduction/RemoteControl/Source/RemoteControlUI/Private/UI/SRCPanelExposedEntitiesGroup.cpp
--- a/Engine/Plugins/VirtualProduction/RemoteControl/Source/RemoteControlUI/Private/UI/SRCPanelExposedEntitiesGroup.cpp
+++ b/Engine/Plugins/VirtualProduction/RemoteControl/Source/RemoteControlUI/Private/UI/SRCPanelExposedEntitiesGroup.cpp
@@ -36,11 +36,16 @@ SRCPanelTreeNode::FMakeNodeWidgetArgs SRCPanelExposedEntitiesGroup::CreateNodeWi
 			.OnTextCommitted(this, &SRCPanelExposedEntitiesGroup::OnPropertyIdTextCommitted)
 		];
 
+		Args.NameWidget = SNew(SBox).HeightOverride(25).VAlign(VAlign_Center)
+			[
+				SNew(STextBlock)
+				.Text(FText::FromString("Property Id Group"))
+			];
+
 		Args.OwnerNameWidget = SNullWidget::NullWidget;
 	}
 	else if (GroupType == EFieldGroupType::Owner)
 	{
-		Args.PropertyIdWidget = SNullWidget::NullWidget;
 		Args.OwnerNameWidget = SNew(SBox)
 			.HeightOverride(25)
 			.VAlign(VAlign_Center)
@@ -48,6 +53,14 @@ SRCPanelTreeNode::FMakeNodeWidgetArgs SRCPanelExposedEntitiesGroup::CreateNodeWi
 				SNew(STextBlock)
 				.Text(FText::FromName(FieldKey))
 			];
+
+		Args.NameWidget = SNew(SBox).HeightOverride(25).VAlign(VAlign_Center)
+			[
+				SNew(STextBlock)
+				.Text(FText::FromString("Owner Group"))
+			];
+
+		Args.PropertyIdWidget = SNullWidget::NullWidget;
 	}
 	else
 	{
@@ -55,11 +68,6 @@ SRCPanelTreeNode::FMakeNodeWidgetArgs SRCPanelExposedEntitiesGroup::CreateNodeWi
 		Args.OwnerNameWidget = SNullWidget::NullWidget;
 	}
 
-	Args.NameWidget = SNew(SBox).HeightOverride(25).VAlign(VAlign_Center)
-		[
-			SNew(STextBlock)
-			.Text(FText::FromString("Group"))
-		];
 	Args.SubObjectPathWidget = SNullWidget::NullWidget;
 	Args.ValueWidget = SNullWidget::NullWidget;
 	Args.ResetButton = SNullWidget::NullWidget;
diff --git a/Engine/Plugins/VirtualProduction/RemoteControl/Source/RemoteControlUI/Private/UI/SRCPanelExposedEntitiesList.cpp b/Engine/Plugins/VirtualProduction/RemoteControl/Source/RemoteControlUI/Private/UI/SRCPanelExposedEntitiesList.cpp
--- a/Engine/Plugins/VirtualProduction/RemoteControl/Source/RemoteControlUI/Private/UI/SRCPanelExposedEntitiesList.cpp
+++ b/Engine/Plugins/VirtualProduction/RemoteControl/Source/RemoteControlUI/Private/UI/SRCPanelExposedEntitiesList.cpp
@@ -824,16 +824,27 @@ void SRCPanelExposedEntitiesList::OnLabelModified(const FName InOldName, const F
 
 FReply SRCPanelExposedEntitiesList::OnNodeDragDetected(const FGeometry& InGeometry, const FPointerEvent& InPointerEvent, TSharedPtr<SRCPanelTreeNode> InNode)
 {
-	if (InNode && InNode->GetRCType() == SRCPanelTreeNode::Field)
+	if (InNode.IsValid())
 	{
-		TArray<FGuid> SelectedIds;
-		Algo::TransformIf(GetSelectedEntities(), SelectedIds
-			, [] (const TSharedPtr<SRCPanelTreeNode>& TreeNode) { return TreeNode->GetRCType() == SRCPanelTreeNode::Field; }
-			, [] (const TSharedPtr<SRCPanelTreeNode>& TreeNode){ return TreeNode->GetRCId(); });
+		const SRCPanelTreeNode::ENodeType FirstNodeType = InNode->GetRCType();
+		if (FirstNodeType == SRCPanelTreeNode::Field || FirstNodeType == SRCPanelTreeNode::FieldGroup)
+		{
+			TArray<TSharedPtr<SRCPanelTreeNode>> SelectedEntities;
+			Algo::TransformIf(GetSelectedEntities(), SelectedEntities
+				, [] (const TSharedPtr<SRCPanelTreeNode>& TreeNode)
+					{
+						const SRCPanelTreeNode::ENodeType NodeType = TreeNode->GetRCType();
+						return NodeType == SRCPanelTreeNode::Field || NodeType == SRCPanelTreeNode::FieldGroup;
+					}
+				, [] (const TSharedPtr<SRCPanelTreeNode>& TreeNode)
+					{
+						return TreeNode;
+					});
 
-		const TSharedRef<FExposedEntityDragDrop> DragDropOp = MakeShared<FExposedEntityDragDrop>(InNode->GetDragAndDropWidget(SelectedIds.Num()), InNode->GetRCId(), SelectedIds);
-		DragDropOp->Construct();
-		return FReply::Handled().BeginDragDrop(DragDropOp);
+			const TSharedRef<FExposedEntityDragDrop> DragDropOp = MakeShared<FExposedEntityDragDrop>(InNode->GetDragAndDropWidget(SelectedEntities.Num()), InNode->GetRCId(), SelectedEntities);
+			DragDropOp->Construct();
+			return FReply::Handled().BeginDragDrop(DragDropOp);
+		}
 	}
 	return FReply::Unhandled();
 }
@@ -1163,7 +1174,7 @@ FReply SRCPanelExposedEntitiesList::OnDropOnGroup(const TSharedPtr<FDragDropOper
 				FRemoteControlPresetLayout::FFieldSwapArgs Args;
 				Args.OriginGroupId = GetSelectedGroup()->GetRCId();
 				Args.TargetGroupId = DragTargetGroup->GetRCId();
-				Args.DraggedFieldsIds = DragDropOp->GetSelectedIds();
+				Args.DraggedFieldsIds = DragDropOp->GetSelectedFieldsId();
 
 				if (TargetEntity)
 				{
@@ -1194,7 +1205,7 @@ FReply SRCPanelExposedEntitiesList::OnDropOnGroup(const TSharedPtr<FDragDropOper
 			FRemoteControlPresetLayout::FFieldSwapArgs Args;
 			Args.OriginGroupId = DragOriginGroupId;
 			Args.TargetGroupId = DragTargetGroup->GetRCId();
-			Args.DraggedFieldsIds = DragDropOp->GetSelectedIds();
+			Args.DraggedFieldsIds = DragDropOp->GetSelectedFieldsId();
 
 			if (TargetEntity)
 			{
diff --git a/Engine/Plugins/VirtualProduction/RemoteControl/Source/RemoteControlUI/Private/UI/SRCPanelExposedEntity.h b/Engine/Plugins/VirtualProduction/RemoteControl/Source/RemoteControlUI/Private/UI/SRCPanelExposedEntity.h
--- a/Engine/Plugins/VirtualProduction/RemoteControl/Source/RemoteControlUI/Private/UI/SRCPanelExposedEntity.h
+++ b/Engine/Plugins/VirtualProduction/RemoteControl/Source/RemoteControlUI/Private/UI/SRCPanelExposedEntity.h
@@ -129,10 +129,20 @@ public:
 
 	using WidgetType = SWidget;
 
-	FExposedEntityDragDrop(const TSharedPtr<SWidget>& InWidget, const FGuid& InNodeId, const TArray<FGuid>& InSelectedIds)
+	FExposedEntityDragDrop(const TSharedPtr<SWidget>& InWidget, const FGuid& InNodeId, const TArray<TSharedPtr<SRCPanelTreeNode>>& InSelectedEntities)
 		: NodeId(InNodeId)
-		, SelectedIds(InSelectedIds)
+		, SelectedEntities(InSelectedEntities)
 	{
+		Algo::TransformIf(SelectedEntities, SelectedFieldsIds,
+			[] (const TSharedPtr<SRCPanelTreeNode>& InNode)
+			{
+				return InNode->GetRCType() == SRCPanelTreeNode::Field;
+			},
+			[] (const TSharedPtr<SRCPanelTreeNode>& InNode)
+			{
+				return InNode->GetRCId();
+			});
+		
 		DecoratorWidget = SNew(SBorder)
 			.Padding(1.0f)
 			.BorderImage(FRemoteControlPanelStyle::Get()->GetBrush("RemoteControlPanel.ExposedFieldBorder"))
@@ -149,9 +159,15 @@ public:
 	}
 
 	/** Get the IDs that were selected at the time the drag started. */
-	const TArray<FGuid>& GetSelectedIds() const
+	const TArray<TSharedPtr<SRCPanelTreeNode>>& GetSelectedEntities() const
 	{
-		return SelectedIds;
+		return SelectedEntities;
+	}
+
+	/** Get the IDs that were selected at the time the drag started. */
+	const TArray<FGuid>& GetSelectedFieldsId() const
+	{
+		return SelectedFieldsIds;
 	}
 
 	virtual void OnDrop(bool bDropWasHandled, const FPointerEvent& MouseEvent) override
@@ -167,8 +183,10 @@ public:
 private:
 	/** ID of the represented entity or group. */
 	FGuid NodeId;
-	/** IDs that were selected at the time the drag started. */
-	TArray<FGuid> SelectedIds;
+	/** Field IDs that were selected at the time the drag started. */
+    TArray<FGuid> SelectedFieldsIds;
+	/** Entities that were selected at the time the drag started. */
+	TArray<TSharedPtr<SRCPanelTreeNode>> SelectedEntities;
 	/** Decorator Drag and Drop widget. */
 	TSharedPtr<SWidget> DecoratorWidget;
 };
diff --git a/Engine/Plugins/VirtualProduction/RemoteControl/Source/RemoteControlUI/Private/UI/SRCPanelFieldGroup.cpp b/Engine/Plugins/VirtualProduction/RemoteControl/Source/RemoteControlUI/Private/UI/SRCPanelFieldGroup.cpp
--- a/Engine/Plugins/VirtualProduction/RemoteControl/Source/RemoteControlUI/Private/UI/SRCPanelFieldGroup.cpp
+++ b/Engine/Plugins/VirtualProduction/RemoteControl/Source/RemoteControlUI/Private/UI/SRCPanelFieldGroup.cpp
@@ -209,7 +209,7 @@ bool SRCPanelGroup::OnAllowDropFromOtherGroup(TSharedPtr<FDragDropOperation> Dra
 		{
 			if (OnGetGroupId.IsBound())
 			{
-				for (FGuid SelectedId : DragDropOp->GetSelectedIds())
+				for (FGuid SelectedId : DragDropOp->GetSelectedFieldsId())
 				{
 					const FGuid OriginGroupId = OnGetGroupId.Execute(SelectedId);
 					if (OriginGroupId != Id)
