Remote Control - Update list of disallowed functions

diff --git a/Engine/Plugins/VirtualProduction/RemoteControl/Source/RemoteControl/Private/RemoteControlModule.cpp b/Engine/Plugins/VirtualProduction/RemoteControl/Source/RemoteControl/Private/RemoteControlModule.cpp
--- a/Engine/Plugins/VirtualProduction/RemoteControl/Source/RemoteControl/Private/RemoteControlModule.cpp
+++ b/Engine/Plugins/VirtualProduction/RemoteControl/Source/RemoteControl/Private/RemoteControlModule.cpp
@@ -19,6 +19,7 @@
 #include "IRemoteControlModule.h"
 #include "IStructDeserializerBackend.h"
 #include "IStructSerializerBackend.h"
+#include "Kismet/KismetSystemLibrary.h"
 #include "Misc/ScopeExit.h"
 #include "RCPropertyUtilities.h"
 #include "RCVirtualProperty.h"
@@ -798,6 +799,8 @@ void FRemoteControlModule::StartupModule()
 
 	// Register PropertyIdHandler
 	RegisterPropertyIdHandler();
+
+	PopulateDisallowedFunctions();
 }
 
 void FRemoteControlModule::ShutdownModule()
@@ -1104,7 +1107,8 @@ bool FRemoteControlModule::ResolveCall(const FString& ObjectPath, const FString&
 					ErrorText = FString::Printf(TEXT("Function: %s does not exist on object: %s"), *FunctionName, *ObjectPath);
 					bSuccess = false;
 				}
-				else if ((!Function->HasAllFunctionFlags(FUNC_BlueprintCallable | FUNC_Public) && !Function->HasAllFunctionFlags(FUNC_BlueprintEvent))
+				else if (!IsFunctionAllowed(Function)
+						|| (!Function->HasAllFunctionFlags(FUNC_BlueprintCallable | FUNC_Public) && !Function->HasAllFunctionFlags(FUNC_BlueprintEvent))
 	#if WITH_EDITOR
 						|| Function->HasMetaData(RemoteControlUtil::NAME_DeprecatedFunction)
 						|| Function->HasMetaData(RemoteControlUtil::NAME_ScriptNoExport)
@@ -2715,6 +2719,11 @@ void FRemoteControlModule::RegisterPropertyIdHandler()
 	RegisterPropertyIdPropertyHandler<FObjectPropertyIdHandler>();
 }
 
+void FRemoteControlModule::PopulateDisallowedFunctions()
+{
+	FunctionDisallowList.Add(GetDefault<UKismetSystemLibrary>()->FindFunction(GET_FUNCTION_NAME_CHECKED(UKismetSystemLibrary, LaunchURL)));
+}
+
 bool FRemoteControlModule::CanInterceptFunction(const FRCCall& RCCall) const
 {
 	if (RCCall.IsValid())
@@ -2866,6 +2875,11 @@ bool FRemoteControlModule::ModifyArrayProperty(const FRCObjectReference& ObjectA
 	return false;
 }
 
+bool FRemoteControlModule::IsFunctionAllowed(UFunction* Function)
+{
+	return !FunctionDisallowList.Contains(Function);
+}
+
 #if WITH_EDITOR
 void FRemoteControlModule::EndOngoingModificationIfMismatched(uint32 TypeHash)
 {
diff --git a/Engine/Plugins/VirtualProduction/RemoteControl/Source/RemoteControl/Private/RemoteControlModule.h b/Engine/Plugins/VirtualProduction/RemoteControl/Source/RemoteControl/Private/RemoteControlModule.h
--- a/Engine/Plugins/VirtualProduction/RemoteControl/Source/RemoteControl/Private/RemoteControlModule.h
+++ b/Engine/Plugins/VirtualProduction/RemoteControl/Source/RemoteControl/Private/RemoteControlModule.h
@@ -132,6 +132,11 @@ private:
  	 */
 	void RegisterPropertyIdHandler();
 
+	/**
+	 * Populate the list of functions that cannot be called remotely.
+	 */
+	void PopulateDisallowedFunctions();
+
 	/**
 	 * Whether function can be intercepted by a remote control interceptor.
 	 */
@@ -158,6 +163,9 @@ private:
 	 */
 	bool ModifyArrayProperty(const FRCObjectReference& ObjectAccess, TFunctionRef<bool(FScriptArrayHelper&)> ModifyFunction);
 
+	/** Returns whether the function is allowed to be called remotely. */
+	bool IsFunctionAllowed(UFunction* Function);
+
 #if WITH_EDITOR
 
 	/**
@@ -286,6 +294,9 @@ private:
 	TMap<FName, TSharedPtr<IRCDefaultValueFactory>> DefaultValueFactories;
 
 	TSet<TSharedPtr<IPropertyIdHandler>> PropertyIdPropertyHandlers; 
+
+	/** List of functions that can't be called remotely. */
+	TSet<TWeakObjectPtr<UFunction>> FunctionDisallowList;
 };
 
 PRAGMA_ENABLE_DEPRECATION_WARNINGS
