[RemoteControl] Add right click option to change Propid

diff --git a/Engine/Plugins/VirtualProduction/RemoteControl/Source/RemoteControlUI/Private/Commands/RemoteControlCommands.cpp b/Engine/Plugins/VirtualProduction/RemoteControl/Source/RemoteControlUI/Private/Commands/RemoteControlCommands.cpp
--- a/Engine/Plugins/VirtualProduction/RemoteControl/Source/RemoteControlUI/Private/Commands/RemoteControlCommands.cpp
+++ b/Engine/Plugins/VirtualProduction/RemoteControl/Source/RemoteControlUI/Private/Commands/RemoteControlCommands.cpp
@@ -32,10 +32,13 @@ void FRemoteControlCommands::RegisterCommands()
 	UI_COMMAND(ToggleLogicEditor, "Logic", "View the logic applied to active selection.", EUserInterfaceActionType::ToggleButton, FInputChord());
 
 	// Delete Entity
-	UI_COMMAND(DeleteEntity, "Delete", "Delete the selected  group/exposed entity from the list.", EUserInterfaceActionType::Button, FInputChord(EKeys::Delete));
+	UI_COMMAND(DeleteEntity, "Delete", "Delete the selected group/exposed entity from the list.", EUserInterfaceActionType::Button, FInputChord(EKeys::Delete));
 
 	// Rename Entity
-	UI_COMMAND(RenameEntity, "Rename", "Rename the selected  group/exposed entity.", EUserInterfaceActionType::Button, FInputChord(EKeys::F2));
+	UI_COMMAND(RenameEntity, "Rename", "Rename the selected group/exposed entity.", EUserInterfaceActionType::Button, FInputChord(EKeys::F2));
+
+	// Modify Entity Prop Id
+	UI_COMMAND(ChangePropId, "ChangePropid", "Change the selected property Ids", EUserInterfaceActionType::Button, FInputChord(EKeys::F3, EModifierKey::Control));
 
 	// Copy Item
 	UI_COMMAND(CopyItem, "Copy", "Copy the selected UI item", EUserInterfaceActionType::Button, FInputChord(EKeys::C, EModifierKey::Control));
diff --git a/Engine/Plugins/VirtualProduction/RemoteControl/Source/RemoteControlUI/Private/Commands/RemoteControlCommands.h b/Engine/Plugins/VirtualProduction/RemoteControl/Source/RemoteControlUI/Private/Commands/RemoteControlCommands.h
--- a/Engine/Plugins/VirtualProduction/RemoteControl/Source/RemoteControlUI/Private/Commands/RemoteControlCommands.h
+++ b/Engine/Plugins/VirtualProduction/RemoteControl/Source/RemoteControlUI/Private/Commands/RemoteControlCommands.h
@@ -47,10 +47,15 @@ public:
 	TSharedPtr<FUICommandInfo> DeleteEntity;
 
 	/**
-	 * Holds the information about UI Command that  renames selected group/exposed entity.
+	 * Holds the information about UI Command that renames selected group/exposed entity.
 	 */
 	TSharedPtr<FUICommandInfo> RenameEntity;
 
+	/**
+	 * Holds the information about UI Command that change the selected property Ids.
+	 */
+	TSharedPtr<FUICommandInfo> ChangePropId;
+
 	/**
 	 * UI Command for copying a UI item in the Remote Control preset. Currently used for Logic panel
 	 */
diff --git a/Engine/Plugins/VirtualProduction/RemoteControl/Source/RemoteControlUI/Private/UI/SRCPanelExposedEntity.cpp b/Engine/Plugins/VirtualProduction/RemoteControl/Source/RemoteControlUI/Private/UI/SRCPanelExposedEntity.cpp
--- a/Engine/Plugins/VirtualProduction/RemoteControl/Source/RemoteControlUI/Private/UI/SRCPanelExposedEntity.cpp
+++ b/Engine/Plugins/VirtualProduction/RemoteControl/Source/RemoteControlUI/Private/UI/SRCPanelExposedEntity.cpp
@@ -89,6 +89,7 @@ TSharedPtr<SWidget> SRCPanelExposedEntity::GetContextMenu()
 	MenuBuilder.BeginSection("Common");
 
 	MenuBuilder.AddMenuEntry(FRemoteControlCommands::Get().RenameEntity, NAME_None, TAttribute<FText>(), TAttribute<FText>(), FSlateIcon(FAppStyle::GetAppStyleSetName(), TEXT("GenericCommands.Rename")));
+	MenuBuilder.AddMenuEntry(FRemoteControlCommands::Get().ChangePropId, NAME_None, TAttribute<FText>(), TAttribute<FText>(), FSlateIcon(FAppStyle::GetAppStyleSetName(), TEXT("Icons.Edit")));
 	MenuBuilder.AddMenuEntry(FRemoteControlCommands::Get().DeleteEntity, NAME_None, TAttribute<FText>(), TAttribute<FText>(), FSlateIcon(FAppStyle::GetAppStyleSetName(), TEXT("GenericCommands.Delete")));
 
 	MenuBuilder.EndSection();
@@ -557,6 +558,8 @@ SRCPanelTreeNode::FMakeNodeWidgetArgs SRCPanelExposedEntity::CreateEntityWidgetI
 			.HighlightText_Lambda([this]() { return HighlightText.Get().ToString().Len() > 3 ? HighlightText.Get() : FText::GetEmpty(); })
 		];
 
+	Args.PropertyIdWidget = SNullWidget::NullWidget;
+
 	Args.ValueWidget = ValueWidget;
 
 	Args.ResetButton = ResetWidget;
diff --git a/Engine/Plugins/VirtualProduction/RemoteControl/Source/RemoteControlUI/Private/UI/SRCPanelExposedField.cpp b/Engine/Plugins/VirtualProduction/RemoteControl/Source/RemoteControlUI/Private/UI/SRCPanelExposedField.cpp
--- a/Engine/Plugins/VirtualProduction/RemoteControl/Source/RemoteControlUI/Private/UI/SRCPanelExposedField.cpp
+++ b/Engine/Plugins/VirtualProduction/RemoteControl/Source/RemoteControlUI/Private/UI/SRCPanelExposedField.cpp
@@ -7,6 +7,7 @@
 #include "Components/SceneComponent.h"
 #include "EditorFontGlyphs.h"
 #include "Factories/IRCDefaultValueFactory.h"
+#include "Framework/Application/SlateApplication.h"
 #include "IDetailTreeNode.h"
 #include "IRCProtocolBindingList.h"
 #include "IRemoteControlModule.h"
@@ -169,6 +170,14 @@ SRCPanelTreeNode::ENodeType SRCPanelExposedField::GetRCType() const
 	return SRCPanelTreeNode::Field;
 }
 
+void SRCPanelExposedField::FocusPropertyIdWidget() const
+{
+	if (GetFieldType() == EExposedFieldType::Property && PropertyIdWidget.IsValid())
+	{
+		FSlateApplication::Get().SetUserFocus(0, PropertyIdWidget);
+	}
+}
+
 bool SRCPanelExposedField::HasChildren() const
 {
 	return ChildWidgets.Num() > 0;
@@ -188,16 +197,20 @@ SRCPanelTreeNode::FMakeNodeWidgetArgs SRCPanelExposedField::CreateEntityWidgetIn
 {
 	FMakeNodeWidgetArgs Args = SRCPanelExposedEntity::CreateEntityWidgetInternal(ValueWidget, ResetWidget, OptionalWarningMessage, EditConditionWidget);
 
-	Args.PropertyIdWidget = SNew(SBox)
-		[
-			SNew(SEditableTextBox)
-			.MinDesiredWidth(50.f)
-			.SelectAllTextWhenFocused(true)
-			.RevertTextOnEscape(true)
-			.ClearKeyboardFocusOnCommit(true)
-			.Text_Lambda([this] () { return GetPropertyIdText(); })
-			.OnTextCommitted(this, &SRCPanelExposedField::OnPropertyIdTextCommitted)
-		];
+	// TODO: Add support for function PropertyId
+	if (GetFieldType() == EExposedFieldType::Property)
+	{
+		Args.PropertyIdWidget = SNew(SBox)
+			[
+				SAssignNew(PropertyIdWidget, SEditableTextBox)
+				.MinDesiredWidth(50.f)
+				.SelectAllTextWhenFocused(true)
+				.RevertTextOnEscape(true)
+				.ClearKeyboardFocusOnCommit(true)
+				.Text_Lambda([this] () { return GetPropertyIdText(); })
+				.OnTextCommitted(this, &SRCPanelExposedField::OnPropertyIdTextCommitted)
+			];
+	}
 
 	return Args;
 }
diff --git a/Engine/Plugins/VirtualProduction/RemoteControl/Source/RemoteControlUI/Private/UI/SRCPanelExposedField.h b/Engine/Plugins/VirtualProduction/RemoteControl/Source/RemoteControlUI/Private/UI/SRCPanelExposedField.h
--- a/Engine/Plugins/VirtualProduction/RemoteControl/Source/RemoteControlUI/Private/UI/SRCPanelExposedField.h
+++ b/Engine/Plugins/VirtualProduction/RemoteControl/Source/RemoteControlUI/Private/UI/SRCPanelExposedField.h
@@ -45,6 +45,7 @@ struct SRCPanelExposedField : public SRCPanelExposedEntity
 	//~ SRCPanelTreeNode Interface 
 	virtual void GetNodeChildren(TArray<TSharedPtr<SRCPanelTreeNode>>& OutChildren) const override;
 	virtual ENodeType GetRCType() const override;
+	virtual void FocusPropertyIdWidget() const override;
 	virtual bool HasChildren() const override;
 	virtual void Refresh() override;
 	virtual TSharedRef<SWidget> GetWidget(const FName ForColumnName, const FName InActiveProtocol) override;
@@ -118,6 +119,8 @@ private:
 	TUniquePtr<uint8[]> DefaultValue;
 	/** Holds the shared reference of reset button for this field. */
 	TSharedPtr<SWidget> ResetButtonWidget;
+	/** Holds the SharedPtr of the PropertyId widget of this field. */
+	TSharedPtr<SWidget> PropertyIdWidget;
 };
 
 
diff --git a/Engine/Plugins/VirtualProduction/RemoteControl/Source/RemoteControlUI/Private/UI/SRCPanelTreeNode.h b/Engine/Plugins/VirtualProduction/RemoteControl/Source/RemoteControlUI/Private/UI/SRCPanelTreeNode.h
--- a/Engine/Plugins/VirtualProduction/RemoteControl/Source/RemoteControlUI/Private/UI/SRCPanelTreeNode.h
+++ b/Engine/Plugins/VirtualProduction/RemoteControl/Source/RemoteControlUI/Private/UI/SRCPanelTreeNode.h
@@ -69,36 +69,54 @@ struct SRCPanelTreeNode : public SCompoundWidget, public IHasProtocolExtensibili
 
 	/** Get this tree node's childen. */
 	virtual void GetNodeChildren(TArray<TSharedPtr<SRCPanelTreeNode>>& OutChildren) const {}
+
 	/** Get this node's ID if any. */
 	virtual FGuid GetRCId() const { return FGuid(); }
+
 	/** Get get this node's type. */
 	virtual ENodeType GetRCType() const { return ENodeType::Invalid; }
+
 	/** Returns true if this tree node has childen. */
 	virtual bool HasChildren() const { return false; }
+
 	/** Refresh the node. */
-	virtual void Refresh() {};
+	virtual void Refresh() {}
+
 	/** Get the context menu for this node. */
 	virtual TSharedPtr<SWidget> GetContextMenu() { return nullptr; }
+
 	/** Set whether this widget is currently hovered when drag and dropping. */
 	virtual void SetIsHovered(bool bIsBeingHovered) {}
+
 	/** Make the node name's text box editable. */
-	virtual void EnterRenameMode() {};
+	virtual void EnterRenameMode() {}
+
 	/** Get the PropertyId of this Node. */
 	virtual FName GetPropertyId() { return FName(TEXT("")); }
+
 	/** Set the PropertyId of this Node. */
-	virtual void SetPropertyId(FName InNewPropertyId) {};
+	virtual void SetPropertyId(FName InNewPropertyId) {}
+
+	/** Focus the PropertyId widget. */
+	virtual void FocusPropertyIdWidget() const {}
+
 	/** Set the Name of this Node. */
-	virtual void SetName(FName InNewName) {};
+	virtual void SetName(FName InNewName) {}
+
 	/** Updates the highlight text to active search term. */
-	virtual void SetHighlightText(const FText& InHightlightText = FText::GetEmpty()) {};
+	virtual void SetHighlightText(const FText& InHightlightText = FText::GetEmpty()) {}
+
 	/** Retrieves the referenced widget corresponding to the given column name. */
 	virtual TSharedRef<SWidget> GetWidget(const FName ForColumnName, const FName InActiveProtocol);
-	/** Retrieves the DragAndDropWidget if possible otherwise returns a NullWidget */
+
+	/** Retrieves the DragAndDropWidget if possible otherwise returns a NullWidget. */
 	TSharedRef<SWidget> GetDragAndDropWidget(int32 InSelectedEntitiesNum = 1);
-	/** Executed when the Name of the node is changed */
-	FOnLabelModified& OnLabelModified() { return OnLabelModifiedDelegate; };
-	/** Executed when the PropertyId of the node is changed */
-	FOnPropertyIdRenamed& OnPropertyIdRenamed() { return OnPropertyIdRenamedDelegate; };
+
+	/** Executed when the Name of the node is changed. */
+	FOnLabelModified& OnLabelModified() { return OnLabelModifiedDelegate; }
+
+	/** Executed when the PropertyId of the node is changed. */
+	FOnPropertyIdRenamed& OnPropertyIdRenamed() { return OnPropertyIdRenamedDelegate; }
 
 protected:
 	struct FMakeNodeWidgetArgs
@@ -113,8 +131,10 @@ protected:
 	
 	/** Create a widget that represents a row with a splitter. */
 	TSharedRef<SWidget> MakeSplitRow(TSharedRef<SWidget> LeftColumn, TSharedRef<SWidget> RightColumn);
+
 	/** Create a widget that represents a node in the panel tree hierarchy. */
 	TSharedRef<SWidget> MakeNodeWidget(const FMakeNodeWidgetArgs& Args);
+
 	/** Creates cached copies of underlying widgets. */
 	void MakeNodeWidgets(const FMakeNodeWidgetArgs& Args);
 
diff --git a/Engine/Plugins/VirtualProduction/RemoteControl/Source/RemoteControlUI/Private/UI/SRemoteControlPanel.cpp b/Engine/Plugins/VirtualProduction/RemoteControl/Source/RemoteControlUI/Private/UI/SRemoteControlPanel.cpp
--- a/Engine/Plugins/VirtualProduction/RemoteControl/Source/RemoteControlUI/Private/UI/SRemoteControlPanel.cpp
+++ b/Engine/Plugins/VirtualProduction/RemoteControl/Source/RemoteControlUI/Private/UI/SRemoteControlPanel.cpp
@@ -1418,6 +1418,11 @@ void SRemoteControlPanel::BindRemoteControlCommands()
 		FExecuteAction::CreateSP(this, &SRemoteControlPanel::RenameEntity_Execute),
 		FCanExecuteAction::CreateSP(this, &SRemoteControlPanel::CanRenameEntity));
 
+	ActionList.MapAction(
+		Commands.ChangePropId,
+		FExecuteAction::CreateSP(this, &SRemoteControlPanel::ChangePropertyId_Execute),
+		FCanExecuteAction::CreateSP(this, &SRemoteControlPanel::CanChangePropertyId));
+
 	ActionList.MapAction(
 		Commands.CopyItem,
 		FExecuteAction::CreateSP(this, &SRemoteControlPanel::CopyItem_Execute),
@@ -2546,6 +2551,39 @@ bool SRemoteControlPanel::CanRenameEntity() const
 	return false;
 }
 
+void SRemoteControlPanel::ChangePropertyId_Execute() const
+{
+	if (!LastSelectedEntity.IsValid())
+	{
+		return;
+	}
+
+	if (LastSelectedEntity->GetRCType() == SRCPanelTreeNode::Field)
+	{
+		LastSelectedEntity->FocusPropertyIdWidget();
+	}
+}
+
+bool SRemoteControlPanel::CanChangePropertyId() const
+{
+	if (bIsInLiveMode)
+	{
+		return false;
+	}
+
+	if (!LastSelectedEntity.IsValid())
+	{
+		return false;
+	}
+
+	if (LastSelectedEntity->GetRCType() != SRCPanelTreeNode::Field)
+	{
+		return false;
+	}
+
+	return true;
+}
+
 void SRemoteControlPanel::SetLogicClipboardItems(const TArray<UObject*>& InItems, const TSharedPtr<SRCLogicPanelBase>& InSourcePanel)
 {
 	LogicClipboardItems = InItems;
diff --git a/Engine/Plugins/VirtualProduction/RemoteControl/Source/RemoteControlUI/Private/UI/SRemoteControlPanel.h b/Engine/Plugins/VirtualProduction/RemoteControl/Source/RemoteControlUI/Private/UI/SRemoteControlPanel.h
--- a/Engine/Plugins/VirtualProduction/RemoteControl/Source/RemoteControlUI/Private/UI/SRemoteControlPanel.h
+++ b/Engine/Plugins/VirtualProduction/RemoteControl/Source/RemoteControlUI/Private/UI/SRemoteControlPanel.h
@@ -319,6 +319,12 @@ private:
 	/** Called to test if user is able to rename a group/exposed entity. */
 	bool CanRenameEntity() const;
 
+	/** Called when user attempts to change property Ids. */
+	void ChangePropertyId_Execute() const;
+
+	/** Called to test if user is able to change property Ids. */
+	bool CanChangePropertyId() const;
+
 	/** Called when user attempts to Copy a logic UI item. */
 	void CopyItem_Execute();
 
