[RemoteControl] Add RightClick Expose/Unexpose in DetailsPanel
Right Click details panel add Expose/Expose Sub-Property Option

diff --git a/Engine/Plugins/VirtualProduction/RemoteControl/Source/RemoteControlUI/Private/RemoteControlUIModule.cpp b/Engine/Plugins/VirtualProduction/RemoteControl/Source/RemoteControlUI/Private/RemoteControlUIModule.cpp
--- a/Engine/Plugins/VirtualProduction/RemoteControl/Source/RemoteControlUI/Private/RemoteControlUIModule.cpp
+++ b/Engine/Plugins/VirtualProduction/RemoteControl/Source/RemoteControlUI/Private/RemoteControlUIModule.cpp
@@ -10,6 +10,7 @@
 #include "Components/SkeletalMeshComponent.h"
 #include "Components/StaticMeshComponent.h"
 #include "Controller/RCCustomControllerUtilities.h"
+#include "DetailRowMenuContext.h"
 #include "DetailTreeNode.h"
 #include "Elements/Framework/TypedElementRegistry.h"
 #include "Elements/Framework/TypedElementSelectionSet.h"
@@ -93,6 +94,12 @@ FRCExposesPropertyArgs::FRCExposesPropertyArgs(TSharedPtr<IPropertyHandle>& InPr
 {
 }
 
+FRCExposesPropertyArgs::FRCExposesPropertyArgs(const TSharedPtr<IPropertyHandle>& InPropertyHandle)
+	: PropertyHandle(InPropertyHandle)
+	, Id(FGuid::NewGuid())
+{
+}
+
 FRCExposesPropertyArgs::FRCExposesPropertyArgs(UObject* InOwnerObject, const FString& InPropertyPath, FProperty* InProperty)
 	: OwnerObject(InOwnerObject)
 	, PropertyPath(InPropertyPath)
@@ -415,6 +422,159 @@ void FRemoteControlUIModule::UnregisterEvents()
 	FEditorDelegates::PostUndoRedo.RemoveAll(this);
 }
 
+void FRemoteControlUIModule::ExtendPropertyRowContextMenu() const
+{
+	UToolMenus* Menus = UToolMenus::Get();
+
+	check(Menus);
+
+	if (UToolMenu* ContextMenu = Menus->FindMenu(UE::PropertyEditor::RowContextMenuName))
+	{
+		ContextMenu->AddDynamicSection(
+			TEXT("FillRemoteControlRowContextSection")
+			, FNewToolMenuDelegate::CreateRaw(this, &FRemoteControlUIModule::FillRemoteControlRowContextSection));
+	}
+}
+
+void FRemoteControlUIModule::FillRemoteControlRowContextSection(UToolMenu* InToolMenu) const
+{
+	if (!InToolMenu)
+	{
+		return;
+	}
+
+	// For context menu in details view
+	const UDetailRowMenuContext* Context = InToolMenu->FindContext<UDetailRowMenuContext>();
+
+	if (!Context || Context->PropertyHandles.IsEmpty())
+	{
+		return;
+	}
+
+	FRCExposesPropertyArgs ExposesPropertyArgs(Context->PropertyHandles[0]);
+
+	if (!ExposesPropertyArgs.IsValid())
+	{
+		return;
+	}
+
+	static const FName RemoteControlSectionName("ContextRemoteControlActions");
+
+	if (!GetActivePreset())
+	{
+		return;
+	}
+
+	FToolMenuSection* RemoteControlSection = nullptr;
+	RemoteControlSection = InToolMenu->FindSection(RemoteControlSectionName);
+	if (!RemoteControlSection)
+	{
+		RemoteControlSection = &InToolMenu->AddSection(RemoteControlSectionName
+			, LOCTEXT("ContextRemoteControlActions", "Remote Control Actions")
+			, FToolMenuInsert(NAME_None, EToolMenuInsertType::First));
+	}
+
+	const TAttribute<FText> ExposeUnexposeLabel = TAttribute<FText>::Create(TAttribute<FText>::FGetter::CreateRaw(this, &FRemoteControlUIModule::GetExposePropertyButtonText, ExposesPropertyArgs));
+	const TAttribute<FText> ExposeUnexposeTooltip = TAttribute<FText>::Create(TAttribute<FText>::FGetter::CreateRaw(this, &FRemoteControlUIModule::GetExposePropertyButtonTooltip, ExposesPropertyArgs));
+	const TAttribute<FSlateIcon> ExposeUnexposeIcon = TAttribute<FSlateIcon>::Create(
+			[this, ExposesPropertyArgs]
+		{
+			return OnGetExposedIcon(ExposesPropertyArgs);
+		});;
+
+	// Unexpose/Expose entry for the Menu
+	RemoteControlSection->AddMenuEntry(
+		TEXT("RemoteControlExposeUnexposeEntry")
+		, ExposeUnexposeLabel
+		, ExposeUnexposeTooltip
+		, ExposeUnexposeIcon
+		, FUIAction(
+			FExecuteAction::CreateRaw(this, &FRemoteControlUIModule::OnToggleExposeProperty, ExposesPropertyArgs),
+			FCanExecuteAction::CreateRaw(this, &FRemoteControlUIModule::CanToggleExposeProperty, ExposesPropertyArgs),
+			FGetActionCheckState::CreateRaw(this, &FRemoteControlUIModule::GetPropertyExposedCheckState, ExposesPropertyArgs),
+			FIsActionButtonVisible::CreateRaw(this, &FRemoteControlUIModule::CanToggleExposeProperty, ExposesPropertyArgs)
+			)
+		);
+
+	// Unexpose/Expose SubProperty SubMenu for the Menu
+	if (ExposesPropertyArgs.PropertyHandle.IsValid() && HasChildProperties(ExposesPropertyArgs.PropertyHandle->GetProperty()))
+	{
+		const TAttribute<FSlateIcon> ExposeSubPropertyIcon = TAttribute<FSlateIcon>::Create([]()
+			{
+				return FSlateIcon(FRemoteControlExposeMenuStyle::GetStyleSetName(),"RemoteControlExposeMenu.Expand");
+			});
+		RemoteControlSection->AddSubMenu(
+			TEXT("ExposeUnexposeSubPropertyMenu")
+			, LOCTEXT("ExposeUnexposeSubPropertyMenu_Label", "Toggle Sub Property of this field")
+			, LOCTEXT("ExposeUnexposeSubPropertyMenu_Tooltip", "Let you toggle Sub Property of the given field if any")
+			, FNewToolMenuDelegate::CreateRaw(this, &FRemoteControlUIModule::GetSubPropertySubMenu, ExposesPropertyArgs));
+	}
+}
+
+void FRemoteControlUIModule::GetSubPropertySubMenu(UToolMenu* InToolMenu, FRCExposesPropertyArgs InExposesPropertyArgs) const
+{
+	if (!InToolMenu || !InExposesPropertyArgs.IsValid())
+	{
+		return;
+	}
+	TArray<FRCExposesAllPropertiesArgs> ExposeAllArgs;
+	GetAllExposableSubPropertyFromStruct(InExposesPropertyArgs, ExposeAllArgs);
+
+	FToolMenuSection* RemoteControlExposeUnexposeAllSubMenuSection = nullptr;
+	const FName RemoteControlExposeUnexposeAllSubMenuSectionName(InExposesPropertyArgs.GetProperty()->GetDisplayNameText().ToString());
+	RemoteControlExposeUnexposeAllSubMenuSection = InToolMenu->FindSection(RemoteControlExposeUnexposeAllSubMenuSectionName);
+	if (!RemoteControlExposeUnexposeAllSubMenuSection)
+	{
+		RemoteControlExposeUnexposeAllSubMenuSection = &InToolMenu->AddSection(RemoteControlExposeUnexposeAllSubMenuSectionName
+			, InExposesPropertyArgs.GetProperty()->GetDisplayNameText()
+			, FToolMenuInsert(NAME_None, EToolMenuInsertType::First));
+	}
+	
+	// Expose and Unexpose ALL section
+	RemoteControlExposeUnexposeAllSubMenuSection->AddMenuEntry(
+		TEXT("SubMenuExposeAllEntry")
+		, LOCTEXT("RCExposeAll", "Expose All")
+		, LOCTEXT("RCExposeAllTooltip", "Expose all sub-properties")
+		, FSlateIcon()
+		, FUIAction(FExecuteAction::CreateRaw(this, &FRemoteControlUIModule::OnExposeAll, ExposeAllArgs))
+		, EUserInterfaceActionType::CollapsedButton);
+
+	RemoteControlExposeUnexposeAllSubMenuSection->AddMenuEntry(
+	TEXT("SubMenuUnexposeAllEntry")
+		, LOCTEXT("RCUnexposeAll", "Unexpose All")
+		, LOCTEXT("RCUnexposeAllToolTip", "Unexpose all sub-properties")
+		, FSlateIcon()
+		, FUIAction(FExecuteAction::CreateRaw(this, &FRemoteControlUIModule::OnUnexposeAll, ExposeAllArgs))
+		, EUserInterfaceActionType::CollapsedButton);
+
+	FToolMenuSection* RemoteControlSubMenuSection = nullptr;
+	const FName RemoteControlSubMenuSectionName("EXPOSE");
+	RemoteControlSubMenuSection = InToolMenu->FindSection(RemoteControlSubMenuSectionName);
+	if (!RemoteControlSubMenuSection)
+	{
+		RemoteControlSubMenuSection = &InToolMenu->AddSection(RemoteControlSubMenuSectionName
+			, LOCTEXT("ContextRemoteControlActionsSubProperties", "EXPOSE")
+			, FToolMenuInsert(NAME_None, EToolMenuInsertType::First));
+	}
+
+	// Expose and Unexpose single property section
+	for (FRCExposesAllPropertiesArgs PropArgs : ExposeAllArgs)
+	{
+		const FName SubPropertyName = FName("ExposeUnexposeSingleEntry_" + PropArgs.PropName.ToString());
+		RemoteControlSubMenuSection->AddMenuEntry(
+			SubPropertyName,
+			PropArgs.ExposedPropertyLabel,
+			PropArgs.ToolTip,
+			FSlateIcon(),
+			FUIAction(
+				FExecuteAction::CreateRaw(this, &FRemoteControlUIModule::OnToggleExposeSubProperty, PropArgs.PropertyArgs, PropArgs.DesiredName),
+				FCanExecuteAction::CreateRaw(this, &FRemoteControlUIModule::CanToggleExposeProperty, PropArgs.PropertyArgs),
+				FGetActionCheckState::CreateRaw(this, &FRemoteControlUIModule::GetPropertyExposedCheckState, PropArgs.PropertyArgs),
+				FIsActionButtonVisible::CreateRaw(this, &FRemoteControlUIModule::CanToggleExposeProperty, PropArgs.PropertyArgs)),
+			EUserInterfaceActionType::ToggleButton);
+	}
+}
+
 URemoteControlPreset* FRemoteControlUIModule::GetActivePreset() const
 {
 	if (const TSharedPtr<SRemoteControlPanel> Panel = GetPanelForObject(nullptr))
@@ -485,6 +645,9 @@ void FRemoteControlUIModule::HandleCreatePropertyRowExtension(const FOnGenerateG
 
 	if (ExposesPropertyArgs.IsValid())
 	{
+		// Extend context row menu
+		ExtendPropertyRowContextMenu();
+
 		// Expose/Unexpose button.
 		FPropertyRowExtensionButton& ExposeButton = OutExtensions.AddDefaulted_GetRef();
 		ExposeButton.Icon = TAttribute<FSlateIcon>::Create(
@@ -651,7 +814,7 @@ ECheckBoxState FRemoteControlUIModule::GetPropertyExposedCheckState(const FRCExp
 	return ECheckBoxState::Unchecked;
 }
 
-void FRemoteControlUIModule::OnToggleExposeProperty(const FRCExposesPropertyArgs InPropertyArgs)
+void FRemoteControlUIModule::OnToggleExposeProperty(const FRCExposesPropertyArgs InPropertyArgs) const
 {
 	if (!ensureMsgf(InPropertyArgs.IsValid(), TEXT("Property could not be exposed because the extension args was invalid.")))
 	{
diff --git a/Engine/Plugins/VirtualProduction/RemoteControl/Source/RemoteControlUI/Private/RemoteControlUIModule.h b/Engine/Plugins/VirtualProduction/RemoteControl/Source/RemoteControlUI/Private/RemoteControlUIModule.h
--- a/Engine/Plugins/VirtualProduction/RemoteControl/Source/RemoteControlUI/Private/RemoteControlUIModule.h
+++ b/Engine/Plugins/VirtualProduction/RemoteControl/Source/RemoteControlUI/Private/RemoteControlUIModule.h
@@ -117,6 +117,15 @@ private:
 	void RegisterEvents();
 	void UnregisterEvents();
 
+	/** Extend context row menu */
+	void ExtendPropertyRowContextMenu() const;
+
+	/** Fills the RemoteControl details view row context menu  */
+	void FillRemoteControlRowContextSection(UToolMenu* InToolMenu) const;
+
+	/** Create the SubProperty expose/unexpose sub menu for the details view row context menu */
+	void GetSubPropertySubMenu(UToolMenu* InToolMenu, FRCExposesPropertyArgs InExposesPropertyArgs) const;
+
 	/** Handle creating the row extensions.  */
 	void HandleCreatePropertyRowExtension(const FOnGenerateGlobalRowExtensionArgs& InArgs, TArray<FPropertyRowExtensionButton>& OutExtensions);
 
@@ -138,7 +147,7 @@ private:
 	ECheckBoxState GetPropertyExposedCheckState(const FRCExposesPropertyArgs InArgs) const;
 
 	/** Handle clicking the expose button. */
-	void OnToggleExposeProperty(const FRCExposesPropertyArgs InArgs);
+	void OnToggleExposeProperty(const FRCExposesPropertyArgs InArgs) const;
 
 	/** Handle clicking the expose SubProperty button. */
 	void OnToggleExposeSubProperty(const FRCExposesPropertyArgs InArgs, const FString InDesiredName = TEXT("")) const;
diff --git a/Engine/Plugins/VirtualProduction/RemoteControl/Source/RemoteControlUI/Public/IRemoteControlUIModule.h b/Engine/Plugins/VirtualProduction/RemoteControl/Source/RemoteControlUI/Public/IRemoteControlUIModule.h
--- a/Engine/Plugins/VirtualProduction/RemoteControl/Source/RemoteControlUI/Public/IRemoteControlUIModule.h
+++ b/Engine/Plugins/VirtualProduction/RemoteControl/Source/RemoteControlUI/Public/IRemoteControlUIModule.h
@@ -61,6 +61,7 @@ struct FRCExposesPropertyArgs
 	FRCExposesPropertyArgs(const FOnGenerateGlobalRowExtensionArgs& InExtensionArgs);
 	FRCExposesPropertyArgs(FOnGenerateGlobalRowExtensionArgs&& InExtensionArgs);
 	FRCExposesPropertyArgs(TSharedPtr<IPropertyHandle>& InPropertyHandle);
+	FRCExposesPropertyArgs(const TSharedPtr<IPropertyHandle>& InPropertyHandle);
 	FRCExposesPropertyArgs(UObject* InOwnerObject, const FString& InPropertyPath, FProperty* InProperty);
 
 
