[RemoteControl] Lock toggle for RemoteControl
Lock toggle for Vectors/Rotators to make x/y/z maintain ratio when changing one

diff --git a/Engine/Plugins/VirtualProduction/RemoteControl/Source/RemoteControlLogic/Private/Action/RCPropertyIdAction.cpp b/Engine/Plugins/VirtualProduction/RemoteControl/Source/RemoteControlLogic/Private/Action/RCPropertyIdAction.cpp
--- a/Engine/Plugins/VirtualProduction/RemoteControl/Source/RemoteControlLogic/Private/Action/RCPropertyIdAction.cpp
+++ b/Engine/Plugins/VirtualProduction/RemoteControl/Source/RemoteControlLogic/Private/Action/RCPropertyIdAction.cpp
@@ -228,6 +228,13 @@ void URCPropertyIdAction::UpdatePropertyId()
 							if (!bIsSpecialCase)
 							{
 								CachedPropertySelfContainer[CurrentKey]->DuplicateProperty(NewPropertyIdName, PropToDuplicate);
+#if WITH_EDITORONLY_DATA
+								// In case the property is a Vector/Rotator we don't want the Lock widget
+								if (FProperty* CachedProperty = CachedPropertySelfContainer[CurrentKey]->GetProperty())
+								{
+									CachedProperty->RemoveMetaData(TEXT("AllowPreserveRatio"));
+								}
+#endif
 							}
 
 							// Do this the first time it is created so that it will have a better default value except for Object
diff --git a/Engine/Plugins/VirtualProduction/RemoteControl/Source/RemoteControlLogic/Private/Controller/RCControllerContainer.cpp b/Engine/Plugins/VirtualProduction/RemoteControl/Source/RemoteControlLogic/Private/Controller/RCControllerContainer.cpp
--- a/Engine/Plugins/VirtualProduction/RemoteControl/Source/RemoteControlLogic/Private/Controller/RCControllerContainer.cpp
+++ b/Engine/Plugins/VirtualProduction/RemoteControl/Source/RemoteControlLogic/Private/Controller/RCControllerContainer.cpp
@@ -26,18 +26,26 @@ void URCControllerContainer::UpdateEntityIds(const TMap<FGuid, FGuid>& InEntityI
 
 URCVirtualPropertyInContainer* URCControllerContainer::AddProperty(const FName& InPropertyName, TSubclassOf<URCVirtualPropertyInContainer> InPropertyClass, const EPropertyBagPropertyType InValueType, UObject* InValueTypeObject /*= nullptr*/, TArray<FPropertyBagPropertyDescMetaData> MetaData /*= TArray<FPropertyBagPropertyDescMetaData>()*/)
 {
-	// Vector Controllers
 	if (InValueType == EPropertyBagPropertyType::Struct)
 	{
+		// Vector Controllers
 		if (InValueTypeObject == TBaseStructure<FVector>::Get() || InValueTypeObject == TBaseStructure<FVector2D>::Get())
 		{
 			MetaData.Add(FPropertyBagPropertyDescMetaData(FName("Delta"), FString::Printf(TEXT("%f"), VectorSliderDelta)));
 			MetaData.Add(FPropertyBagPropertyDescMetaData(FName("LinearDeltaSensitivity"), FString::Printf(TEXT("%f"), VectorLinearDeltaSensitivity)));
+			MetaData.Add(FPropertyBagPropertyDescMetaData(FName("AllowPreserveRatio"), TEXT("")));
 		}
+		// Color Controllers
+		else if (InValueTypeObject == TBaseStructure<FColor>::Get())
+		{
+			MetaData.Add(FPropertyBagPropertyDescMetaData(FName("OnlyUpdateOnInteractionEnd"), TEXT("true")));
+		}
+		// Rotator Controllers
 		else if (InValueTypeObject == TBaseStructure<FRotator>::Get())
 		{
 			MetaData.Add(FPropertyBagPropertyDescMetaData(FName("Delta"), FString::Printf(TEXT("%f"), RotatorSliderDelta)));
 			MetaData.Add(FPropertyBagPropertyDescMetaData(FName("LinearDeltaSensitivity"), FString::Printf(TEXT("%f"), RotatorLinearDeltaSensitivity)));
+			MetaData.Add(FPropertyBagPropertyDescMetaData(FName("AllowPreserveRatio"), TEXT("")));
 		}
 	}
 
diff --git a/Engine/Plugins/VirtualProduction/RemoteControl/Source/RemoteControlUI/Private/UI/Action/RCActionModel.cpp b/Engine/Plugins/VirtualProduction/RemoteControl/Source/RemoteControlUI/Private/UI/Action/RCActionModel.cpp
--- a/Engine/Plugins/VirtualProduction/RemoteControl/Source/RemoteControlUI/Private/UI/Action/RCActionModel.cpp
+++ b/Engine/Plugins/VirtualProduction/RemoteControl/Source/RemoteControlUI/Private/UI/Action/RCActionModel.cpp
@@ -196,6 +196,12 @@ FRCPropertyActionType::FRCPropertyActionType(URCPropertyAction* InPropertyAction
 
 	if (InPropertyAction)
 	{
+		// In case the Action is a Vector/Rotator we don't want the Lock widget
+		if (FProperty* ActionProperty = InPropertyAction->GetProperty())
+		{
+			ActionProperty->RemoveMetaData(TEXT("AllowPreserveRatio"));
+		}
+
 		// Generate UI widget for Action input
 		if (const TSharedPtr<FStructOnScope> StructOnScope = InPropertyAction->PropertySelfContainer->CreateStructOnScope())
 		{
diff --git a/Engine/Plugins/VirtualProduction/RemoteControl/Source/RemoteControlUI/Private/UI/Behaviour/Builtin/Conditional/RCBehaviourConditionalModel.cpp b/Engine/Plugins/VirtualProduction/RemoteControl/Source/RemoteControlUI/Private/UI/Behaviour/Builtin/Conditional/RCBehaviourConditionalModel.cpp
--- a/Engine/Plugins/VirtualProduction/RemoteControl/Source/RemoteControlUI/Private/UI/Behaviour/Builtin/Conditional/RCBehaviourConditionalModel.cpp
+++ b/Engine/Plugins/VirtualProduction/RemoteControl/Source/RemoteControlUI/Private/UI/Behaviour/Builtin/Conditional/RCBehaviourConditionalModel.cpp
@@ -71,6 +71,12 @@ void FRCBehaviourConditionalModel::CreateComparandInputField()
 				Comparand->DuplicateProperty(FName("Comparand"), Controller->GetProperty());
 			}
 
+			// In case the Comparand is a Vector/Rotator we don't want the Lock widget
+			if (FProperty* ComparandProperty = Comparand->GetProperty())
+			{
+				ComparandProperty->RemoveMetaData(TEXT("AllowPreserveRatio"));
+			}
+
 			// UI widget (via Property Generator)
 			DetailTreeNodeWeakPtr = UE::RCUIHelpers::GetDetailTreeNodeForVirtualProperty(Comparand, PropertyRowGenerator);
 		}
diff --git a/Engine/Plugins/VirtualProduction/RemoteControl/Source/RemoteControlUI/Private/UI/Controller/RCControllerModel.cpp b/Engine/Plugins/VirtualProduction/RemoteControl/Source/RemoteControlUI/Private/UI/Controller/RCControllerModel.cpp
--- a/Engine/Plugins/VirtualProduction/RemoteControl/Source/RemoteControlUI/Private/UI/Controller/RCControllerModel.cpp
+++ b/Engine/Plugins/VirtualProduction/RemoteControl/Source/RemoteControlUI/Private/UI/Controller/RCControllerModel.cpp
@@ -55,18 +55,6 @@ TSharedRef<SWidget> FRCControllerModel::GetWidget() const
 {
 	const FNodeWidgets NodeWidgets = DetailTreeNodeWeakPtr.Pin()->CreateNodeWidgets();
 
-	// We need to add this metadata to the ColorController to avoid the update when dragging causing a lot of lag
-	const TSharedPtr<IPropertyHandle> PropertyHandle = DetailTreeNodeWeakPtr.Pin()->CreatePropertyHandle();
-	if (PropertyHandle.IsValid())
-	{
-		FStructProperty* StructProperty = CastField<FStructProperty>(PropertyHandle->GetProperty());
-		if (StructProperty && StructProperty->Struct &&
-			StructProperty->Struct.GetFName() == FName("Color") &&
-			!StructProperty->HasMetaData("OnlyUpdateOnInteractionEnd"))
-		{
-			StructProperty->AppendMetaData({{FName("OnlyUpdateOnInteractionEnd"), TEXT("true")}});
-		}
-	}
 	const TSharedRef<SHorizontalBox> FieldWidget = SNew(SHorizontalBox);
 	if (VirtualPropertyWeakPtr.IsValid())
 	{
@@ -85,7 +73,6 @@ TSharedRef<SWidget> FRCControllerModel::GetWidget() const
 		{
 			FieldWidget->AddSlot()
 				.Padding(SlotMargin)
-				.HAlign(HAlign_Left)
 				[
 					NodeWidgets.ValueWidget.ToSharedRef()
 				];
diff --git a/Engine/Plugins/VirtualProduction/RemoteControl/Source/RemoteControlUI/Private/UI/Controller/SRCControllerPanelList.cpp b/Engine/Plugins/VirtualProduction/RemoteControl/Source/RemoteControlUI/Private/UI/Controller/SRCControllerPanelList.cpp
--- a/Engine/Plugins/VirtualProduction/RemoteControl/Source/RemoteControlUI/Private/UI/Controller/SRCControllerPanelList.cpp
+++ b/Engine/Plugins/VirtualProduction/RemoteControl/Source/RemoteControlUI/Private/UI/Controller/SRCControllerPanelList.cpp
@@ -476,11 +476,6 @@ void SRCControllerPanelList::Reset()
 			FProperty* Property = Child->CreatePropertyHandle()->GetProperty();
 			check(Property);
 
-			if (Property->IsA<FStrProperty>() || Property->IsA<FTextProperty>())
-			{
-				Property->AppendMetaData({{TEXT("multiline"), TEXT("true")}});
-			}
-
 			if (URCVirtualPropertyBase* Controller = Preset->GetController(Property->GetFName()))
 			{
 				bool bIsVisible = true;
