[RemoteControl] Fix ensure hit with Text action
Fix Text action that hit an ensure when changing their value

diff --git a/Engine/Plugins/VirtualProduction/RemoteControl/Source/RemoteControlLogic/Private/Action/RCActionContainer.cpp b/Engine/Plugins/VirtualProduction/RemoteControl/Source/RemoteControlLogic/Private/Action/RCActionContainer.cpp
--- a/Engine/Plugins/VirtualProduction/RemoteControl/Source/RemoteControlLogic/Private/Action/RCActionContainer.cpp
+++ b/Engine/Plugins/VirtualProduction/RemoteControl/Source/RemoteControlLogic/Private/Action/RCActionContainer.cpp
@@ -192,6 +192,7 @@ URCPropertyAction* URCActionContainer::AddPropertyAction(const TSharedRef<const
 		{
 			const FName& PropertyName = InRemoteControlProperty->GetProperty()->GetFName();
 			NewPropertyAction->PropertySelfContainer->DuplicatePropertyWithCopy(PropertyName, InRemoteControlProperty->GetProperty(), (uint8*)ObjectRefReading.ContainerAdress);
+			NewPropertyAction->PropertySelfContainer->PresetWeakPtr = PresetWeakPtr;
 		}
 	}
 
