[RemoteControl] Add button to create controller from the PathBehavior
Added `CreateController` button in PathBehavior
Avoiding resetting selection when creating new controllers
Added soft refresh when changing controller/behavior selection

diff --git a/Engine/Plugins/VirtualProduction/RemoteControl/Source/RemoteControlUI/Private/UI/Action/SRCActionPanel.cpp b/Engine/Plugins/VirtualProduction/RemoteControl/Source/RemoteControlUI/Private/UI/Action/SRCActionPanel.cpp
--- a/Engine/Plugins/VirtualProduction/RemoteControl/Source/RemoteControlUI/Private/UI/Action/SRCActionPanel.cpp
+++ b/Engine/Plugins/VirtualProduction/RemoteControl/Source/RemoteControlUI/Private/UI/Action/SRCActionPanel.cpp
@@ -37,11 +37,12 @@
 #include "Widgets/Input/SButton.h"
 #include "Widgets/Input/SComboButton.h"
 #include "Widgets/Input/SCheckBox.h"
+#include "Widgets/Input/SComboButton.h"
 #include "Widgets/Layout/SBox.h"
 
 #define LOCTEXT_NAMESPACE "SRCActionPanel"
 
-TSharedRef<SBox> SRCActionPanel::CreateNoneSelectedWidget()
+TSharedRef<SBox> SRCActionPanel::GetNoneSelectedWidget()
 {
 	return SNew(SBox)
 	.Padding(0.f)
@@ -63,7 +64,6 @@
 	
 	RCPanelStyle = &FRemoteControlPanelStyle::Get()->GetWidgetStyle<FRCPanelStyle>("RemoteControlPanel.MinorPanel");
 
-	WrappedBoxWidget = SNew(SBox);
 	UpdateWrappedWidget();
 	
 	ChildSlot
@@ -97,118 +97,27 @@
 
 void SRCActionPanel::UpdateWrappedWidget(TSharedPtr<FRCBehaviourModel> InBehaviourItem)
 {
+	if (!AreActionPanelWidgetsValid())
+	{
+		CreateActionPanelWidgets();
+	}
+
+	// Update Behavior details widget
+	BehaviourDetailsWidget->SetNewBehavior(InBehaviourItem);
+
 	if (InBehaviourItem.IsValid())
 	{
+		// Create action list based on the behavior
 		ActionPanelList = InBehaviourItem->GetActionsListWidget(SharedThis(this));
+		ActionDockPanel->SetContent(ActionPanelList.ToSharedRef());
 
-		// Action Dock Panel
-		TSharedPtr<SRCMinorPanel> ActionDockPanel = SNew(SRCMinorPanel)
-			.HeaderLabel(LOCTEXT("ActionsLabel", "Actions"))
-			.EnableFooter(false)
-			[
-				ActionPanelList.ToSharedRef()
-			];
+		// Update add new action menu content
+		AddNewActionButton->SetMenuContent(GetActionMenuContentWidget());
 
-		// Add New Action Button
-		bAddActionMenuNeedsRefresh = true;
-
-		const TSharedRef<SWidget> AddNewActionButton = SNew(SComboButton)
-			.AddMetaData<FTagMetaData>(FTagMetaData(TEXT("Add Action")))
-			.HAlign(HAlign_Center)
-			.VAlign(VAlign_Center)
-			.ButtonStyle(&RCPanelStyle->FlatButtonStyle)
-			.ForegroundColor(FSlateColor::UseForeground())
-			.CollapseMenuOnParentFocus(true)
-			.HasDownArrow(false)
-			.ContentPadding(FMargin(4.f, 2.f))
-			.ButtonContent()
-			[
-				SNew(SBox)
-				.WidthOverride(RCPanelStyle->IconSize.X)
-				.HeightOverride(RCPanelStyle->IconSize.Y)
-				[
-					SNew(SImage)
-					.ColorAndOpacity(FSlateColor::UseForeground())
-					.Image(FAppStyle::GetBrush("Icons.PlusCircle"))
-				]
-			]
-			.OnGetMenuContent(this, &SRCActionPanel::GetActionMenuContentWidget);
-		
-		// Add All Button
-		TSharedRef<SWidget> AddAllActionsButton = SNew(SButton)
-			.AddMetaData<FTagMetaData>(FTagMetaData(TEXT("Add All Actions")))
-			.HAlign(HAlign_Center)
-			.VAlign(VAlign_Center)
-			.ForegroundColor(FSlateColor::UseForeground())
-			.ButtonStyle(&RCPanelStyle->FlatButtonStyle)
-			.ToolTipText(LOCTEXT("AddAllToolTip", "Adds all the available actions."))
-			.OnClicked(this, &SRCActionPanel::OnAddAllFields)
-			.Visibility(this, &SRCActionPanel::HandleAddAllButtonVisibility)
-			[
-				SNew(SBox)
-				.WidthOverride(RCPanelStyle->IconSize.X)
-				.HeightOverride(RCPanelStyle->IconSize.Y)
-				[
-					SNew(SImage)
-					.ColorAndOpacity(FSlateColor::UseForeground())
-					.Image(FAppStyle::GetBrush("Icons.Duplicate"))
-				]
-			];
-
-		// Add All Button
-		const TSharedRef<SWidget> AddAllSelectedActionsButton = SNew(SButton)
-			.AddMetaData<FTagMetaData>(FTagMetaData(TEXT("Add All Selected Fields")))
-			.HAlign(HAlign_Center)
-			.VAlign(VAlign_Center)
-			.ForegroundColor(FSlateColor::UseForeground())
-			.ButtonStyle(&RCPanelStyle->FlatButtonStyle)
-			.ToolTipText(LOCTEXT("RCAddAllSelectedToolTip", "Adds all the selected fields."))
-			.OnClicked(this, &SRCActionPanel::OnAddAllSelectedFields)
-			.Visibility(this, &SRCActionPanel::HandleAddAllButtonVisibility)
-			[
-				SNew(SBox)
-				.WidthOverride(RCPanelStyle->IconSize.X)
-				.HeightOverride(RCPanelStyle->IconSize.Y)
-				[
-					SNew(SImage)
-					.ColorAndOpacity(FSlateColor::UseForeground())
-					.Image(FAppStyle::GetBrush("DataLayerBrowser.AddSelection"))
-				]
-			];
-
-		ActionDockPanel->AddHeaderToolbarItem(EToolbar::Left, AddNewActionButton);
-		ActionDockPanel->AddHeaderToolbarItem(EToolbar::Right, AddAllSelectedActionsButton);
-		ActionDockPanel->AddHeaderToolbarItem(EToolbar::Right, AddAllActionsButton);
-
-		TSharedRef<SRCMajorPanel> ActionsPanel = SNew(SRCMajorPanel)
-			.EnableFooter(false)
-			.EnableHeader(false)
-			.ChildOrientation(Orient_Vertical);
-
-		if (InBehaviourItem->HasBehaviourDetailsWidget())
-		{
-			// Header Dock Panel
-			TSharedPtr<SRCMinorPanel> BehaviourDetailsPanel = SNew(SRCMinorPanel)
-				.EnableHeader(false)
-				[
-					SAssignNew(BehaviourDetailsWidget, SRCBehaviourDetails, SharedThis(this), InBehaviourItem.ToSharedRef())
-				];
-
-			// Panel size of zero forces use of "SizeToContent" ensuring that each Behaviour only takes up as much space as necessary
-			ActionsPanel->AddPanel(BehaviourDetailsPanel.ToSharedRef(), 0.f);
-		}
-		ActionsPanel->AddPanel(ActionDockPanel.ToSharedRef(), 0.5f);
-
-		WrappedBoxWidget->SetContent(ActionsPanel);
-
 		const bool bIsBehaviourEnabled = InBehaviourItem->IsBehaviourEnabled();
 		InBehaviourItem->RefreshIsBehaviourEnabled(bIsBehaviourEnabled);
 		RefreshIsBehaviourEnabled(bIsBehaviourEnabled);
 	}
-	else
-	{
-		WrappedBoxWidget->SetContent(CreateNoneSelectedWidget());
-	}
 }
 
 FReply SRCActionPanel::OnClickOverrideBlueprintButton()
@@ -342,6 +251,185 @@
 	return AddNewActionMenuWidget.ToSharedRef();
 }
 
+void SRCActionPanel::CreateActionPanelWidgets()
+{
+	// Action Dock Panel
+	if (!ActionDockPanel.IsValid())
+	{
+		ActionDockPanel = SNew(SRCMinorPanel)
+			.HeaderLabel(LOCTEXT("ActionsLabel", "Actions"))
+			.EnableHeader(true)
+			.EnableFooter(false);
+	}
+
+	bool bIsAnyToolbarItemCreated = false;
+
+	// Add new action combo button
+	if (!AddNewActionButton.IsValid())
+	{
+		bIsAnyToolbarItemCreated = true;
+
+		AddNewActionButton = SNew(SComboButton)
+			.AddMetaData<FTagMetaData>(FTagMetaData(TEXT("Add Action")))
+			.HAlign(HAlign_Center)
+			.VAlign(VAlign_Center)
+			.ButtonStyle(&RCPanelStyle->FlatButtonStyle)
+			.ForegroundColor(FSlateColor::UseForeground())
+			.CollapseMenuOnParentFocus(true)
+			.HasDownArrow(false)
+			.ContentPadding(FMargin(4.f, 2.f))
+			.ButtonContent()
+			[
+				SNew(SBox)
+				.WidthOverride(RCPanelStyle->IconSize.X)
+				.HeightOverride(RCPanelStyle->IconSize.Y)
+				[
+					SNew(SImage)
+					.ColorAndOpacity(FSlateColor::UseForeground())
+					.Image(FAppStyle::GetBrush("Icons.PlusCircle"))
+				]
+			]
+			.MenuContent()
+			[
+				GetActionMenuContentWidget()
+			];
+	}
+
+	// Add all button
+	if (!AddAllActionsButton.IsValid())
+	{
+		bIsAnyToolbarItemCreated = true;
+
+		AddAllActionsButton = SNew(SButton)
+			.AddMetaData<FTagMetaData>(FTagMetaData(TEXT("Add All Actions")))
+			.HAlign(HAlign_Center)
+			.VAlign(VAlign_Center)
+			.ForegroundColor(FSlateColor::UseForeground())
+			.ButtonStyle(&RCPanelStyle->FlatButtonStyle)
+			.ToolTipText(LOCTEXT("AddAllToolTip", "Adds all the available actions."))
+			.OnClicked(this, &SRCActionPanel::OnAddAllFields)
+			.Visibility(this, &SRCActionPanel::HandleAddAllButtonVisibility)
+			[
+				SNew(SBox)
+				.WidthOverride(RCPanelStyle->IconSize.X)
+				.HeightOverride(RCPanelStyle->IconSize.Y)
+				[
+					SNew(SImage)
+					.ColorAndOpacity(FSlateColor::UseForeground())
+					.Image(FAppStyle::GetBrush("Icons.Duplicate"))
+				]
+			];
+	}
+
+	// Add all selected button
+	if (!AddAllSelectedActionsButton.IsValid())
+	{
+		bIsAnyToolbarItemCreated = true;
+
+		AddAllSelectedActionsButton = SNew(SButton)
+			.AddMetaData<FTagMetaData>(FTagMetaData(TEXT("Add All Selected Fields")))
+			.HAlign(HAlign_Center)
+			.VAlign(VAlign_Center)
+			.ForegroundColor(FSlateColor::UseForeground())
+			.ButtonStyle(&RCPanelStyle->FlatButtonStyle)
+			.ToolTipText(LOCTEXT("RCAddAllSelectedToolTip", "Adds all the selected fields."))
+			.OnClicked(this, &SRCActionPanel::OnAddAllSelectedFields)
+			.Visibility(this, &SRCActionPanel::HandleAddAllButtonVisibility)
+			[
+				SNew(SBox)
+				.WidthOverride(RCPanelStyle->IconSize.X)
+				.HeightOverride(RCPanelStyle->IconSize.Y)
+				[
+					SNew(SImage)
+					.ColorAndOpacity(FSlateColor::UseForeground())
+					.Image(FAppStyle::GetBrush("DataLayerBrowser.AddSelection"))
+				]
+			];
+	}
+
+	// Adding header toolbar items
+	if (bIsAnyToolbarItemCreated)
+	{
+		ActionDockPanel->ClearHeaderToolbarItems();
+		ActionDockPanel->AddHeaderToolbarItem(EToolbar::Left, AddNewActionButton.ToSharedRef());
+		ActionDockPanel->AddHeaderToolbarItem(EToolbar::Right, AddAllSelectedActionsButton.ToSharedRef());
+		ActionDockPanel->AddHeaderToolbarItem(EToolbar::Right, AddAllActionsButton.ToSharedRef());
+	}
+
+	// Behavior Details widget
+	if (!BehaviourDetailsWidget.IsValid())
+	{
+		BehaviourDetailsWidget = SNew(SRCBehaviourDetails, SharedThis(this), nullptr);
+	}
+
+	// Behavior Details panel
+	if (!BehaviorDetailsPanel.IsValid())
+	{
+		BehaviorDetailsPanel = SNew(SRCMinorPanel)
+			.EnableHeader(false)
+			.Visibility_Lambda([this] ()
+			{
+				if (const TSharedPtr<FRCBehaviourModel>& Behavior = SelectedBehaviourItemWeakPtr.Pin())
+				{
+					return Behavior->HasBehaviourDetailsWidget()? EVisibility::Visible : EVisibility::Collapsed;
+				}
+				return EVisibility::Collapsed;
+			})
+		[
+			BehaviourDetailsWidget.ToSharedRef()
+		];
+	}
+
+	// Action Major Panel
+	if (!ActionsPanel.IsValid())
+	{
+		ActionsPanel = SNew(SRCMajorPanel)
+			.EnableFooter(false)
+			.EnableHeader(false)
+			.ChildOrientation(Orient_Vertical);
+
+		if (BehaviorDetailsPanel.IsValid())
+		{
+			// Panel size of zero forces use of "SizeToContent" ensuring that each Behaviour only takes up as much space as necessary
+			ActionsPanel->AddPanel(BehaviorDetailsPanel.ToSharedRef(), 0.f);
+		}
+		ActionsPanel->AddPanel(ActionDockPanel.ToSharedRef(), 0.5f);
+	}
+
+	// Wrapped Box
+	if (!WrappedBoxWidget.IsValid())
+	{
+		WrappedBoxWidget = SNew(SBox)
+		[
+			SNew(SWidgetSwitcher)
+			.WidgetIndex_Lambda([this] () { return SelectedBehaviourItemWeakPtr.IsValid()? 0 : 1; })
+
+			// Index 0 is for valid behavior
+			+ SWidgetSwitcher::Slot()
+			[
+				ActionsPanel.ToSharedRef()
+			]
+
+			// Index 1 is for invalid behavior
+			+ SWidgetSwitcher::Slot()
+			[
+				GetNoneSelectedWidget()
+			]
+		];
+	}
+}
+
+bool SRCActionPanel::AreActionPanelWidgetsValid() const
+{
+	return WrappedBoxWidget.IsValid() &&
+			ActionDockPanel.IsValid() &&
+			AddNewActionButton.IsValid() &&
+			AddAllActionsButton.IsValid() &&
+			AddAllSelectedActionsButton.IsValid() &&
+			BehaviorDetailsPanel.IsValid() &&
+			ActionsPanel.IsValid();
+}
+
 URCAction* SRCActionPanel::AddAction(const TSharedRef<const FRemoteControlField> InRemoteControlField)
 {
 	if (const TSharedPtr<FRCBehaviourModel> BehaviourItem = SelectedBehaviourItemWeakPtr.Pin())
diff --git a/Engine/Plugins/VirtualProduction/RemoteControl/Source/RemoteControlUI/Private/UI/Action/SRCActionPanel.h b/Engine/Plugins/VirtualProduction/RemoteControl/Source/RemoteControlUI/Private/UI/Action/SRCActionPanel.h
--- a/Engine/Plugins/VirtualProduction/RemoteControl/Source/RemoteControlUI/Private/UI/Action/SRCActionPanel.h
+++ b/Engine/Plugins/VirtualProduction/RemoteControl/Source/RemoteControlUI/Private/UI/Action/SRCActionPanel.h
@@ -8,8 +8,11 @@ class FRCActionModel;
 class FRCBehaviourModel;
 class SBox;
 class SCheckBox;
+class SComboButton;
 class SRCBehaviourDetails;
 class SRCLogicPanelListBase;
+class SRCMajorPanel;
+class SRCMinorPanel;
 class SRemoteControlPanel;
 class URCAction;
 class URCBehaviour;
@@ -125,6 +128,12 @@ private:
 	 */
 	TSharedRef<SWidget> GetActionMenuContentWidget();
 
+	/** Check the validity of the panel widgets and if not valid will create them */
+	void CreateActionPanelWidgets();
+
+	/** Return whether or not all panel widgets are valid */
+	bool AreActionPanelWidgetsValid() const;
+
 	/** Handles click event for Add Action button*/
 	void OnAddActionClicked(TSharedPtr<FRemoteControlField> InRemoteControlField);
 
@@ -151,7 +160,7 @@ private:
 
 private:
 	/** Helper widget for behavior details. */
-	static TSharedRef<SBox> CreateNoneSelectedWidget();
+	static TSharedRef<SBox> GetNoneSelectedWidget();
 
 	void DuplicateAction(URCAction* InAction);
 
@@ -175,6 +184,24 @@ private:
 	/** Cached menu widget for Add New Action */
 	TSharedPtr<SWidget> AddNewActionMenuWidget;
 
+	/** Action Minor Panel */
+	TSharedPtr<SRCMinorPanel> ActionDockPanel;
+
+	/** ComboButton to add new action */
+	TSharedPtr<SComboButton> AddNewActionButton;
+
+	/** Add all actions button */
+	TSharedPtr<SWidget> AddAllActionsButton;
+
+	/** Add all selected actions button */
+	TSharedPtr<SWidget> AddAllSelectedActionsButton;
+
+	/** Action Major Panel */
+	TSharedPtr<SRCMajorPanel> ActionsPanel;
+
+	/** Behavior Details Minor Panel */
+	TSharedPtr<SRCMinorPanel> BehaviorDetailsPanel;
+
 	/** Whether the Add Actions menu list is outdated and needs to be refreshed */
 	bool bAddActionMenuNeedsRefresh = false;
 };
diff --git a/Engine/Plugins/VirtualProduction/RemoteControl/Source/RemoteControlUI/Private/UI/Action/SRCBehaviourDetails.cpp b/Engine/Plugins/VirtualProduction/RemoteControl/Source/RemoteControlUI/Private/UI/Action/SRCBehaviourDetails.cpp
--- a/Engine/Plugins/VirtualProduction/RemoteControl/Source/RemoteControlUI/Private/UI/Action/SRCBehaviourDetails.cpp
+++ b/Engine/Plugins/VirtualProduction/RemoteControl/Source/RemoteControlUI/Private/UI/Action/SRCBehaviourDetails.cpp
@@ -18,66 +18,82 @@
 #include "UI/RemoteControlPanelStyle.h"
 #include "UI/SRemoteControlPanel.h"
 #include "Widgets/Input/SButton.h"
+#include "Widgets/Layout/SBox.h"
 #include "Widgets/Layout/SSpacer.h"
 
 #define LOCTEXT_NAMESPACE "SRCBehaviourDetails"
 
-void SRCBehaviourDetails::Construct(const FArguments& InArgs, TSharedRef<SRCActionPanel> InActionPanel, TSharedRef<FRCBehaviourModel> InBehaviourItem)
+void SRCBehaviourDetails::Construct(const FArguments& InArgs, TSharedRef<SRCActionPanel> InActionPanel, TSharedPtr<FRCBehaviourModel> InBehaviourItem)
 {
 	RCPanelStyle = &FRemoteControlPanelStyle::Get()->GetWidgetStyle<FRCPanelStyle>("RemoteControlPanel.MinorPanel");
 
 	ActionPanelWeakPtr = InActionPanel;
 	BehaviourItemWeakPtr = InBehaviourItem;
 
-	URCBehaviour* Behaviour = BehaviourItemWeakPtr.Pin()->GetBehaviour();
-	const FText BehaviourDisplayName = Behaviour->GetDisplayName();
-	const FText BehaviourDescription = Behaviour->GetBehaviorDescription();
+	bool bIsEnabled = false;
 
-	const FSlateFontInfo& FontBehaviorDesc = FRemoteControlPanelStyle::Get()->GetFontStyle("RemoteControlPanel.Behaviours.BehaviorDescription");
+	BehaviourDetailsBox = SNew(SBox);
 
-	FLinearColor TypeColor = FLinearColor::White;
-	FString TypeDisplayName = (LOCTEXT("NoneDisplayName", "None")).ToString();
-	if (URCController* Controller = Behaviour->ControllerWeakPtr.Get())
+	if (InBehaviourItem.IsValid())
 	{
-		if (FProperty* Property = Behaviour->ControllerWeakPtr->GetProperty())
-		{
-			TypeColor = UE::RCUIHelpers::GetFieldClassTypeColor(Property);
-			TypeDisplayName = FName::NameToDisplayString(UE::RCUIHelpers::GetFieldClassDisplayName(Property).ToString(), false);
-		}
+		BehaviourDetailsWidget = InBehaviourItem->GetBehaviourDetailsWidget();
+		bIsEnabled = InBehaviourItem->IsBehaviourEnabled();
+	}
+	else
+	{
+		BehaviourDetailsWidget = SNullWidget::NullWidget;
 	}
-
-	BehaviourDetailsWidget = InBehaviourItem->GetBehaviourDetailsWidget();
-
-	const bool bIsChecked = InBehaviourItem->IsBehaviourEnabled();
 	
+	BehaviourDetailsBox->SetContent(BehaviourDetailsWidget.ToSharedRef());
+
 	ChildSlot
 		.Padding(RCPanelStyle->PanelPadding)
 		[
-			SNew(SVerticalBox)
-
-			// Behaviour Specific Details Panel
-			+ SVerticalBox::Slot()
-			.Padding(8.f, 4.f)
-			.AutoHeight()
+			SNew(SWidgetSwitcher)
+			.WidgetIndex_Lambda([this] ()
+			{
+				if (const TSharedPtr<FRCBehaviourModel>& Behavior = BehaviourItemWeakPtr.Pin())
+				{
+					return Behavior->HasBehaviourDetailsWidget()? 0 : 1;
+				}
+				return 1;
+			})
+
+			// Index 0 is for valid behavior
+			+ SWidgetSwitcher::Slot()
 			[
-				BehaviourDetailsWidget.ToSharedRef()
+				SNew(SVerticalBox)
+
+				// Behaviour Specific Details Panel
+				+ SVerticalBox::Slot()
+				.Padding(8.f, 4.f)
+				.AutoHeight()
+				[
+					BehaviourDetailsBox.ToSharedRef()
+				]
+
+				// Spacer to fill the gap.
+				+ SVerticalBox::Slot()
+				.Padding(0)
+				.FillHeight(1.f)
+				[
+					SNew(SSpacer)
+				]
 			]
 
-			// Spacer to fill the gap.
-			+ SVerticalBox::Slot()
-			.Padding(0)
-			.FillHeight(1.f)
+			// Index 1 is for invalid behavior
+			+ SWidgetSwitcher::Slot()
 			[
-				SNew(SSpacer)
+				SNullWidget::NullWidget
 			]
 		];
 
-	RefreshIsBehaviourEnabled(bIsChecked);
+	RefreshIsBehaviourEnabled(bIsEnabled);
 }
 
-void SRCBehaviourDetails::SetIsBehaviourEnabled(const bool bIsEnabled)
+void SRCBehaviourDetails::SetIsBehaviourEnabled(const bool bIsEnabled) const
 {
-	if (TSharedPtr<FRCBehaviourModel> BehaviourItem = BehaviourItemWeakPtr.Pin())
+	if (const TSharedPtr<FRCBehaviourModel> BehaviourItem = BehaviourItemWeakPtr.Pin())
 	{
 		BehaviourItem->SetIsBehaviourEnabled(bIsEnabled);
 
@@ -85,7 +101,27 @@ void SRCBehaviourDetails::SetIsBehaviourEnabled(const bool bIsEnabled)
 	}
 }
 
-void SRCBehaviourDetails::RefreshIsBehaviourEnabled(const bool bIsEnabled)
+void SRCBehaviourDetails::SetNewBehavior(const TSharedPtr<FRCBehaviourModel>& InBehaviourItem)
+{
+	// TODO: do not update if the Behavior type is the same, this needs further change
+	const TSharedPtr<FRCBehaviourModel>& CurrentBehavior = GetCurrentBehavior();
+	if (CurrentBehavior != InBehaviourItem)
+	{
+		BehaviourItemWeakPtr = InBehaviourItem;
+		Refresh();
+	}
+}
+
+TSharedPtr<FRCBehaviourModel> SRCBehaviourDetails::GetCurrentBehavior() const
+{
+	if (BehaviourItemWeakPtr.IsValid())
+	{
+		return BehaviourItemWeakPtr.Pin();
+	}
+	return nullptr;
+}
+
+void SRCBehaviourDetails::RefreshIsBehaviourEnabled(const bool bIsEnabled) const
 {
 	if(BehaviourDetailsWidget && BehaviourTitleWidget)
 	{
@@ -93,10 +129,36 @@ void SRCBehaviourDetails::RefreshIsBehaviourEnabled(const bool bIsEnabled)
 		BehaviourTitleWidget->SetEnabled(bIsEnabled);
 	}
 
-	if (TSharedPtr<SRCActionPanel> ActionPanel = ActionPanelWeakPtr.Pin())
+	if (const TSharedPtr<SRCActionPanel> ActionPanel = ActionPanelWeakPtr.Pin())
 	{
 		ActionPanel->RefreshIsBehaviourEnabled(bIsEnabled);
 	}
 }
 
-#undef LOCTEXT_NAMESPACE
\ No newline at end of file
+void SRCBehaviourDetails::Refresh()
+{
+	const TSharedPtr<FRCBehaviourModel>& CurrentBehavior = GetCurrentBehavior();
+
+	if (CurrentBehavior.IsValid() && CurrentBehavior->HasBehaviourDetailsWidget())
+	{
+		BehaviourDetailsWidget = CurrentBehavior->GetBehaviourDetailsWidget();
+		const bool bIsEnabled = CurrentBehavior->IsBehaviourEnabled();
+
+		if (BehaviourDetailsWidget.IsValid() && BehaviourDetailsBox.IsValid())
+		{
+			BehaviourDetailsBox->SetContent(BehaviourDetailsWidget.ToSharedRef());
+		}
+
+		RefreshIsBehaviourEnabled(bIsEnabled);
+	}
+	else
+	{
+		BehaviourDetailsWidget = SNullWidget::NullWidget;
+		if (BehaviourDetailsBox.IsValid())
+		{
+			BehaviourDetailsBox->SetContent(BehaviourDetailsWidget.ToSharedRef());
+		}
+	}
+}
+
+#undef LOCTEXT_NAMESPACE
diff --git a/Engine/Plugins/VirtualProduction/RemoteControl/Source/RemoteControlUI/Private/UI/Action/SRCBehaviourDetails.h b/Engine/Plugins/VirtualProduction/RemoteControl/Source/RemoteControlUI/Private/UI/Action/SRCBehaviourDetails.h
--- a/Engine/Plugins/VirtualProduction/RemoteControl/Source/RemoteControlUI/Private/UI/Action/SRCBehaviourDetails.h
+++ b/Engine/Plugins/VirtualProduction/RemoteControl/Source/RemoteControlUI/Private/UI/Action/SRCBehaviourDetails.h
@@ -5,6 +5,7 @@
 #include "Widgets/SCompoundWidget.h"
 
 struct FRCPanelStyle;
+class SBox;
 class SRCActionPanel;
 class STextBlock;
 class FRCBehaviourModel;
@@ -22,14 +23,23 @@ public:
 	SLATE_END_ARGS()
 
 	/** Constructs this widget with InArgs */
-	void Construct(const FArguments& InArgs, TSharedRef<SRCActionPanel> InActionPanel, TSharedRef<FRCBehaviourModel> InBehaviourItem);
+	void Construct(const FArguments& InArgs, TSharedRef<SRCActionPanel> InActionPanel, TSharedPtr<FRCBehaviourModel> InBehaviourItem);
 
 	/** Set the Enabled state of our parent Behaviour */
-	void SetIsBehaviourEnabled(const bool bIsEnabled);
+	void SetIsBehaviourEnabled(const bool bIsEnabled) const;
+
+	/** Set a new behavior for this Details widget */
+	void SetNewBehavior(const TSharedPtr<FRCBehaviourModel>& InBehaviourItem);
+
+	/** Get the current Behavior used for this details */
+	TSharedPtr<FRCBehaviourModel> GetCurrentBehavior() const;
 
 private:
 	/** Refreshes the UI widgets enabled state depending on whether the parent behaviour is currently enabled */
-	void RefreshIsBehaviourEnabled(const bool bIsEnabled);
+	void RefreshIsBehaviourEnabled(const bool bIsEnabled) const;
+
+	/** Refresh the behaviour widgets */
+	void Refresh();
 
 	/** The parent Action panel UI widget holding this Header*/
 	TWeakPtr<SRCActionPanel> ActionPanelWeakPtr;
@@ -40,6 +50,9 @@ private:
 	/** Panel Style reference. */
 	const FRCPanelStyle* RCPanelStyle;
 
+	/** Behaviour Box containing the details widget */
+	TSharedPtr<SBox> BehaviourDetailsBox;
+
 	/** Behaviour specific details panel*/
 	TSharedPtr<SWidget> BehaviourDetailsWidget;
 
diff --git a/Engine/Plugins/VirtualProduction/RemoteControl/Source/RemoteControlUI/Private/UI/Behaviour/Builtin/Path/RCBehaviourSetAssetByPathModel.cpp b/Engine/Plugins/VirtualProduction/RemoteControl/Source/RemoteControlUI/Private/UI/Behaviour/Builtin/Path/RCBehaviourSetAssetByPathModel.cpp
--- a/Engine/Plugins/VirtualProduction/RemoteControl/Source/RemoteControlUI/Private/UI/Behaviour/Builtin/Path/RCBehaviourSetAssetByPathModel.cpp
+++ b/Engine/Plugins/VirtualProduction/RemoteControl/Source/RemoteControlUI/Private/UI/Behaviour/Builtin/Path/RCBehaviourSetAssetByPathModel.cpp
@@ -10,8 +10,10 @@
 #include "RCVirtualProperty.h"
 #include "RCVirtualPropertyContainer.h"
 #include "RemoteControlPreset.h"
-#include "SRCBehaviourSetAssetByPath.h"
 #include "Selection.h"
+#include "SRCBehaviourSetAssetByPath.h"
+#include "UI/Controller/SRCControllerPanel.h"
+#include "UI/SRemoteControlPanel.h"
 #include "Widgets/Images/SImage.h"
 #include "Widgets/Input/SButton.h"
 #include "Widgets/Input/SComboBox.h"
@@ -21,9 +23,9 @@
 
 #define LOCTEXT_NAMESPACE "FRCSetAssetByPathBehaviourModel"
 
-FRCSetAssetByPathBehaviourModel::FRCSetAssetByPathBehaviourModel(URCSetAssetByPathBehaviour* SetAssetByPathBehaviour)
-	: FRCBehaviourModel(SetAssetByPathBehaviour)
-	, SetAssetByPathBehaviourWeakPtr(SetAssetByPathBehaviour)
+FRCSetAssetByPathBehaviourModel::FRCSetAssetByPathBehaviourModel(URCSetAssetByPathBehaviour* InSetAssetByPathBehaviour, const TSharedPtr<SRemoteControlPanel> InRemoteControlPanel)
+	: FRCBehaviourModel(InSetAssetByPathBehaviour, InRemoteControlPanel)
+	, SetAssetByPathBehaviourWeakPtr(InSetAssetByPathBehaviour)
 {
 	FPropertyRowGeneratorArgs ArgsRowGenerator;
 	ArgsRowGenerator.bShouldShowHiddenProperties = true;
@@ -32,9 +34,9 @@ FRCSetAssetByPathBehaviourModel::FRCSetAssetByPathBehaviourModel(URCSetAssetByPa
 	const FPropertyRowGeneratorArgs ArgsRowGeneratorArray;
 	PropertyRowGeneratorArray = FModuleManager::GetModuleChecked<FPropertyEditorModule>("PropertyEditor").CreatePropertyRowGenerator(ArgsRowGeneratorArray);
 
-	if (SetAssetByPathBehaviour)
+	if (InSetAssetByPathBehaviour)
 	{
-		PropertyRowGenerator->SetStructure(SetAssetByPathBehaviour->PropertyInContainer->CreateStructOnScope());
+		PropertyRowGenerator->SetStructure(InSetAssetByPathBehaviour->PropertyInContainer->CreateStructOnScope());
 		DetailTreeNodeWeakPtr.Empty();
 		for (const TSharedRef<IDetailTreeNode>& CategoryNode : PropertyRowGenerator->GetRootTreeNodes())
 		{
@@ -47,18 +49,10 @@ FRCSetAssetByPathBehaviourModel::FRCSetAssetByPathBehaviourModel(URCSetAssetByPa
 		}
 
 		// Secondary TArray Struct
-		PropertyRowGeneratorArray->SetStructure(MakeShareable(new FStructOnScope(FRCSetAssetPath::StaticStruct(), (uint8*) &SetAssetByPathBehaviour->PathStruct)));
-		PropertyRowGeneratorArray->OnFinishedChangingProperties().AddLambda([this](const FPropertyChangedEvent& InEvent)
+		PropertyRowGeneratorArray->SetStructure(MakeShareable(new FStructOnScope(FRCSetAssetPath::StaticStruct(), (uint8*) &InSetAssetByPathBehaviour->PathStruct)));
+		PropertyRowGeneratorArray->OnFinishedChangingProperties().AddLambda([this] (const FPropertyChangedEvent& InEvent)
 		{
-			if (InEvent.ChangeType == EPropertyChangeType::ValueSet)
-			{
-				// Just refresh the preview without reconstructing the widget if we just set values
-				RefreshPreview();
-			}
-			else
-			{
-				RefreshPathAndPreview();
-			}
+			OnAssetPathFinishedChangingProperties(InEvent);
 		});
 		PropertyRowGeneratorArray->OnRowsRefreshed().AddLambda([this]()
 		{
@@ -351,6 +345,12 @@ TSharedRef<SWidget> FRCSetAssetByPathBehaviourModel::GetSelectorWidget(TWeakPtr<
 		];
 }
 
+void FRCSetAssetByPathBehaviourModel::NotifyControllerValueChanged(TSharedPtr<FRCControllerModel> InControllerModel)
+{
+	// Update the preview when a controller change to always get the correct path
+	RefreshPreview();
+}
+
 FText FRCSetAssetByPathBehaviourModel::GetSelectedEntityText() const
 {
 	FText ReturnText = FText::FromString(*FString("Nothing selected yet"));
@@ -365,4 +365,56 @@ FText FRCSetAssetByPathBehaviourModel::GetSelectedEntityText() const
 	return ReturnText;
 }
 
+void FRCSetAssetByPathBehaviourModel::OnAssetPathFinishedChangingProperties(const FPropertyChangedEvent& InEvent)
+{
+	if (InEvent.ChangeType == EPropertyChangeType::ValueSet)
+	{
+		// if the actual PropertyName is the AssetPath than it means that we want to add a controller since we force this in the Customization
+		if (InEvent.GetPropertyName() == GET_MEMBER_NAME_CHECKED(FRCSetAssetPath, AssetPath))
+		{
+			const int32 Index = InEvent.GetArrayIndex(InEvent.GetMemberPropertyName().ToString());
+			URCSetAssetByPathBehaviour* SetAssetByPathBehaviour = SetAssetByPathBehaviourWeakPtr.Get();
+			if (!SetAssetByPathBehaviour || Index == INDEX_NONE)
+			{
+				return;
+			}
+			TArray<FRCAssetPathElement>& AssetPath = SetAssetByPathBehaviour->PathStruct.AssetPath;
+			if (!AssetPath.IsValidIndex(Index))
+			{
+				return;
+			}
+			FRCAssetPathElement& CurrentElement = AssetPath[Index];
+			if (const TSharedPtr<SRemoteControlPanel>& RCPanel = GetRemoteControlPanel())
+			{
+				if (URemoteControlPreset* Preset = RCPanel->GetPreset())
+				{
+					URCVirtualPropertyInContainer* NewController = Preset->AddController(URCController::StaticClass(), EPropertyBagPropertyType::String);
+					if (NewController)
+					{
+						NewController->DisplayIndex = RCPanel->NumControllerItems();
+						FName NewControllerName = TEXT("DefaultInput");
+						if (!CurrentElement.Path.IsEmpty())
+						{
+							NewControllerName = FName(CurrentElement.Path);
+						}
+						const FName NewUniqueName = Preset->SetControllerDisplayName(NewController->Id, NewControllerName);
+						if (CurrentElement.Path != NewUniqueName.ToString())
+						{
+							// Assign to the path the actual new name of the controller
+							CurrentElement.Path = NewUniqueName.ToString();
+						}
+						RCPanel->OnControllerAdded.Broadcast(NewController->PropertyName);
+					}
+				}
+			}
+		}
+		// Just refresh the preview without reconstructing the widget if we just set values
+		RefreshPreview();
+	}
+	else
+	{
+		RefreshPathAndPreview();
+	}
+}
+
 #undef LOCTEXT_NAMESPACE
diff --git a/Engine/Plugins/VirtualProduction/RemoteControl/Source/RemoteControlUI/Private/UI/Behaviour/Builtin/Path/RCBehaviourSetAssetByPathModel.h b/Engine/Plugins/VirtualProduction/RemoteControl/Source/RemoteControlUI/Private/UI/Behaviour/Builtin/Path/RCBehaviourSetAssetByPathModel.h
--- a/Engine/Plugins/VirtualProduction/RemoteControl/Source/RemoteControlUI/Private/UI/Behaviour/Builtin/Path/RCBehaviourSetAssetByPathModel.h
+++ b/Engine/Plugins/VirtualProduction/RemoteControl/Source/RemoteControlUI/Private/UI/Behaviour/Builtin/Path/RCBehaviourSetAssetByPathModel.h
@@ -24,7 +24,7 @@ namespace SetAssetByPathModelHelper
 class FRCSetAssetByPathBehaviourModel : public FRCBehaviourModel
 {
 public:
-	FRCSetAssetByPathBehaviourModel(URCSetAssetByPathBehaviour* SetAssetByPathBehaviour);
+	FRCSetAssetByPathBehaviourModel(URCSetAssetByPathBehaviour* SetAssetByPathBehaviour, const TSharedPtr<SRemoteControlPanel> InRemoteControlPanel);
 
 	/** Returns true if this behaviour have a details widget or false if not*/
 	virtual bool HasBehaviourDetailsWidget() override;
@@ -42,10 +42,17 @@ public:
 	/** Builds the Combobox which the Behaviour will use to select and and apply it's Asset Setting on. */
 	TSharedRef<SWidget> GetSelectorWidget(TWeakPtr<const FRemoteControlEntity> InInitialSelected);
 
+	/** Called when a controller value changed */
+	virtual void NotifyControllerValueChanged(TSharedPtr<FRCControllerModel> InControllerModel) override;
+
 protected:
 	/** Gets the currently selected Entities Name, or a failure Text in case it fails. */
 	FText GetSelectedEntityText() const;
-	
+
+private:
+	/** Callback called when the AssetPath properties change to either refresh or reconstruct it updating the path values */
+	void OnAssetPathFinishedChangingProperties(const FPropertyChangedEvent& InEvent);
+
 private:
 	/** The SetAssetByPath Behaviour associated with this Model */
 	TWeakObjectPtr<URCSetAssetByPathBehaviour> SetAssetByPathBehaviourWeakPtr;
diff --git a/Engine/Plugins/VirtualProduction/RemoteControl/Source/RemoteControlUI/Private/UI/Behaviour/RCBehaviourModel.h b/Engine/Plugins/VirtualProduction/RemoteControl/Source/RemoteControlUI/Private/UI/Behaviour/RCBehaviourModel.h
--- a/Engine/Plugins/VirtualProduction/RemoteControl/Source/RemoteControlUI/Private/UI/Behaviour/RCBehaviourModel.h
+++ b/Engine/Plugins/VirtualProduction/RemoteControl/Source/RemoteControlUI/Private/UI/Behaviour/RCBehaviourModel.h
@@ -30,52 +30,55 @@ public:
 	/** Add a Logic Action using as an identity action. */
 	virtual URCAction* AddAction(FName InFieldId);
 
-	/** Add a Logic Action using a remote control field as input */
+	/** Add a Logic Action using a remote control field as input. */
 	virtual URCAction* AddAction(const TSharedRef<const FRemoteControlField> InRemoteControlField);
 
 	/** The widget to be rendered for this Property in the Behaviour panel
-	* Currently displays Behaviour Name metadata in the Behaviours Panel List
-	*/
+	 * Currently displays Behaviour Name metadata in the Behaviours Panel List
+	 */
 	virtual TSharedRef<SWidget> GetWidget() const override;
 
-	/** Returns true if this behaviour have a details widget or false if not*/
+	/** Returns true if this behaviour have a details widget or false if not. */
 	virtual bool HasBehaviourDetailsWidget();
 
-	/** Builds a Behaviour specific widget that child Behaviour classes can implement as required*/
+	/** Builds a Behaviour specific widget that child Behaviour classes can implement as required. */
 	virtual TSharedRef<SWidget> GetBehaviourDetailsWidget();
 
-	/** Returns the Behaviour (Data model) associated with us*/
+	/** Returns the Behaviour (Data model) associated with us. */
 	URCBehaviour* GetBehaviour() const;
 
-	/** Handling for user action to override this behaviour via new Blueprint class */
+	/** Handling for user action to override this behaviour via new Blueprint class. */
 	void OnOverrideBlueprint() const;
 
-	/** Whether the underlying Behaviour is currently enabled */
+	/** Whether the underlying Behaviour is currently enabled. */
 	bool IsBehaviourEnabled() const;
 
-	/** Set the Enabled state of our underlying Behaviour */
+	/** Set the Enabled state of our underlying Behaviour. */
 	void SetIsBehaviourEnabled(const bool bIsEnabled);
 
-	/** Refreshes the UI widgets enabled state depending on whether the parent behaviour is currently enabled */
+	/** Refreshes the UI widgets enabled state depending on whether the parent behaviour is currently enabled. */
 	void RefreshIsBehaviourEnabled(const bool bIsEnabled);
 
 	/** The Actions List widget to be used for this Behaviour.
-	* The default actions table is used if a behaviour doesn't specify it.
-	*/
+	 * The default actions table is used if a behaviour doesn't specify it.
+	 */
 	virtual TSharedPtr<SRCLogicPanelListBase> GetActionsListWidget(TSharedRef<SRCActionPanel> InActionPanel);
 
-	/** Whether or not the behaviour support PropertyId action */
+	/** Whether or not the behaviour support PropertyId action. */
 	bool SupportPropertyId() const;
 
+	/** Called when a controller value changed */
+	virtual void NotifyControllerValueChanged(TSharedPtr<class FRCControllerModel> InControllerModel) {}
+
 protected:
-	/** Invoked after an action has been added for this Behaviour in the actions panel */
+	/** Invoked after an action has been added for this Behaviour in the actions panel. */
 	virtual void OnActionAdded(URCAction* Action) {}
 
 private:
-	/** The Behaviour (Data model) associated with us*/
+	/** The Behaviour (Data model) associated with us. */
 	TWeakObjectPtr<URCBehaviour> BehaviourWeakPtr;
 
-	/** Text block widget for representing the Behaviour's title */
+	/** Text block widget for representing the Behaviour's title. */
 	TSharedPtr<STextBlock> BehaviourTitleText;
 
 	/** Panel Style reference. */
diff --git a/Engine/Plugins/VirtualProduction/RemoteControl/Source/RemoteControlUI/Private/UI/Behaviour/SRCBehaviourPanel.cpp b/Engine/Plugins/VirtualProduction/RemoteControl/Source/RemoteControlUI/Private/UI/Behaviour/SRCBehaviourPanel.cpp
--- a/Engine/Plugins/VirtualProduction/RemoteControl/Source/RemoteControlUI/Private/UI/Behaviour/SRCBehaviourPanel.cpp
+++ b/Engine/Plugins/VirtualProduction/RemoteControl/Source/RemoteControlUI/Private/UI/Behaviour/SRCBehaviourPanel.cpp
@@ -27,7 +27,7 @@
 
 #define LOCTEXT_NAMESPACE "SRCBehaviourPanel"
 
-TSharedRef<SBox> SRCBehaviourPanel::CreateNoneSelectedWidget()
+TSharedRef<SBox> SRCBehaviourPanel::GetNoneSelectedWidget()
 {
 	return SNew(SBox)
 		.Padding(10.f)
@@ -48,9 +48,8 @@ void SRCBehaviourPanel::Construct(const FArguments& InArgs, const TSharedRef<SRe
 
 	RCPanelStyle = &FRemoteControlPanelStyle::Get()->GetWidgetStyle<FRCPanelStyle>("RemoteControlPanel.MinorPanel");
 
-	WrappedBoxWidget = SNew(SBox);
 	UpdateWrappedWidget();
-	
+
 	ChildSlot
 		.Padding(RCPanelStyle->PanelPadding)
 		[
@@ -59,65 +58,61 @@ void SRCBehaviourPanel::Construct(const FArguments& InArgs, const TSharedRef<SRe
 
 	// Register delegates
 	InPanel->OnControllerSelectionChanged.AddSP(this, &SRCBehaviourPanel::OnControllerSelectionChanged);
+	InPanel->OnControllerValueChangedDelegate.AddSP(this, &SRCBehaviourPanel::OnControllerValueChanged);
 }
 
-void SRCBehaviourPanel::OnControllerSelectionChanged(TSharedPtr<FRCControllerModel> InControllerItem)
+void SRCBehaviourPanel::OnControllerSelectionChanged(TSharedPtr<FRCControllerModel> InControllerItem, ESelectInfo::Type InSelectInfo)
 {
 	SelectedControllerItemWeakPtr = InControllerItem;
+
+	// Cache the selection here since UpdateWrappedWidget will recreate the BehaviorList
+	bool bShouldRestoreSelection = false;
+	TArray<TSharedPtr<FRCBehaviourModel>> CurrentSelectedItems;
+	if (BehaviourPanelList.IsValid() && InSelectInfo == ESelectInfo::Direct)
+	{
+		if (InSelectInfo == ESelectInfo::Direct)
+		{
+			CurrentSelectedItems = BehaviourPanelList->GetSelectedBehaviourItems();
+			bShouldRestoreSelection = true;
+		}
+	}
+
 	UpdateWrappedWidget(InControllerItem);
+
 	if (BehaviourPanelList.IsValid())
 	{
-		BehaviourPanelList->SelectFirstItem();
+		if (bShouldRestoreSelection)
+		{
+			BehaviourPanelList->SetSelection(CurrentSelectedItems);
+		}
+		else if (InSelectInfo != ESelectInfo::Direct)
+		{
+			BehaviourPanelList->SelectFirstItem();
+		}
 	}
 }
 
-void SRCBehaviourPanel::UpdateWrappedWidget(TSharedPtr<FRCControllerModel> InControllerItem)
+void SRCBehaviourPanel::OnControllerValueChanged(TSharedPtr<FRCControllerModel> InController)
 {
-	TSharedPtr<SRemoteControlPanel> RemoteControlPanel = GetRemoteControlPanel();
+	BehaviourPanelList->NotifyControllerValueChanged(InController);
+}
 
-	if (InControllerItem.IsValid() && RemoteControlPanel.IsValid())
+void SRCBehaviourPanel::UpdateWrappedWidget(const TSharedPtr<FRCControllerModel>& InControllerItem)
+{
+	// Create panel widgets if invalid
+	if (!AreBehaviorPanelWidgetsValid())
 	{
-		// Behaviour Dock Panel
-		TSharedPtr<SRCMinorPanel> BehaviourDockPanel = SNew(SRCMinorPanel)
-			.HeaderLabel(LOCTEXT("BehavioursLabel", "Behavior"))
-			.EnableFooter(false)
-			[
-				SAssignNew(BehaviourPanelList, SRCBehaviourPanelList, SharedThis(this), InControllerItem, RemoteControlPanel.ToSharedRef())
-			];
-
-		// Add New Behaviour Button
-		const TSharedRef<SWidget> AddNewBehaviourButton = SNew(SComboButton)
-			.AddMetaData<FTagMetaData>(FTagMetaData(TEXT("Add Behavior")))
-			.HAlign(HAlign_Center)
-			.VAlign(VAlign_Center)
-			.ButtonStyle(&RCPanelStyle->FlatButtonStyle)
-			.ForegroundColor(FSlateColor::UseForeground())
-			.CollapseMenuOnParentFocus(true)
-			.HasDownArrow(false)
-			.ContentPadding(FMargin(4.f, 2.f))
-			.ButtonContent()
-			[
-				SNew(SBox)
-				.WidthOverride(RCPanelStyle->IconSize.X)
-				.HeightOverride(RCPanelStyle->IconSize.Y)
-				[
-					SNew(SImage)
-					.ColorAndOpacity(FSlateColor::UseForeground())
-					.Image(FAppStyle::GetBrush("Icons.PlusCircle"))
-				]
-			]
-			.MenuContent()
-			[
-				GetBehaviourMenuContentWidget()
-			];
-
-		BehaviourDockPanel->AddHeaderToolbarItem(EToolbar::Left, AddNewBehaviourButton);
+		CreateBehaviorPanelWidgets();
+	}
 
-		WrappedBoxWidget->SetContent(BehaviourDockPanel.ToSharedRef());
+	if (InControllerItem.IsValid() && BehaviourPanelList->GetControllerItem() != InControllerItem)
+	{
+		BehaviourPanelList->SetControllerItem(InControllerItem);
+		AddNewBehaviorComboButton->SetMenuContent(GetBehaviourMenuContentWidget());
 	}
 	else
 	{
-		WrappedBoxWidget->SetContent(CreateNoneSelectedWidget());
+		BehaviourPanelList->SetControllerItem(nullptr);
 	}
 }
 
@@ -227,6 +222,90 @@ FReply SRCBehaviourPanel::OnClickEmptyButton()
 	return FReply::Handled();
 }
 
+void SRCBehaviourPanel::CreateBehaviorPanelWidgets()
+{
+	const TSharedPtr<SRemoteControlPanel> RemoteControlPanel = GetRemoteControlPanel();
+
+	// Behavior list
+	if (!BehaviourPanelList.IsValid())
+	{
+		BehaviourPanelList = SNew(SRCBehaviourPanelList, SharedThis(this), nullptr, RemoteControlPanel.ToSharedRef());
+	}
+
+	// Behavior Dock Panel
+	if (!BehaviourDockPanel.IsValid())
+	{
+		BehaviourDockPanel = SNew(SRCMinorPanel)
+			.HeaderLabel(LOCTEXT("BehavioursLabel", "Behavior"))
+			.EnableFooter(false)
+			[
+				BehaviourPanelList.ToSharedRef()
+			];
+	}
+
+	// Add New Behaviour Button
+	if (!AddNewBehaviorComboButton.IsValid())
+	{
+		AddNewBehaviorComboButton = SNew(SComboButton)
+			.AddMetaData<FTagMetaData>(FTagMetaData(TEXT("Add Behavior")))
+			.HAlign(HAlign_Center)
+			.VAlign(VAlign_Center)
+			.ButtonStyle(&RCPanelStyle->FlatButtonStyle)
+			.ForegroundColor(FSlateColor::UseForeground())
+			.CollapseMenuOnParentFocus(true)
+			.HasDownArrow(false)
+			.ContentPadding(FMargin(4.f, 2.f))
+			.ButtonContent()
+			[
+				SNew(SBox)
+				.WidthOverride(RCPanelStyle->IconSize.X)
+				.HeightOverride(RCPanelStyle->IconSize.Y)
+				[
+					SNew(SImage)
+					.ColorAndOpacity(FSlateColor::UseForeground())
+					.Image(FAppStyle::GetBrush("Icons.PlusCircle"))
+				]
+			]
+			.MenuContent()
+			[
+				GetBehaviourMenuContentWidget()
+			];
+
+		BehaviourDockPanel->ClearHeaderToolbarItems();
+		BehaviourDockPanel->AddHeaderToolbarItem(EToolbar::Left, AddNewBehaviorComboButton.ToSharedRef());
+	}
+
+	// Wrapped Box
+	if (!WrappedBoxWidget.IsValid())
+	{
+		WrappedBoxWidget = SNew(SBox)
+		[
+			SNew(SWidgetSwitcher)
+			.WidgetIndex_Lambda([this] () { return BehaviourPanelList->GetControllerItem().IsValid()? 0 : 1; })
+
+			// Index 0 is for valid behavior
+			+ SWidgetSwitcher::Slot()
+			[
+				BehaviourDockPanel.ToSharedRef()
+			]
+
+			// Index 1 is for invalid behavior
+			+ SWidgetSwitcher::Slot()
+			[
+				GetNoneSelectedWidget()
+			]
+		];
+	}
+}
+
+bool SRCBehaviourPanel::AreBehaviorPanelWidgetsValid() const
+{
+	return WrappedBoxWidget.IsValid() &&
+			BehaviourPanelList.IsValid() &&
+			BehaviourDockPanel.IsValid() &&
+			AddNewBehaviorComboButton.IsValid();
+}
+
 bool SRCBehaviourPanel::IsListFocused() const
 {
 	return BehaviourPanelList.IsValid() && BehaviourPanelList->IsListFocused();
@@ -234,7 +313,10 @@ bool SRCBehaviourPanel::IsListFocused() const
 
 void SRCBehaviourPanel::DeleteSelectedPanelItems()
 {
-	BehaviourPanelList->DeleteSelectedPanelItems();
+	if (BehaviourPanelList.IsValid())
+	{
+		BehaviourPanelList->DeleteSelectedPanelItems();
+	}
 }
 
 void SRCBehaviourPanel::DuplicateBehaviour(URCBehaviour* InBehaviour)
@@ -243,7 +325,10 @@ void SRCBehaviourPanel::DuplicateBehaviour(URCBehaviour* InBehaviour)
 	{
 		URCBehaviour* NewBehaviour = URCController::DuplicateBehaviour(Controller, InBehaviour);
 
-		BehaviourPanelList->AddNewLogicItem(NewBehaviour);
+		if (BehaviourPanelList.IsValid())
+		{
+			BehaviourPanelList->AddNewLogicItem(NewBehaviour);
+		}
 	}
 }
 
diff --git a/Engine/Plugins/VirtualProduction/RemoteControl/Source/RemoteControlUI/Private/UI/Behaviour/SRCBehaviourPanel.h b/Engine/Plugins/VirtualProduction/RemoteControl/Source/RemoteControlUI/Private/UI/Behaviour/SRCBehaviourPanel.h
--- a/Engine/Plugins/VirtualProduction/RemoteControl/Source/RemoteControlUI/Private/UI/Behaviour/SRCBehaviourPanel.h
+++ b/Engine/Plugins/VirtualProduction/RemoteControl/Source/RemoteControlUI/Private/UI/Behaviour/SRCBehaviourPanel.h
@@ -9,7 +9,9 @@ struct FRCPanelStyle;
 class FRCBehaviourModel;
 class FRCControllerModel;
 class SBox;
+class SComboButton;
 class SRCBehaviourPanel;
+class SRCMinorPanel;
 class URCController;
 
 /*
@@ -26,7 +28,7 @@ public:
 		}
 
 	SLATE_END_ARGS()
-	
+
 	/** Constructs this widget with InArgs */
 	void Construct(const FArguments& InArgs, const TSharedRef<SRemoteControlPanel>& InPanel);
 
@@ -67,7 +69,7 @@ protected:
 
 private:
 	/** Get a Helper widget for behavior details. */
-	static TSharedRef<SBox> CreateNoneSelectedWidget();
+	static TSharedRef<SBox> GetNoneSelectedWidget();
 
 	/** Duplicates a given Behaviour object*/
 	void DuplicateBehaviour(URCBehaviour* InBehaviour);
@@ -76,10 +78,13 @@ private:
 	* Controller list selection change listener.
 	* Updates the list of behaviours from newly selected Controller
 	*/
-	void OnControllerSelectionChanged(TSharedPtr<FRCControllerModel> InControllerItem);
+	void OnControllerSelectionChanged(TSharedPtr<FRCControllerModel> InControllerItem, ESelectInfo::Type InSelectInfo);
+
+	/** Called when a controller value is changed */
+	void OnControllerValueChanged(TSharedPtr<FRCControllerModel> InController);
 
 	/* Rebuilds the Behaviour Panel for a newly selected Controller*/
-	void UpdateWrappedWidget(TSharedPtr<FRCControllerModel> InControllerItem = nullptr);
+	void UpdateWrappedWidget(const TSharedPtr<FRCControllerModel>& InControllerItem = nullptr);
 
 	/** Builds a menu containing the list of all possible Behaviours */
 	TSharedRef<SWidget> GetBehaviourMenuContentWidget();
@@ -93,6 +98,12 @@ private:
 	/** Handles click event for "Empty" button; clears all Behaviours from the panel*/
 	FReply OnClickEmptyButton();
 
+	/** Check the validity of the panel widgets and if not valid will create them */
+	void CreateBehaviorPanelWidgets();
+
+	/** Return whether or not all panel widgets are valid */
+	bool AreBehaviorPanelWidgetsValid() const;
+
 private:
 	/** The parent Controller that this Behaviour panel is associated with */
 	TWeakPtr<FRCControllerModel> SelectedControllerItemWeakPtr = nullptr;
@@ -103,6 +114,12 @@ private:
 	/** Widget representing List of Behaviours */
 	TSharedPtr<class SRCBehaviourPanelList> BehaviourPanelList;
 
+	/** ComboButton to add new behavior */
+	TSharedPtr<SComboButton> AddNewBehaviorComboButton;
+
+	/** Behavior Minor Panel */
+	TSharedPtr<SRCMinorPanel> BehaviourDockPanel;
+
 	/** Panel Style reference. */
 	const FRCPanelStyle* RCPanelStyle = nullptr;
 };
diff --git a/Engine/Plugins/VirtualProduction/RemoteControl/Source/RemoteControlUI/Private/UI/Behaviour/SRCBehaviourPanelList.cpp b/Engine/Plugins/VirtualProduction/RemoteControl/Source/RemoteControlUI/Private/UI/Behaviour/SRCBehaviourPanelList.cpp
--- a/Engine/Plugins/VirtualProduction/RemoteControl/Source/RemoteControlUI/Private/UI/Behaviour/SRCBehaviourPanelList.cpp
+++ b/Engine/Plugins/VirtualProduction/RemoteControl/Source/RemoteControlUI/Private/UI/Behaviour/SRCBehaviourPanelList.cpp
@@ -69,9 +69,12 @@ void SRCBehaviourPanelList::Construct(const FArguments& InArgs, TSharedRef<SRCBe
 	RemoteControlPanel->OnBehaviourAdded.AddSP(this, &SRCBehaviourPanelList::OnBehaviourAdded);
 	RemoteControlPanel->OnEmptyBehaviours.AddSP(this, &SRCBehaviourPanelList::OnEmptyBehaviours);
 
-	if (URCController* Controller = Cast<URCController>(InControllerItem->GetVirtualProperty()))
+	if (InControllerItem.IsValid())
 	{
-		Controller->OnBehaviourListModified.AddSP(this, &SRCBehaviourPanelList::OnBehaviourListModified);
+		if (URCController* Controller = Cast<URCController>(InControllerItem->GetVirtualProperty()))
+		{
+			Controller->OnBehaviourListModified.AddSP(this, &SRCBehaviourPanelList::OnBehaviourListModified);
+		}
 	}
 
 	Reset();
@@ -121,6 +124,92 @@ void SRCBehaviourPanelList::SelectFirstItem()
 	}
 }
 
+void SRCBehaviourPanelList::SetSelection(TArray<TSharedPtr<FRCBehaviourModel>> InBehaviorToSelect)
+{
+	for (const TSharedPtr<FRCBehaviourModel>& BehaviorModel : InBehaviorToSelect)
+	{
+		if (BehaviorModel.IsValid())
+		{
+			const URCBehaviour* Behavior = BehaviorModel->GetBehaviour();
+
+			const TSharedPtr<FRCBehaviourModel>* SelectedBehavior = BehaviourItems.FindByPredicate([&Behavior]
+				(const TSharedPtr<FRCBehaviourModel>& InBehaviorModel)
+				{
+					if (!InBehaviorModel.IsValid())
+					{
+						return false;
+					}
+
+					const URCBehaviour* CurrentBehavior = InBehaviorModel->GetBehaviour();
+
+					if (!Behavior || !CurrentBehavior)
+					{
+						return false;
+					}
+
+					// Use the Behavior Internal Id to check for uniqueness
+					return Behavior->Id == CurrentBehavior->Id;
+				});
+
+			if (SelectedBehavior)
+			{
+				ListView->SetItemSelection(*SelectedBehavior, true);
+			}
+		}
+	}
+}
+
+void SRCBehaviourPanelList::NotifyControllerValueChanged(TSharedPtr<FRCControllerModel> InControllerModel)
+{
+	// Get the current behavior displaying its action panel
+	if (const TSharedPtr<FRCBehaviourModel>& SelectedBehavior = GetSelectedBehaviourItem())
+	{
+		SelectedBehavior->NotifyControllerValueChanged(InControllerModel);
+	}
+}
+
+TSharedPtr<FRCControllerModel> SRCBehaviourPanelList::GetControllerItem() const
+{
+	if (ControllerItemWeakPtr.IsValid())
+	{
+		return ControllerItemWeakPtr.Pin();
+	}
+	return nullptr;
+}
+
+void SRCBehaviourPanelList::SetControllerItem(const TSharedPtr<FRCControllerModel>& InNewControllerItem)
+{
+	// Remove old Controller callback
+	const TSharedPtr<FRCControllerModel> CurrentController = GetControllerItem();
+
+	if (CurrentController == InNewControllerItem)
+	{
+		return;
+	}
+
+	if (CurrentController.IsValid())
+	{
+		if (URCController* Controller = Cast<URCController>(CurrentController->GetVirtualProperty()))
+		{
+			Controller->OnBehaviourListModified.RemoveAll(this);
+		}
+	}
+
+	ControllerItemWeakPtr = InNewControllerItem;
+
+	// Add new Controller callback
+	const TSharedPtr<FRCControllerModel> NewController = GetControllerItem();
+	if (NewController.IsValid())
+	{
+		if (URCController* Controller = Cast<URCController>(NewController->GetVirtualProperty()))
+		{
+			Controller->OnBehaviourListModified.AddSP(this, &SRCBehaviourPanelList::OnBehaviourListModified);
+		}
+	}
+
+	Reset();
+}
+
 void SRCBehaviourPanelList::SetIsBehaviourEnabled(const bool bIsEnabled)
 {
 	// Disable all selected behaviour here
@@ -180,7 +269,7 @@ void SRCBehaviourPanelList::AddBehaviourToList(URCBehaviour* InBehaviour)
 	}
 	else if (URCSetAssetByPathBehaviour* SetAssetByPathBehaviour = Cast<URCSetAssetByPathBehaviour>(InBehaviour))
 	{
-		BehaviourItems.Add(MakeShared<FRCSetAssetByPathBehaviourModel>(SetAssetByPathBehaviour));
+		BehaviourItems.Add(MakeShared<FRCSetAssetByPathBehaviourModel>(SetAssetByPathBehaviour, RemoteControlPanel));
 	}
 	else if (URCBehaviourBind* BindBehaviour = Cast<URCBehaviourBind>(InBehaviour))
 	{
@@ -200,6 +289,8 @@ void SRCBehaviourPanelList::AddBehaviourToList(URCBehaviour* InBehaviour)
 
 void SRCBehaviourPanelList::Reset()
 {
+	const TArray<TSharedPtr<FRCBehaviourModel>> SelectedBehaviors = ListView->GetSelectedItems();
+
 	BehaviourItems.Empty();
 
 	if (TSharedPtr<FRCControllerModel> ControllerItem = ControllerItemWeakPtr.Pin())
@@ -214,6 +305,7 @@ void SRCBehaviourPanelList::Reset()
 	}
 
 	ListView->RebuildList();
+	SetSelection(SelectedBehaviors);
 }
 
 TSharedRef<ITableRow> SRCBehaviourPanelList::OnGenerateWidgetForList(TSharedPtr<FRCBehaviourModel> InItem,
@@ -363,6 +455,11 @@ TArray<TSharedPtr<FRCLogicModeBase>> SRCBehaviourPanelList::GetSelectedLogicItem
 	return SelectedValidLogicItems;
 }
 
+TArray<TSharedPtr<FRCBehaviourModel>> SRCBehaviourPanelList::GetSelectedBehaviourItems() const
+{
+	return ListView->GetSelectedItems();
+}
+
 void SRCBehaviourPanelList::RequestRefresh()
 {
 	ListView->RequestListRefresh();
diff --git a/Engine/Plugins/VirtualProduction/RemoteControl/Source/RemoteControlUI/Private/UI/Behaviour/SRCBehaviourPanelList.h b/Engine/Plugins/VirtualProduction/RemoteControl/Source/RemoteControlUI/Private/UI/Behaviour/SRCBehaviourPanelList.h
--- a/Engine/Plugins/VirtualProduction/RemoteControl/Source/RemoteControlUI/Private/UI/Behaviour/SRCBehaviourPanelList.h
+++ b/Engine/Plugins/VirtualProduction/RemoteControl/Source/RemoteControlUI/Private/UI/Behaviour/SRCBehaviourPanelList.h
@@ -53,11 +53,15 @@ public:
 	/** Returns the UI items currently selected by the user (if any). */
 	virtual TArray<TSharedPtr<FRCLogicModeBase>> GetSelectedLogicItems() override;
 
+	/** Return the behavior currently selected and displaying its ActionPanel */
 	TSharedPtr<FRCBehaviourModel> GetSelectedBehaviourItem()
 	{
 		return SelectedBehaviourItemWeakPtr.Pin();
 	}
 
+	/** Return the list of all behavior currently selected */
+	TArray<TSharedPtr<FRCBehaviourModel>> GetSelectedBehaviourItems() const;
+
 	void RequestRefresh();
 
 	virtual void AddNewLogicItem(UObject* InLogicItem) override;
@@ -68,6 +72,18 @@ public:
 	/** Will set the selection to the first item */
 	void SelectFirstItem();
 
+	/** Set the current selection to the given InBehaviorToSelect list */
+	void SetSelection(TArray<TSharedPtr<FRCBehaviourModel>> InBehaviorToSelect);
+
+	/** Called when a controller value changed */
+	void NotifyControllerValueChanged(TSharedPtr<FRCControllerModel> InController);
+
+	/** Return the controller item for this behavior */
+	TSharedPtr<FRCControllerModel> GetControllerItem() const;
+	
+	/** Set a new controller item to update the behavior view based on the new controller */
+	void SetControllerItem(const TSharedPtr<FRCControllerModel>& InNewControllerItem);
+
 private:
 
 	/** Enables or Disables the currently selected behaviour */
diff --git a/Engine/Plugins/VirtualProduction/RemoteControl/Source/RemoteControlUI/Private/UI/Controller/RCControllerModel.cpp b/Engine/Plugins/VirtualProduction/RemoteControl/Source/RemoteControlUI/Private/UI/Controller/RCControllerModel.cpp
--- a/Engine/Plugins/VirtualProduction/RemoteControl/Source/RemoteControlUI/Private/UI/Controller/RCControllerModel.cpp
+++ b/Engine/Plugins/VirtualProduction/RemoteControl/Source/RemoteControlUI/Private/UI/Controller/RCControllerModel.cpp
@@ -287,12 +287,9 @@ void FRCControllerModel::OnTextControlValueTypeChanged(TSharedPtr<FString, ESPMo
 
 void FRCControllerModel::OnPropertyValueChanged(const FPropertyChangedEvent& InPropertyChangedEvent)
 {	
-	if (URCVirtualPropertyBase* ControllerProperty = GetVirtualProperty())
+	if (OnValueChanged.IsBound())
 	{
-		if (OnValueChanged.IsBound())
-		{
-			OnValueChanged.Broadcast(ControllerProperty);
-		}
+		OnValueChanged.Broadcast(StaticCastSharedRef<FRCControllerModel>(AsShared()));
 	}
 }
 
diff --git a/Engine/Plugins/VirtualProduction/RemoteControl/Source/RemoteControlUI/Private/UI/Controller/RCControllerModel.h b/Engine/Plugins/VirtualProduction/RemoteControl/Source/RemoteControlUI/Private/UI/Controller/RCControllerModel.h
--- a/Engine/Plugins/VirtualProduction/RemoteControl/Source/RemoteControlUI/Private/UI/Controller/RCControllerModel.h
+++ b/Engine/Plugins/VirtualProduction/RemoteControl/Source/RemoteControlUI/Private/UI/Controller/RCControllerModel.h
@@ -31,7 +31,7 @@ class FRCControllerModel : public FRCLogicModeBase, public FSelfRegisteringEdito
 {
 public:
 	DECLARE_MULTICAST_DELEGATE_TwoParams(FOnValueTypeChanged, URCVirtualPropertyBase* /* InController */, EPropertyBagPropertyType /* InValueType */);
-	DECLARE_MULTICAST_DELEGATE_OneParam(FOnValueChanged, URCVirtualPropertyBase* /* InController */);
+	DECLARE_MULTICAST_DELEGATE_OneParam(FOnValueChanged, TSharedPtr<FRCControllerModel> /* InControllerModel */);
 	
 	FRCControllerModel(URCVirtualPropertyBase* InVirtualProperty, const TSharedRef<IDetailTreeNode>& InTreeNode, const TSharedPtr<SRemoteControlPanel> InRemoteControlPanel);
 	
diff --git a/Engine/Plugins/VirtualProduction/RemoteControl/Source/RemoteControlUI/Private/UI/Controller/SRCControllerPanel.cpp b/Engine/Plugins/VirtualProduction/RemoteControl/Source/RemoteControlUI/Private/UI/Controller/SRCControllerPanel.cpp
--- a/Engine/Plugins/VirtualProduction/RemoteControl/Source/RemoteControlUI/Private/UI/Controller/SRCControllerPanel.cpp
+++ b/Engine/Plugins/VirtualProduction/RemoteControl/Source/RemoteControlUI/Private/UI/Controller/SRCControllerPanel.cpp
@@ -229,6 +229,15 @@ TArray<TSharedPtr<FRCLogicModeBase>> SRCControllerPanel::GetSelectedLogicItems()
 	return {};
 }
 
+int32 SRCControllerPanel::NumControllerItems() const
+{
+	if (ControllerPanelList.IsValid())
+	{
+		return ControllerPanelList->NumControllerItems();
+	}
+	return INDEX_NONE;
+}
+
 void SRCControllerPanel::DuplicateController(URCController* InController)
 {
 	if (!ensure(InController))
diff --git a/Engine/Plugins/VirtualProduction/RemoteControl/Source/RemoteControlUI/Private/UI/Controller/SRCControllerPanel.h b/Engine/Plugins/VirtualProduction/RemoteControl/Source/RemoteControlUI/Private/UI/Controller/SRCControllerPanel.h
--- a/Engine/Plugins/VirtualProduction/RemoteControl/Source/RemoteControlUI/Private/UI/Controller/SRCControllerPanel.h
+++ b/Engine/Plugins/VirtualProduction/RemoteControl/Source/RemoteControlUI/Private/UI/Controller/SRCControllerPanel.h
@@ -76,6 +76,9 @@ public:
 	/** Returns the UI items currently selected by the user (if any). To be implemented per child panel*/
 	virtual TArray<TSharedPtr<FRCLogicModeBase>> GetSelectedLogicItems() const override;
 
+	/** Retrieve the current number of Controllers */
+	int32 NumControllerItems() const;
+
 	void EnterRenameMode();
 
 protected:
diff --git a/Engine/Plugins/VirtualProduction/RemoteControl/Source/RemoteControlUI/Private/UI/Controller/SRCControllerPanelList.cpp b/Engine/Plugins/VirtualProduction/RemoteControl/Source/RemoteControlUI/Private/UI/Controller/SRCControllerPanelList.cpp
--- a/Engine/Plugins/VirtualProduction/RemoteControl/Source/RemoteControlUI/Private/UI/Controller/SRCControllerPanelList.cpp
+++ b/Engine/Plugins/VirtualProduction/RemoteControl/Source/RemoteControlUI/Private/UI/Controller/SRCControllerPanelList.cpp
@@ -440,7 +440,10 @@ void SRCControllerPanelList::Reset()
 			ControllerModel->OnValueTypeChanged.RemoveAll(this);
 		}
 	}
-	
+
+	// Cache Controller Selection
+	const TArray<TSharedPtr<FRCControllerModel>> SelectedControllers = ListView->GetSelectedItems();
+
 	ControllerItems.Empty();
 
 	check(ControllerPanelWeakPtr.IsValid());
@@ -501,14 +504,13 @@ void SRCControllerPanelList::Reset()
 					if (ensureAlways(ControllerItems.IsValidIndex(Controller->DisplayIndex)))
 					{
 						const TSharedRef<FRCControllerModel> ControllerModel = MakeShared<FRCControllerModel>(Controller, Child, RemoteControlPanel);
-						ControllerItems[Controller->DisplayIndex] = ControllerModel;
-
+						ControllerModel->OnValueChanged.AddSP(this, &SRCControllerPanelList::OnControllerValueChanged, bIsMultiController);
 						if (bIsMultiController)
 						{
 							ControllerModel->SetMultiController(bIsMultiController);
 							ControllerModel->OnValueTypeChanged.AddSP(this, &SRCControllerPanelList::OnControllerValueTypeChanged);
-							ControllerModel->OnValueChanged.AddSP(this, &SRCControllerPanelList::OnControllerValueChanged);
 						}
+						ControllerItems[Controller->DisplayIndex] = ControllerModel;
 					}
 				}
 			}
@@ -560,6 +562,27 @@ void SRCControllerPanelList::Reset()
 	}
 	
 	ListView->RebuildList();
+
+	// Restore Controller Selection
+	for (const TSharedPtr<FRCControllerModel>& ControllerModel : SelectedControllers)
+	{
+		if (ControllerModel.IsValid())
+		{
+			const FName SelectedControllerName = ControllerModel->GetPropertyName();
+
+			const TSharedPtr<FRCControllerModel>* SelectedController = ControllerItems.FindByPredicate([&SelectedControllerName]
+				(const TSharedPtr<FRCControllerModel>& InControllerModel)
+				{
+					// Internal PropertyName is unique, so we use that
+					return InControllerModel.IsValid() && SelectedControllerName == InControllerModel->GetPropertyName();
+				});
+
+			if (SelectedController)
+			{
+				ListView->SetItemSelection(*SelectedController, true);
+			}
+		}
+	}
 }
 
 TSharedRef<ITableRow> SRCControllerPanelList::OnGenerateWidgetForList(TSharedPtr<FRCControllerModel> InItem, const TSharedRef<STableViewBase>& OwnerTable)
@@ -571,7 +594,7 @@ TSharedRef<ITableRow> SRCControllerPanelList::OnGenerateWidgetForList(TSharedPtr
 		.Padding(FMargin(4.5f));
 }
 
-void SRCControllerPanelList::OnTreeSelectionChanged(TSharedPtr<FRCControllerModel> InItem, ESelectInfo::Type)
+void SRCControllerPanelList::OnTreeSelectionChanged(TSharedPtr<FRCControllerModel> InItem, ESelectInfo::Type InSelectInfo)
 {
 	if (TSharedPtr<SRCControllerPanel> ControllerPanel = ControllerPanelWeakPtr.Pin())
 	{
@@ -580,7 +603,7 @@ void SRCControllerPanelList::OnTreeSelectionChanged(TSharedPtr<FRCControllerMode
 			if (InItem != SelectedControllerItemWeakPtr.Pin())
 			{
 				SelectedControllerItemWeakPtr = InItem;
-				RemoteControlPanel->OnControllerSelectionChanged.Broadcast(InItem);
+				RemoteControlPanel->OnControllerSelectionChanged.Broadcast(InItem, InSelectInfo);
 				RemoteControlPanel->OnBehaviourSelectionChanged.Broadcast(InItem.IsValid() ? InItem->GetSelectedBehaviourModel() : nullptr);
 			}
 		}
@@ -635,16 +658,25 @@ void SRCControllerPanelList::OnControllerValueTypeChanged(URCVirtualPropertyBase
 	}
 }
 
-void SRCControllerPanelList::OnControllerValueChanged(URCVirtualPropertyBase* InController)
-{	
-	const FName& FieldId = InController->FieldId;
+void SRCControllerPanelList::OnControllerValueChanged(TSharedPtr<FRCControllerModel> InControllerModel, bool bInIsMultiController)
+{
+	if (bInIsMultiController)
+	{
+		if (const URCVirtualPropertyBase* Controller = InControllerModel->GetVirtualProperty())
+		{
+			FRCMultiController MultiController = MultiControllers.GetMultiController(Controller->FieldId);
 
-	FRCMultiController MultiController = MultiControllers.GetMultiController(FieldId);
+			if (MultiController.IsValid())
+			{
+				MultiController.UpdateHandledControllersValue();
+			}
+		}
+	}
 
-	if (MultiController.IsValid())
+	if (const TSharedPtr<SRemoteControlPanel>& RemoteControlPanel = GetRemoteControlPanel())
 	{
-		MultiController.UpdateHandledControllersValue();
-	}	
+		RemoteControlPanel->OnControllerValueChangedDelegate.Broadcast(InControllerModel);
+	}
 }
 
 
@@ -654,7 +686,7 @@ void SRCControllerPanelList::OnEmptyControllers()
 	{
 		if (TSharedPtr<SRemoteControlPanel> RemoteControlPanel = ControllerPanel->GetRemoteControlPanel())
 		{
-			RemoteControlPanel->OnControllerSelectionChanged.Broadcast(nullptr);
+			RemoteControlPanel->OnControllerSelectionChanged.Broadcast(nullptr, ESelectInfo::Direct);
 			RemoteControlPanel->OnBehaviourSelectionChanged.Broadcast(nullptr);
 		}
 
@@ -671,7 +703,7 @@ void SRCControllerPanelList::BroadcastOnItemRemoved()
 {
 	if (const TSharedPtr<SRemoteControlPanel> RemoteControlPanel = ControllerPanelWeakPtr.Pin()->GetRemoteControlPanel())
 	{
-		RemoteControlPanel->OnControllerSelectionChanged.Broadcast(nullptr);
+		RemoteControlPanel->OnControllerSelectionChanged.Broadcast(nullptr, ESelectInfo::Direct);
 		RemoteControlPanel->OnBehaviourSelectionChanged.Broadcast(nullptr);
 	}
 }
diff --git a/Engine/Plugins/VirtualProduction/RemoteControl/Source/RemoteControlUI/Private/UI/Controller/SRCControllerPanelList.h b/Engine/Plugins/VirtualProduction/RemoteControl/Source/RemoteControlUI/Private/UI/Controller/SRCControllerPanelList.h
--- a/Engine/Plugins/VirtualProduction/RemoteControl/Source/RemoteControlUI/Private/UI/Controller/SRCControllerPanelList.h
+++ b/Engine/Plugins/VirtualProduction/RemoteControl/Source/RemoteControlUI/Private/UI/Controller/SRCControllerPanelList.h
@@ -186,7 +186,7 @@ private:
 	 * Called when a Controller Value changes.
 	 * Currently used to update handled controllers, in case of MultiControllers.
 	 */
-	void OnControllerValueChanged(URCVirtualPropertyBase* InController);
+	void OnControllerValueChanged(TSharedPtr<FRCControllerModel> InControllerModel, bool bInIsMultiController);
 
 	/** Creates a new Controller for the given Remote Control Property and also binds to it */
 	void CreateAutoBindForProperty(TSharedPtr<const FRemoteControlProperty> RemoteControlProperty);
diff --git a/Engine/Plugins/VirtualProduction/RemoteControl/Source/RemoteControlUI/Private/UI/Customizations/RCAssetPathElementCustomization.cpp b/Engine/Plugins/VirtualProduction/RemoteControl/Source/RemoteControlUI/Private/UI/Customizations/RCAssetPathElementCustomization.cpp
--- a/Engine/Plugins/VirtualProduction/RemoteControl/Source/RemoteControlUI/Private/UI/Customizations/RCAssetPathElementCustomization.cpp
+++ b/Engine/Plugins/VirtualProduction/RemoteControl/Source/RemoteControlUI/Private/UI/Customizations/RCAssetPathElementCustomization.cpp
@@ -7,12 +7,14 @@
 #include "Editor.h"
 #include "Framework/Application/SlateApplication.h"
 #include "IDetailChildrenBuilder.h"
+#include "IPropertyUtilities.h"
 #include "PropertyCustomizationHelpers.h"
 #include "UI/RemoteControlPanelStyle.h"
 #include "Widgets/Images/SImage.h"
 #include "Widgets/Input/SButton.h"
 #include "Widgets/Input/SCheckBox.h"
 #include "Widgets/Input/SEditableTextBox.h"
+#include "Widgets/Layout/SWidgetSwitcher.h"
 #include "Widgets/SToolTip.h"
 #include "Widgets/Text/STextBlock.h"
 
@@ -26,6 +28,7 @@ TSharedRef<IPropertyTypeCustomization> FRCAssetPathElementCustomization::MakeIns
 void FRCAssetPathElementCustomization::CustomizeHeader(TSharedRef<IPropertyHandle> InPropertyHandle, FDetailWidgetRow& InHeaderRow, IPropertyTypeCustomizationUtils& InCustomizationUtils)
 {
 	ArrayEntryHandle = InPropertyHandle;
+	PropertyUtilities = InCustomizationUtils.GetPropertyUtilities();
 	IsInputHandle = ArrayEntryHandle->GetChildHandle(GET_MEMBER_NAME_CHECKED(FRCAssetPathElement, bIsInput));
 	PathHandle = ArrayEntryHandle->GetChildHandle(GET_MEMBER_NAME_CHECKED(FRCAssetPathElement, Path));
 
@@ -37,6 +40,9 @@ void FRCAssetPathElementCustomization::CustomizeHeader(TSharedRef<IPropertyHandl
 	const TSharedRef<SToolTip> GetAssetPathToolTipWidget = SNew(SToolTip)
 		.Text(LOCTEXT("RCGetAssetPathButton_Tooltip", "Get the path of the currently first selected asset in the content browser and set it to the current path"));
 
+	const TSharedRef<SToolTip> CreateControllerToolTipWidget = SNew(SToolTip)
+		.Text(LOCTEXT("RCCreateController_Tooltip", "Create a controller for the given RC Input path entry"));
+
 	InHeaderRow.NameContent()
 		[
 			InPropertyHandle->CreatePropertyNameWidget()
@@ -45,6 +51,7 @@ void FRCAssetPathElementCustomization::CustomizeHeader(TSharedRef<IPropertyHandl
 	InHeaderRow.ValueContent()
 	[
 		SNew(SHorizontalBox)
+        // RC Input CheckBox
 		+SHorizontalBox::Slot()
 		.AutoWidth()
 		.Padding(5.f, 0.f)
@@ -62,8 +69,10 @@ void FRCAssetPathElementCustomization::CustomizeHeader(TSharedRef<IPropertyHandl
 			]
 		]
 
+		// Path String
 		+SHorizontalBox::Slot()
 		.FillWidth(1.f)
+		.Padding(5.f, 0.f)
 		.VAlign(VAlign_Center)
 		[
 			PathHandle->CreatePropertyValueWidget(false)
@@ -73,17 +82,40 @@ void FRCAssetPathElementCustomization::CustomizeHeader(TSharedRef<IPropertyHandl
 		.AutoWidth()
 		.VAlign(VAlign_Center)
 		[
-			SNew(SButton)
-			.ButtonStyle(FAppStyle::Get(), "SimpleButton")
-			.OnClicked(this, &FRCAssetPathElementCustomization::OnGetAssetFromSelectionClicked)
-			.ToolTip(GetAssetPathToolTipWidget)
-			.IsFocusable(false)
-			.ContentPadding(0)
-			.Content()
+			SNew(SWidgetSwitcher)
+			.WidgetIndex(this, &FRCAssetPathElementCustomization::OnGetWidgetSwitcherIndex)
+			// [0] Get Current Selected Asset Path Button
+			+ SWidgetSwitcher::Slot()
 			[
-				SNew(SImage)
-				.Image(FAppStyle::GetBrush("Icons.Use"))
-				.ColorAndOpacity(FSlateColor::UseForeground())
+				SNew(SButton)
+				.ButtonStyle(FAppStyle::Get(), "SimpleButton")
+				.OnClicked(this, &FRCAssetPathElementCustomization::OnGetAssetFromSelectionClicked)
+				.ToolTip(GetAssetPathToolTipWidget)
+				.IsFocusable(false)
+				.ContentPadding(0)
+				.Content()
+				[
+					SNew(SImage)
+					.Image(FAppStyle::GetBrush("Icons.Use"))
+					.ColorAndOpacity(FSlateColor::UseForeground())
+				]
+			]
+
+			// [1] Create Controller button
+			+ SWidgetSwitcher::Slot()
+			[
+				SNew(SButton)
+				.ButtonStyle(FAppStyle::Get(), "SimpleButton")
+				.OnClicked(this, &FRCAssetPathElementCustomization::OnCreateControllerButtonClicked)
+				.ToolTip(CreateControllerToolTipWidget)
+				.IsFocusable(false)
+				.ContentPadding(0)
+				.Content()
+				[
+					SNew(SImage)
+					.Image(FAppStyle::GetBrush("Icons.PlusCircle"))
+					.ColorAndOpacity(FSlateColor::UseForeground())
+				]
 			]
 		]
 	];
@@ -159,4 +191,26 @@ FReply FRCAssetPathElementCustomization::OnGetAssetFromSelectionClicked() const
 	return FReply::Handled();
 }
 
+FReply FRCAssetPathElementCustomization::OnCreateControllerButtonClicked() const
+{
+	if (PropertyUtilities.IsValid() && ArrayEntryHandle.IsValid())
+	{
+		FPropertyChangedEvent Event(ArrayEntryHandle->GetProperty(), EPropertyChangeType::ValueSet);
+		Event.MemberProperty = ArrayEntryHandle->GetProperty();
+
+		TMap<FString, int32> ArrayIndexPerObject;
+		ArrayIndexPerObject.Add(Event.GetMemberPropertyName().ToString(), ArrayEntryHandle->GetArrayIndex());
+
+		Event.SetArrayIndexPerObject(MakeArrayView(&ArrayIndexPerObject, 1));
+		Event.ObjectIteratorIndex = 0;
+		PropertyUtilities->NotifyFinishedChangingProperties(Event);
+	}
+	return FReply::Handled();
+}
+
+int32 FRCAssetPathElementCustomization::OnGetWidgetSwitcherIndex() const
+{
+	return IsChecked() == ECheckBoxState::Checked ? 1 : 0;
+}
+
 #undef LOCTEXT_NAMESPACE
diff --git a/Engine/Plugins/VirtualProduction/RemoteControl/Source/RemoteControlUI/Private/UI/Customizations/RCAssetPathElementCustomization.h b/Engine/Plugins/VirtualProduction/RemoteControl/Source/RemoteControlUI/Private/UI/Customizations/RCAssetPathElementCustomization.h
--- a/Engine/Plugins/VirtualProduction/RemoteControl/Source/RemoteControlUI/Private/UI/Customizations/RCAssetPathElementCustomization.h
+++ b/Engine/Plugins/VirtualProduction/RemoteControl/Source/RemoteControlUI/Private/UI/Customizations/RCAssetPathElementCustomization.h
@@ -27,9 +27,14 @@ private:
 	void OnCheckStateChanged(ECheckBoxState InNewState) const;
 	/** Callback called when you click on the arrow button next to the entries to retrieve the path of the selected Asset */
 	FReply OnGetAssetFromSelectionClicked() const;
+	/** Callback called when you click on the plus button next to the entries to create a controller associated with the current RC Input entry */
+	FReply OnCreateControllerButtonClicked() const;
+	/** Retrieve the current widget switcher index */
+	int32 OnGetWidgetSwitcherIndex() const;
 
 private:
 	TSharedPtr<SWidget> PathWidget;
+	TSharedPtr<IPropertyUtilities> PropertyUtilities;
 	TSharedPtr<IPropertyHandle> ArrayEntryHandle;
 	TSharedPtr<IPropertyHandle> IsInputHandle;
 	TSharedPtr<IPropertyHandle> PathHandle;
diff --git a/Engine/Plugins/VirtualProduction/RemoteControl/Source/RemoteControlUI/Private/UI/Panels/SRCDockPanel.cpp b/Engine/Plugins/VirtualProduction/RemoteControl/Source/RemoteControlUI/Private/UI/Panels/SRCDockPanel.cpp
--- a/Engine/Plugins/VirtualProduction/RemoteControl/Source/RemoteControlUI/Private/UI/Panels/SRCDockPanel.cpp
+++ b/Engine/Plugins/VirtualProduction/RemoteControl/Source/RemoteControlUI/Private/UI/Panels/SRCDockPanel.cpp
@@ -244,6 +244,8 @@ void SRCMinorPanel::Construct(const SRCMinorPanel::FArguments& InArgs)
 	// Enable Header upon explicit request.
 	if (bIsHeaderEnabled.Get() && ContentPanel.IsValid())
 	{
+		ContentIndex = 1;
+
 		// Header Panel
 		ContentPanel->AddSlot(0)
 			.SizeRule(SSplitter::SizeToContent)
@@ -334,10 +336,17 @@ END_SLATE_FUNCTION_BUILD_OPTIMIZATION
 
 void SRCMinorPanel::SetContent(TSharedRef<SWidget> InContent)
 {
-	ChildSlot
-		[
-			InContent
-		];
+	if (ContentPanel.IsValid())
+	{
+		if (const FChildren* Children = ContentPanel->GetChildren())
+		{
+			if (ContentIndex >= 0 && Children->NumSlot() > ContentIndex)
+			{
+				SSplitter::FSlot& Content = ContentPanel->SlotAt(ContentIndex);
+				Content[InContent];
+			}
+		}
+	}
 }
 
 const TSharedRef<SWidget>& SRCMinorPanel::GetContent() const
@@ -427,3 +436,24 @@ void SRCMinorPanel::AddHeaderToolbarItem(EToolbar InToolbar, TSharedRef<SWidget>
 			break;
 	}
 }
+
+void SRCMinorPanel::ClearHeaderToolbarItems() const
+{
+	// Clear left header
+	if (LeftHeaderToolbar.IsValid())
+	{
+		LeftHeaderToolbar->ClearChildren();
+	}
+	
+	// Clear center header
+	if (CenterHeaderToolbar.IsValid())
+	{
+		CenterHeaderToolbar->ClearChildren();
+	}
+	
+	// Clear right header
+	if (RightHeaderToolbar.IsValid())
+	{
+		RightHeaderToolbar->ClearChildren();
+	}
+}
diff --git a/Engine/Plugins/VirtualProduction/RemoteControl/Source/RemoteControlUI/Private/UI/Panels/SRCDockPanel.h b/Engine/Plugins/VirtualProduction/RemoteControl/Source/RemoteControlUI/Private/UI/Panels/SRCDockPanel.h
--- a/Engine/Plugins/VirtualProduction/RemoteControl/Source/RemoteControlUI/Private/UI/Panels/SRCDockPanel.h
+++ b/Engine/Plugins/VirtualProduction/RemoteControl/Source/RemoteControlUI/Private/UI/Panels/SRCDockPanel.h
@@ -189,7 +189,10 @@ public:
 	
 	/** Adds widgets to the specified header toolbar (Left or Right). */
 	void AddHeaderToolbarItem(EToolbar InToolbar, TSharedRef<SWidget> InWidget);
-	
+
+	/** Clear toolbar items */
+	void ClearHeaderToolbarItems() const;
+
 private:
 
 	/** Actual toolbar widget located left to the footer. */
@@ -206,4 +209,7 @@ private:
 
 	/** Actual toolbar widget located Center to the header. */
     TSharedPtr<SHorizontalBox> CenterHeaderToolbar;
+
+	/** Index of the Content slot */
+	int32 ContentIndex = 0;
 };
diff --git a/Engine/Plugins/VirtualProduction/RemoteControl/Source/RemoteControlUI/Private/UI/SRemoteControlPanel.cpp b/Engine/Plugins/VirtualProduction/RemoteControl/Source/RemoteControlUI/Private/UI/SRemoteControlPanel.cpp
--- a/Engine/Plugins/VirtualProduction/RemoteControl/Source/RemoteControlUI/Private/UI/SRemoteControlPanel.cpp
+++ b/Engine/Plugins/VirtualProduction/RemoteControl/Source/RemoteControlUI/Private/UI/SRemoteControlPanel.cpp
@@ -2551,6 +2551,15 @@ bool SRemoteControlPanel::CanRenameEntity() const
 	return false;
 }
 
+int32 SRemoteControlPanel::NumControllerItems() const
+{
+	if (ControllerPanel.IsValid())
+	{
+		return ControllerPanel->NumControllerItems();
+	}
+	return INDEX_NONE;
+}
+
 void SRemoteControlPanel::ChangePropertyId_Execute() const
 {
 	if (!LastSelectedEntity.IsValid())
diff --git a/Engine/Plugins/VirtualProduction/RemoteControl/Source/RemoteControlUI/Private/UI/SRemoteControlPanel.h b/Engine/Plugins/VirtualProduction/RemoteControl/Source/RemoteControlUI/Private/UI/SRemoteControlPanel.h
--- a/Engine/Plugins/VirtualProduction/RemoteControl/Source/RemoteControlUI/Private/UI/SRemoteControlPanel.h
+++ b/Engine/Plugins/VirtualProduction/RemoteControl/Source/RemoteControlUI/Private/UI/SRemoteControlPanel.h
@@ -64,7 +64,8 @@ class SRemoteControlPanel : public SCompoundWidget, public FGCObject
 public:
 	// Remote Control Logic Delegates
 	DECLARE_MULTICAST_DELEGATE_OneParam(FOnControllerAdded, const FName& /* InPropertyName */);
-	DECLARE_MULTICAST_DELEGATE_OneParam(FOnControllerSelectionChanged, TSharedPtr<FRCControllerModel> /* InControllerItem */);
+	DECLARE_MULTICAST_DELEGATE_TwoParams(FOnControllerSelectionChanged, TSharedPtr<FRCControllerModel> /* InControllerItem */, ESelectInfo::Type /* Select Type Info */);
+	DECLARE_MULTICAST_DELEGATE_OneParam(FOnControllerValueChanged, TSharedPtr<FRCControllerModel> /* InControllerItem */);
 	DECLARE_MULTICAST_DELEGATE_OneParam(FOnBehaviourAdded, const URCBehaviour* /* InBehaviour */);
 	DECLARE_MULTICAST_DELEGATE_OneParam(FOnBehaviourSelectionChanged, TSharedPtr<FRCBehaviourModel> /* InBehaviourItem */);
 	DECLARE_MULTICAST_DELEGATE_OneParam(FOnActionAdded, URCAction* /* InAction */);
@@ -152,6 +153,9 @@ public:
 		return ActionPanel;
 	}
 
+	/** Retrieves the number of controllers. */
+	int32 NumControllerItems() const;
+
 	/** For Copy UI command - Sets the logic clipboard item and source */
 	void SetLogicClipboardItems(const TArray<UObject*>& InItems, const TSharedPtr<SRCLogicPanelBase>& SourcePanel);
 
@@ -456,6 +460,7 @@ public:
 	FOnBehaviourAdded OnBehaviourAdded;
 	FOnActionAdded OnActionAdded;
 	FOnControllerSelectionChanged OnControllerSelectionChanged;
+	FOnControllerValueChanged OnControllerValueChangedDelegate;
 	FOnBehaviourSelectionChanged OnBehaviourSelectionChanged;
 	FOnEmptyControllers OnEmptyControllers;
 	FOnEmptyBehaviours OnEmptyBehaviours;
