[RemoteControl] Fix issue with right click in details panel
Some struct do not show the Expose Sub-Property entry in the right click menu for the details panel

diff --git a/Engine/Plugins/VirtualProduction/RemoteControl/Source/RemoteControlUI/Private/RemoteControlUIModule.cpp b/Engine/Plugins/VirtualProduction/RemoteControl/Source/RemoteControlUI/Private/RemoteControlUIModule.cpp
--- a/Engine/Plugins/VirtualProduction/RemoteControl/Source/RemoteControlUI/Private/RemoteControlUIModule.cpp
+++ b/Engine/Plugins/VirtualProduction/RemoteControl/Source/RemoteControlUI/Private/RemoteControlUIModule.cpp
@@ -451,7 +451,29 @@ void FRemoteControlUIModule::FillRemoteControlRowContextSection(UToolMenu* InToo
 		return;
 	}
 
-	FRCExposesPropertyArgs ExposesPropertyArgs(Context->PropertyHandles[0]);
+	if (!Context->PropertyHandles[0].IsValid())
+	{
+		return;
+	}
+
+	TSharedPtr<IPropertyHandle> PropertyHandleToUse = Context->PropertyHandles[0];
+
+	// If there are multiple handles and the parent property is a struct use that
+	if (Context->PropertyHandles.Num() > 1)
+	{
+		if (const TSharedPtr<IPropertyHandle>& ParentHandle = PropertyHandleToUse->GetParentHandle())
+		{
+			if (const FProperty* ParentProperty = ParentHandle->GetProperty())
+			{
+				if (ParentProperty->IsA<FStructProperty>())
+				{
+					PropertyHandleToUse = ParentHandle;
+				}
+			}
+		}
+	}
+
+	FRCExposesPropertyArgs ExposesPropertyArgs(PropertyHandleToUse);
 
 	if (!ExposesPropertyArgs.IsValid())
 	{
