[RemoteControl] Fix Add Action button not updating
Add Action button not updating when exposing new properties

diff --git a/Engine/Plugins/VirtualProduction/RemoteControl/Source/RemoteControlUI/Private/UI/Action/SRCActionPanel.cpp b/Engine/Plugins/VirtualProduction/RemoteControl/Source/RemoteControlUI/Private/UI/Action/SRCActionPanel.cpp
--- a/Engine/Plugins/VirtualProduction/RemoteControl/Source/RemoteControlUI/Private/UI/Action/SRCActionPanel.cpp
+++ b/Engine/Plugins/VirtualProduction/RemoteControl/Source/RemoteControlUI/Private/UI/Action/SRCActionPanel.cpp
@@ -111,9 +111,6 @@ void SRCActionPanel::UpdateWrappedWidget(TSharedPtr<FRCBehaviourModel> InBehavio
 		ActionPanelList = InBehaviourItem->GetActionsListWidget(SharedThis(this));
 		ActionDockPanel->SetContent(ActionPanelList.ToSharedRef());
 
-		// Update add new action menu content
-		AddNewActionButton->SetMenuContent(GetActionMenuContentWidget());
-
 		const bool bIsBehaviourEnabled = InBehaviourItem->IsBehaviourEnabled();
 		InBehaviourItem->RefreshIsBehaviourEnabled(bIsBehaviourEnabled);
 		RefreshIsBehaviourEnabled(bIsBehaviourEnabled);
@@ -293,6 +290,7 @@ void SRCActionPanel::CreateActionPanelWidgets()
 			[
 				GetActionMenuContentWidget()
 			];
+		AddNewActionButton->SetOnGetMenuContent(FOnGetContent::CreateSP(this, &SRCActionPanel::GetActionMenuContentWidget));
 	}
 
 	// Add all button
