[RemoteControl] Fix direct bind controller creation with BP properties
Directly creating a bind controller dragging a BP float will cause the start value of the controller to mismatch the one of the float

diff --git a/Engine/Plugins/VirtualProduction/RemoteControl/Source/RemoteControlLogic/Private/Behaviour/Builtin/Bind/RCBehaviourBind.cpp b/Engine/Plugins/VirtualProduction/RemoteControl/Source/RemoteControlLogic/Private/Behaviour/Builtin/Bind/RCBehaviourBind.cpp
--- a/Engine/Plugins/VirtualProduction/RemoteControl/Source/RemoteControlLogic/Private/Behaviour/Builtin/Bind/RCBehaviourBind.cpp
+++ b/Engine/Plugins/VirtualProduction/RemoteControl/Source/RemoteControlLogic/Private/Behaviour/Builtin/Bind/RCBehaviourBind.cpp
@@ -450,10 +450,19 @@ bool URCBehaviourBind::CopyPropertyValueToController(URCController* InController
 	// Float Controller
 	else if (ControllerAsProperty->IsA(FFloatProperty::StaticClass()))
 	{
-		float FloatValue;
-		PropertyHandle->GetValue(FloatValue);
-
-		InController->SetValueFloat(FloatValue);
+		// check also the property handle type since double can be bound to float controllers
+		if (RemoteControlProperty->IsA(FFloatProperty::StaticClass()))
+		{
+			float FloatValue;
+			PropertyHandle->GetValue(FloatValue);
+			InController->SetValueFloat(FloatValue);
+		}
+		else if (RemoteControlProperty->IsA(FDoubleProperty::StaticClass()))
+		{
+			double DoubleValue;
+			PropertyHandle->GetValue(DoubleValue);
+			InController->SetValueFloat(DoubleValue);
+		}
 
 		return true;
 	}
