[RemoteControl] Path Behavior fixes
Fix path preview not updating when switching from Internal/External
Automatically remove ending slashes (/) from path entries if not RCInput

diff --git a/Engine/Plugins/VirtualProduction/RemoteControl/Source/RemoteControlUI/Private/UI/Behaviour/Builtin/Path/RCBehaviourSetAssetByPathModel.h b/Engine/Plugins/VirtualProduction/RemoteControl/Source/RemoteControlUI/Private/UI/Behaviour/Builtin/Path/RCBehaviourSetAssetByPathModel.h
--- a/Engine/Plugins/VirtualProduction/RemoteControl/Source/RemoteControlUI/Private/UI/Behaviour/Builtin/Path/RCBehaviourSetAssetByPathModel.h
+++ b/Engine/Plugins/VirtualProduction/RemoteControl/Source/RemoteControlUI/Private/UI/Behaviour/Builtin/Path/RCBehaviourSetAssetByPathModel.h
@@ -45,6 +45,9 @@ public:
 	/** Called when a controller value changed */
 	virtual void NotifyControllerValueChanged(TSharedPtr<FRCControllerModel> InControllerModel) override;
 
+	/** Refresh the PreviewPath */
+	void RefreshPreview() const;
+
 protected:
 	/** Gets the currently selected Entities Name, or a failure Text in case it fails. */
 	FText GetSelectedEntityText() const;
@@ -53,6 +56,11 @@ private:
 	/** Callback called when the AssetPath properties change to either refresh or reconstruct it updating the path values */
 	void OnAssetPathFinishedChangingProperties(const FPropertyChangedEvent& InEvent);
 
+	/** Reconstruct the PropertyRow Generator */
+	void RegenerateWeakPtrInternal();
+	
+	/** Regenerates and creates a new PathArray Widget if changed */
+	void RefreshPathAndPreview();
 private:
 	/** The SetAssetByPath Behaviour associated with this Model */
 	TWeakObjectPtr<URCSetAssetByPathBehaviour> SetAssetByPathBehaviourWeakPtr;
@@ -80,12 +88,4 @@ private:
 	
 	/** Array of Exposed Entities from the Preset to help create a Widget to choose the Target */
 	TArray<TWeakPtr<const FRemoteControlEntity>> ExposedEntities;
-	
-private:
-	void RegenerateWeakPtrInternal();
-	
-	/** Regenerates and creates a new PathArray Widget if changed */
-	void RefreshPathAndPreview();
-
-	void RefreshPreview() const;
 };
diff --git a/Engine/Plugins/VirtualProduction/RemoteControl/Source/RemoteControlUI/Private/UI/Behaviour/Builtin/Path/SRCBehaviourSetAssetByPath.cpp b/Engine/Plugins/VirtualProduction/RemoteControl/Source/RemoteControlUI/Private/UI/Behaviour/Builtin/Path/SRCBehaviourSetAssetByPath.cpp
--- a/Engine/Plugins/VirtualProduction/RemoteControl/Source/RemoteControlUI/Private/UI/Behaviour/Builtin/Path/SRCBehaviourSetAssetByPath.cpp
+++ b/Engine/Plugins/VirtualProduction/RemoteControl/Source/RemoteControlUI/Private/UI/Behaviour/Builtin/Path/SRCBehaviourSetAssetByPath.cpp
@@ -19,15 +19,9 @@ BEGIN_SLATE_FUNCTION_BUILD_OPTIMIZATION
 
 void SRCBehaviourSetAssetByPath::Construct(const FArguments& InArgs, TSharedRef<FRCSetAssetByPathBehaviourModel> InBehaviourItem)
 {
-	URCSetAssetByPathBehaviour* PathBehaviour = Cast<URCSetAssetByPathBehaviour>(InBehaviourItem->GetBehaviour());
-	const TSharedRef<SHorizontalBox> SelectedClassWidget = SNew(SHorizontalBox);
-	const TSharedRef<SHorizontalBox> InternalExternalSwitchWidget = SNew(SHorizontalBox)
-		.Visibility_Lambda([PathBehaviour]()
-		{
-				return PathBehaviour->AssetClass->GetFName() == "Texture" ? EVisibility::Visible : EVisibility::Collapsed;
-		});
-
-	if (!PathBehaviour)
+	SetAssetByPathWeakPtr = InBehaviourItem;
+	PathBehaviour = Cast<URCSetAssetByPathBehaviour>(InBehaviourItem->GetBehaviour());
+	if (!PathBehaviour.IsValid())
 	{
 		ChildSlot
 		[
@@ -35,8 +29,11 @@ void SRCBehaviourSetAssetByPath::Construct(const FArguments& InArgs, TSharedRef<
 		];
 		return;
 	}
-		
-	
+
+	const TSharedRef<SHorizontalBox> SelectedClassWidget = SNew(SHorizontalBox);
+	const TSharedRef<SHorizontalBox> InternalExternalSwitchWidget = SNew(SHorizontalBox)
+		.Visibility(this, &SRCBehaviourSetAssetByPath::GetInternalExternalVisibility);
+
 	for (UClass* SupportedClass : PathBehaviour->GetSupportedClasses())
 	{
 		SelectedClassWidget->AddSlot()
@@ -48,15 +45,8 @@ void SRCBehaviourSetAssetByPath::Construct(const FArguments& InArgs, TSharedRef<
 		[
 			SNew(SButton)
 			.Text(FText::FromString(SupportedClass->GetName()))
-			.ButtonColorAndOpacity_Lambda([PathBehaviour, SupportedClass]()
-			{
-				FSlateColor Color = PathBehaviour->AssetClass == SupportedClass ? FAppStyle::Get().GetSlateColor("Colors.Select") : FAppStyle::Get().GetSlateColor("Colors.White");
-				return Color;
-			})
-			.OnPressed_Lambda([PathBehaviour, SupportedClass]()
-			{
-				PathBehaviour->AssetClass = SupportedClass;
-			})
+			.ButtonColorAndOpacity(this, &SRCBehaviourSetAssetByPath::GetSelectedClassWidgetColor, SupportedClass)
+			.OnPressed(this, &SRCBehaviourSetAssetByPath::OnSelectedClassWidgetPressed, SupportedClass)
 		];
 	}
 		
@@ -69,15 +59,8 @@ void SRCBehaviourSetAssetByPath::Construct(const FArguments& InArgs, TSharedRef<
 	[
 		SNew(SButton)
 		.Text(FText::FromString(FString("Internal")))
-		.ButtonColorAndOpacity_Lambda([PathBehaviour]()
-		{
-			FSlateColor Color = PathBehaviour->bInternal ? FAppStyle::Get().GetSlateColor("Colors.Highlight") : FAppStyle::Get().GetSlateColor("Colors.AccentWhite");
-			return Color;
-		})
-		.OnPressed_Lambda([PathBehaviour]()
-		{
-			PathBehaviour->bInternal = true;
-		})
+		.ButtonColorAndOpacity(this, &SRCBehaviourSetAssetByPath::GetInternalWidgetColor)
+		.OnPressed(this, &SRCBehaviourSetAssetByPath::OnInternalExternalSwitchWidgetPressed, true)
 	];
 	
 	InternalExternalSwitchWidget->AddSlot()
@@ -89,15 +72,8 @@ void SRCBehaviourSetAssetByPath::Construct(const FArguments& InArgs, TSharedRef<
 	[
 		SNew(SButton)
 		.Text(FText::FromString(FString("External")))
-		.ButtonColorAndOpacity_Lambda([PathBehaviour]()
-		{
-			FSlateColor Color = PathBehaviour->bInternal ? FAppStyle::Get().GetSlateColor("Colors.AccentWhite") : FAppStyle::Get().GetSlateColor("Colors.Highlight");
-			return Color;
-		})
-		.OnPressed_Lambda([PathBehaviour]()
-		{
-			PathBehaviour->bInternal = false;
-		})
+		.ButtonColorAndOpacity(this, &SRCBehaviourSetAssetByPath::GetExternalWidgetColor)
+		.OnPressed(this, &SRCBehaviourSetAssetByPath::OnInternalExternalSwitchWidgetPressed, false)
 	];
 	
 	ChildSlot
@@ -120,6 +96,67 @@ void SRCBehaviourSetAssetByPath::Construct(const FArguments& InArgs, TSharedRef<
 	];
 }
 
+void SRCBehaviourSetAssetByPath::OnInternalExternalSwitchWidgetPressed(bool bInIsInternal) const
+{
+	if (!PathBehaviour.IsValid())
+	{
+		return;
+	}
+
+	PathBehaviour->bInternal = bInIsInternal;
+	if (const TSharedPtr<const FRCSetAssetByPathBehaviourModel>& SetAssetPath = SetAssetByPathWeakPtr.Pin())
+	{
+		SetAssetPath->RefreshPreview();
+	}
+}
+
+FSlateColor SRCBehaviourSetAssetByPath::GetExternalWidgetColor() const
+{
+	if (!PathBehaviour.IsValid())
+	{
+		return FAppStyle::Get().GetSlateColor("Colors.AccentRed");
+	}
+	return PathBehaviour->bInternal ? FAppStyle::Get().GetSlateColor("Colors.AccentWhite") : FAppStyle::Get().GetSlateColor("Colors.Highlight");
+}
+
+FSlateColor SRCBehaviourSetAssetByPath::GetInternalWidgetColor() const
+{
+	if (!PathBehaviour.IsValid())
+	{
+		return FAppStyle::Get().GetSlateColor("Colors.AccentRed");
+	}
+	return PathBehaviour->bInternal ? FAppStyle::Get().GetSlateColor("Colors.Highlight") : FAppStyle::Get().GetSlateColor("Colors.AccentWhite");
+}
+
+EVisibility SRCBehaviourSetAssetByPath::GetInternalExternalVisibility() const
+{
+	if (!PathBehaviour->AssetClass)
+	{
+		return EVisibility::Collapsed;
+	}
+	return PathBehaviour->AssetClass->GetFName() == TEXT("Texture") ? EVisibility::Visible : EVisibility::Collapsed;
+}
+
+FSlateColor SRCBehaviourSetAssetByPath::GetSelectedClassWidgetColor(UClass* InClass) const
+{
+	if (!PathBehaviour.IsValid())
+	{
+		return FAppStyle::Get().GetSlateColor("Colors.AccentRed");
+	}
+	return PathBehaviour->AssetClass == InClass ? FAppStyle::Get().GetSlateColor("Colors.Select") : FAppStyle::Get().GetSlateColor("Colors.White");
+}
+
+void SRCBehaviourSetAssetByPath::OnSelectedClassWidgetPressed(UClass* InClass) const
+{
+	if (!PathBehaviour.IsValid())
+	{
+		return;
+	}
+	PathBehaviour->AssetClass = InClass;
+	OnInternalExternalSwitchWidgetPressed(true);
+}
+
+
 END_SLATE_FUNCTION_BUILD_OPTIMIZATION
 
 #undef LOCTEXT_NAMESPACE
diff --git a/Engine/Plugins/VirtualProduction/RemoteControl/Source/RemoteControlUI/Private/UI/Behaviour/Builtin/Path/SRCBehaviourSetAssetByPath.h b/Engine/Plugins/VirtualProduction/RemoteControl/Source/RemoteControlUI/Private/UI/Behaviour/Builtin/Path/SRCBehaviourSetAssetByPath.h
--- a/Engine/Plugins/VirtualProduction/RemoteControl/Source/RemoteControlUI/Private/UI/Behaviour/Builtin/Path/SRCBehaviourSetAssetByPath.h
+++ b/Engine/Plugins/VirtualProduction/RemoteControl/Source/RemoteControlUI/Private/UI/Behaviour/Builtin/Path/SRCBehaviourSetAssetByPath.h
@@ -2,10 +2,14 @@
 
 #pragma once
 
-#include "CoreMinimal.h"
+#include "Styling/SlateColor.h"
+#include "Templates/SharedPointer.h"
+#include "UObject/WeakObjectPtrTemplates.h"
 #include "Widgets/SCompoundWidget.h"
 
 class FRCSetAssetByPathBehaviourModel;
+class UClass;
+class URCSetAssetByPathBehaviour;
 
 class REMOTECONTROLUI_API SRCBehaviourSetAssetByPath : public  SCompoundWidget
 {
@@ -20,6 +24,28 @@ public:
 	void Construct(const FArguments& InArgs, TSharedRef<FRCSetAssetByPathBehaviourModel> InBehaviourItem);
 
 private:
-	/** The Behaviour (UI model) associated with us*/
+	/** Called when the Internal/External button are pressed */
+	void OnInternalExternalSwitchWidgetPressed(bool bInIsInternal) const;
+
+	/** Return the color of the External widget */
+	FSlateColor GetExternalWidgetColor() const;
+
+	/** Return the color of the Internal widget */
+	FSlateColor GetInternalWidgetColor() const;
+
+	/** Return the visibility of the InternalExternal widget */
+	EVisibility GetInternalExternalVisibility() const;
+
+	/** Return the color of the SelectedClass widget */
+	FSlateColor GetSelectedClassWidgetColor(UClass* InClass) const;
+
+	/** Called when the SelectedClass buttons are pressed */
+	void OnSelectedClassWidgetPressed(UClass* InClass) const;
+
+private:
+	/** The Behaviour (UI model) associated with us */
 	TWeakPtr<const FRCSetAssetByPathBehaviourModel> SetAssetByPathWeakPtr;
+
+	/** The PathBehaviour associated with us */
+	TWeakObjectPtr<URCSetAssetByPathBehaviour> PathBehaviour;
 };
diff --git a/Engine/Plugins/VirtualProduction/RemoteControl/Source/RemoteControlUI/Private/UI/Customizations/RCAssetPathElementCustomization.cpp b/Engine/Plugins/VirtualProduction/RemoteControl/Source/RemoteControlUI/Private/UI/Customizations/RCAssetPathElementCustomization.cpp
--- a/Engine/Plugins/VirtualProduction/RemoteControl/Source/RemoteControlUI/Private/UI/Customizations/RCAssetPathElementCustomization.cpp
+++ b/Engine/Plugins/VirtualProduction/RemoteControl/Source/RemoteControlUI/Private/UI/Customizations/RCAssetPathElementCustomization.cpp
@@ -37,6 +37,8 @@ void FRCAssetPathElementCustomization::CustomizeHeader(TSharedRef<IPropertyHandl
 		return;
 	}
 
+	PathHandle->SetOnPropertyValueChanged(FSimpleDelegate::CreateSP(this, &FRCAssetPathElementCustomization::OnPathChanged));
+
 	const TSharedRef<SToolTip> GetAssetPathToolTipWidget = SNew(SToolTip)
 		.Text(LOCTEXT("RCGetAssetPathButton_Tooltip", "Get the path of the currently first selected asset in the content browser and set it to the current path"));
 
@@ -125,6 +127,24 @@ void FRCAssetPathElementCustomization::CustomizeChildren(TSharedRef<IPropertyHan
 {
 }
 
+void FRCAssetPathElementCustomization::OnPathChanged() const
+{
+	if (!PathHandle.IsValid() || !IsInputHandle.IsValid())
+	{
+		return;
+	}
+
+	bool bIsRCInput = false;
+	IsInputHandle->GetValue(bIsRCInput);
+
+	if (bIsRCInput)
+	{
+		return;
+	}
+
+	RemoveSlashFromPathEnd();
+}
+
 ECheckBoxState FRCAssetPathElementCustomization::IsChecked() const
 {
 	ECheckBoxState ReturnValue = ECheckBoxState::Undetermined;
@@ -146,7 +166,13 @@ void FRCAssetPathElementCustomization::OnCheckStateChanged(ECheckBoxState InNewS
 	{
 		return;
 	}
-	IsInputHandle->SetValue(InNewState == ECheckBoxState::Checked ? true : false);
+	const bool bIsInput = InNewState == ECheckBoxState::Checked;
+	IsInputHandle->SetValue(bIsInput);
+
+	if (!bIsInput)
+	{
+		RemoveSlashFromPathEnd();
+	}
 }
 
 FReply FRCAssetPathElementCustomization::OnGetAssetFromSelectionClicked() const
@@ -213,4 +239,19 @@ int32 FRCAssetPathElementCustomization::OnGetWidgetSwitcherIndex() const
 	return IsChecked() == ECheckBoxState::Checked ? 1 : 0;
 }
 
+void FRCAssetPathElementCustomization::RemoveSlashFromPathEnd() const
+{
+	FString CurrentPath;
+	PathHandle->GetValue(CurrentPath);
+
+	// remove all / that are at the end of the path
+	bool bRemoved = false;
+	do
+	{
+		bRemoved = CurrentPath.RemoveFromEnd(TEXT("/"), 1);
+	} while (bRemoved);
+
+	PathHandle->SetValue(CurrentPath, EPropertyValueSetFlags::NotTransactable);
+}
+
 #undef LOCTEXT_NAMESPACE
diff --git a/Engine/Plugins/VirtualProduction/RemoteControl/Source/RemoteControlUI/Private/UI/Customizations/RCAssetPathElementCustomization.h b/Engine/Plugins/VirtualProduction/RemoteControl/Source/RemoteControlUI/Private/UI/Customizations/RCAssetPathElementCustomization.h
--- a/Engine/Plugins/VirtualProduction/RemoteControl/Source/RemoteControlUI/Private/UI/Customizations/RCAssetPathElementCustomization.h
+++ b/Engine/Plugins/VirtualProduction/RemoteControl/Source/RemoteControlUI/Private/UI/Customizations/RCAssetPathElementCustomization.h
@@ -21,17 +21,27 @@ public:
 	//~ End IPropertyTypeCustomization
 
 private:
+	/** Called when the path change */
+	void OnPathChanged() const;
+
 	/** Returns if the CheckBox is checked or not */
 	ECheckBoxState IsChecked() const;
+
 	/** Callback called when you click on the CheckBox to enable/disable it */
 	void OnCheckStateChanged(ECheckBoxState InNewState) const;
+
 	/** Callback called when you click on the arrow button next to the entries to retrieve the path of the selected Asset */
 	FReply OnGetAssetFromSelectionClicked() const;
+
 	/** Callback called when you click on the plus button next to the entries to create a controller associated with the current RC Input entry */
 	FReply OnCreateControllerButtonClicked() const;
+
 	/** Retrieve the current widget switcher index */
 	int32 OnGetWidgetSwitcherIndex() const;
 
+	/** Remove all / at the end of the path */
+	void RemoveSlashFromPathEnd() const;
+
 private:
 	TSharedPtr<SWidget> PathWidget;
 	TSharedPtr<IPropertyUtilities> PropertyUtilities;
