[RemoteControl] Fix Blueprint issue
Reset owner object of RC properties when Blueprint reconstruct themselves

diff --git a/Engine/Plugins/VirtualProduction/RemoteControl/Source/RemoteControlUI/Private/UI/RCPanelWidgetRegistry.cpp b/Engine/Plugins/VirtualProduction/RemoteControl/Source/RemoteControlUI/Private/UI/RCPanelWidgetRegistry.cpp
--- a/Engine/Plugins/VirtualProduction/RemoteControl/Source/RemoteControlUI/Private/UI/RCPanelWidgetRegistry.cpp
+++ b/Engine/Plugins/VirtualProduction/RemoteControl/Source/RemoteControlUI/Private/UI/RCPanelWidgetRegistry.cpp
@@ -211,6 +211,21 @@ void FRCPanelWidgetRegistry::Refresh(const TSharedPtr<FStructOnScope>& InStruct)
 	}
 }
 
+void FRCPanelWidgetRegistry::ReplaceGeneratorObject(UObject* InOldObject, UObject* InNewObject)
+{
+	for (TPair<TWeakObjectPtr<UObject>, TSharedPtr<IPropertyRowGenerator>> Generator : ObjectToRowGenerator)
+	{
+		if (InOldObject && Generator.Value.IsValid() && Generator.Key != InNewObject)
+		{
+			if (Generator.Key == InOldObject)
+			{
+				Generator.Key = InNewObject;
+				Generator.Value->SetObjects({InNewObject});
+			}
+		}
+	}
+}
+
 void FRCPanelWidgetRegistry::Clear()
 {
 	for (const TPair <TWeakObjectPtr<UObject>, TSharedPtr<IPropertyRowGenerator>>& Pair : ObjectToRowGenerator)
@@ -323,4 +338,3 @@ TSharedPtr<IPropertyRowGenerator> FRCPanelWidgetRegistry::CreateGenerator(UObjec
 
 	return Generator;
 }
-
diff --git a/Engine/Plugins/VirtualProduction/RemoteControl/Source/RemoteControlUI/Private/UI/RCPanelWidgetRegistry.h b/Engine/Plugins/VirtualProduction/RemoteControl/Source/RemoteControlUI/Private/UI/RCPanelWidgetRegistry.h
--- a/Engine/Plugins/VirtualProduction/RemoteControl/Source/RemoteControlUI/Private/UI/RCPanelWidgetRegistry.h
+++ b/Engine/Plugins/VirtualProduction/RemoteControl/Source/RemoteControlUI/Private/UI/RCPanelWidgetRegistry.h
@@ -31,7 +31,6 @@ struct FRCTreeNodeFinderHandler
 class FRCPanelWidgetRegistry : public TSharedFromThis<FRCPanelWidgetRegistry>
 {
 public:
-	
 	FRCPanelWidgetRegistry();
 	~FRCPanelWidgetRegistry();
 
@@ -61,6 +60,11 @@ public:
 	 */
 	void Refresh(const TSharedPtr<FStructOnScope>& InStruct);
 
+	/**
+	 * Replace old object if found with new object
+	 */
+	void ReplaceGeneratorObject(UObject* InOldObject, UObject* InNewObject);
+
 	/**
 	 * Clear the registry and its cache.
 	 */
diff --git a/Engine/Plugins/VirtualProduction/RemoteControl/Source/RemoteControlUI/Private/UI/SRemoteControlPanel.cpp b/Engine/Plugins/VirtualProduction/RemoteControl/Source/RemoteControlUI/Private/UI/SRemoteControlPanel.cpp
--- a/Engine/Plugins/VirtualProduction/RemoteControl/Source/RemoteControlUI/Private/UI/SRemoteControlPanel.cpp
+++ b/Engine/Plugins/VirtualProduction/RemoteControl/Source/RemoteControlUI/Private/UI/SRemoteControlPanel.cpp
@@ -1444,6 +1444,22 @@ void SRemoteControlPanel::BindRemoteControlCommands()
 		FCanExecuteAction::CreateSP(this, &SRemoteControlPanel::CanUpdateValue));
 }
 
+void SRemoteControlPanel::OnObjectReplaced(const TMap<UObject*, UObject*>& InObjectReplaced)
+{
+	if (!WidgetRegistry.IsValid())
+	{
+		return;
+	}
+
+	for (const TPair<UObject*, UObject*> ReplacedObject : InObjectReplaced)
+	{
+		if (ReplacedObject.Key && ReplacedObject.Value)
+		{
+			WidgetRegistry->ReplaceGeneratorObject(ReplacedObject.Key, ReplacedObject.Value);
+		}
+	}
+}
+
 void SRemoteControlPanel::RegisterEvents()
 {
 	FEditorDelegates::MapChange.AddSP(this, &SRemoteControlPanel::OnMapChange);
@@ -1464,6 +1480,7 @@ void SRemoteControlPanel::RegisterEvents()
 	Preset->OnEntityUnexposed().AddSP(this, &SRemoteControlPanel::OnEntityUnexposed);
 
 	UMaterial::OnMaterialCompilationFinished().AddSP(this, &SRemoteControlPanel::OnMaterialCompiled);
+	FCoreUObjectDelegates::OnObjectsReplaced.AddSP(this, &SRemoteControlPanel::OnObjectReplaced);
 }
 
 void SRemoteControlPanel::UnregisterEvents()
@@ -1484,8 +1501,8 @@ void SRemoteControlPanel::UnregisterEvents()
 	}
 
 	FEditorDelegates::MapChange.RemoveAll(this);
-
 	UMaterial::OnMaterialCompilationFinished().RemoveAll(this);
+	FCoreUObjectDelegates::OnObjectsReplaced.RemoveAll(this);
 }
 
 void SRemoteControlPanel::RegisterPanels()
diff --git a/Engine/Plugins/VirtualProduction/RemoteControl/Source/RemoteControlUI/Private/UI/SRemoteControlPanel.h b/Engine/Plugins/VirtualProduction/RemoteControl/Source/RemoteControlUI/Private/UI/SRemoteControlPanel.h
--- a/Engine/Plugins/VirtualProduction/RemoteControl/Source/RemoteControlUI/Private/UI/SRemoteControlPanel.h
+++ b/Engine/Plugins/VirtualProduction/RemoteControl/Source/RemoteControlUI/Private/UI/SRemoteControlPanel.h
@@ -176,6 +176,9 @@ private:
 	//~ Remote Control Commands
 	void BindRemoteControlCommands();
 
+	/** Called when object are replaced, blueprint for example */
+	void OnObjectReplaced(const TMap<UObject*, UObject*>& InObjectReplaced);
+
 	/** Register editor events needed to handle reloading objects and blueprint libraries. */
 	void RegisterEvents();
 	/** Unregister editor events */
