[Remote Control] Fix for properties with setter through nDisplay interception.

Details:

Properties with setters have been added in CL 20273706. It has been added by extending InvokeCall to support FProperty::CallSetter path. Prior to this, Setters where a UFunction, so it was fine to intercept the InvokeCall. After the change, it would still intercept InvokeCall, but lose the information about which property is being set, since there is no UFunction.

Fix:

CanInterceptFunction was added in CL 24562094 (to fix something else), but it correctly makes InvokeCall that can't be intercepted fallback to intercepting SetObjectProperties, which works correctly with the CallSetter path.

diff --git a/Engine/Plugins/VirtualProduction/RemoteControl/Source/RemoteControl/Private/RemoteControlModule.cpp b/Engine/Plugins/VirtualProduction/RemoteControl/Source/RemoteControl/Private/RemoteControlModule.cpp
--- a/Engine/Plugins/VirtualProduction/RemoteControl/Source/RemoteControl/Private/RemoteControlModule.cpp
+++ b/Engine/Plugins/VirtualProduction/RemoteControl/Source/RemoteControl/Private/RemoteControlModule.cpp
@@ -2728,6 +2728,11 @@ bool FRemoteControlModule::CanInterceptFunction(const FRCCall& RCCall) const
 {
 	if (RCCall.IsValid())
 	{
+		if (!RCCall.CallRef.Function.IsValid())
+		{
+			return false;
+		}
+
 		for (TFieldIterator<FProperty> It(RCCall.CallRef.Function.Get()); It; ++It)
 		{
 			// At the moment interceptors do not support return values.
