[RemoteControl] Add more support for PIE
Adding support to change target world in RC and control property also in other worlds

diff --git a/Engine/Plugins/VirtualProduction/RemoteControl/Source/RemoteControl/Private/RemoteControlBinding.cpp b/Engine/Plugins/VirtualProduction/RemoteControl/Source/RemoteControl/Private/RemoteControlBinding.cpp
--- a/Engine/Plugins/VirtualProduction/RemoteControl/Source/RemoteControl/Private/RemoteControlBinding.cpp
+++ b/Engine/Plugins/VirtualProduction/RemoteControl/Source/RemoteControl/Private/RemoteControlBinding.cpp
@@ -12,6 +12,7 @@
 #include "IRemoteControlModule.h"
 #include "RemoteControlPreset.h"
 #include "RemoteControlSettings.h"
+#include "UObject/Package.h"
 #include "UObject/SoftObjectPath.h"
 #include "UObject/SoftObjectPtr.h"
 
@@ -85,41 +86,30 @@ namespace UE::RemoteControlBinding
 namespace
 {
 	/**
-	 * What world are we looking in to find the counterpart actor/component
-	 */
-	enum class ECounterpartWorldTarget
-	{
-		Editor,
-		PIE
-	};
-
-	/**
-	 * Find the counterpart actor/component in PIE/Editor
+	 * Find the Object in Editor
 	 * 
 	 * @TODO Does this need to be updated to possibly include a non-editor related preset?
 	 */
-	UObject* FindObjectInCounterpartWorld(UObject* Object, ECounterpartWorldTarget WorldTarget)
+	UObject* FindObjectInEditorWorld(UObject* InObject)
 	{
 		UObject* CounterpartObject = nullptr;
 #if WITH_EDITOR
-		if (Object && GEditor)
+		if (InObject && GEditor)
 		{
-			const bool bForPie = WorldTarget == ECounterpartWorldTarget::PIE ? true : false;
-			
-			if (AActor* Actor = Cast<AActor>(Object))
+			if (AActor* Actor = Cast<AActor>(InObject))
 			{
-				CounterpartObject = bForPie ? EditorUtilities::GetSimWorldCounterpartActor(Actor) : EditorUtilities::GetEditorWorldCounterpartActor(Actor);
+				CounterpartObject = EditorUtilities::GetEditorWorldCounterpartActor(Actor);
 			}
-			else if(AActor* Owner = Object->GetTypedOuter<AActor>())
+			else if(AActor* Owner = InObject->GetTypedOuter<AActor>())
 			{
-				if (AActor* CounterpartWorldOwner = bForPie ? EditorUtilities::GetSimWorldCounterpartActor(Owner) : EditorUtilities::GetEditorWorldCounterpartActor(Owner))
+				if (AActor* CounterpartWorldOwner = EditorUtilities::GetEditorWorldCounterpartActor(Owner))
 				{
-					CounterpartObject = FindObject<UObject>(CounterpartWorldOwner, *Object->GetName());
+					CounterpartObject = FindObject<UObject>(CounterpartWorldOwner, *InObject->GetName());
 				}
 			}
 		}
 #endif
-		return CounterpartObject ? CounterpartObject : Object;
+		return CounterpartObject ? CounterpartObject : InObject;
 	}
 
 	void DumpBindings(const TArray<FString>& Args)
@@ -170,6 +160,11 @@ UObject* URemoteControlLevelIndependantBinding::Resolve() const
 	return BoundObject.Get();
 }
 
+UObject* URemoteControlLevelIndependantBinding::ResolveForWorld(const UWorld* InWorld) const
+{
+	return BoundObject.Get();
+}
+
 bool URemoteControlLevelIndependantBinding::IsValid() const
 {
 	return BoundObject.IsValid();
@@ -205,8 +200,8 @@ void URemoteControlLevelDependantBinding::SetBoundObject(const TSoftObjectPtr<UO
 {
 	if (ensure(InObject))
 	{
-		UObject* EditorObject = FindObjectInCounterpartWorld(InObject.Get(), ECounterpartWorldTarget::Editor);
-		ULevel* OuterLevel = EditorObject->GetTypedOuter<ULevel>();
+		const UObject* EditorObject = FindObjectInEditorWorld(InObject.Get());
+		const ULevel* OuterLevel = EditorObject->GetTypedOuter<ULevel>();
 		BoundObjectMapByPath.FindOrAdd(OuterLevel) = EditorObject;
 		SubLevelSelectionMapByPath.FindOrAdd(OuterLevel->GetWorld()) = OuterLevel;
 		
@@ -271,7 +266,7 @@ UObject* URemoteControlLevelDependantBinding::Resolve() const
 	bool bAllowPIE = true;
 
 	// Disallow PIE for Embedded Presets as there should be an Embedded Preset counterpart in PIE taking care of this
-	URemoteControlPreset* Preset = Cast<URemoteControlPreset>(GetOuter());
+	const URemoteControlPreset* Preset = Cast<URemoteControlPreset>(GetOuter());
 	if (Preset && Preset->IsEmbeddedPreset())
 	{
 		bAllowPIE = false;
@@ -287,7 +282,7 @@ UObject* URemoteControlLevelDependantBinding::Resolve() const
 			return Object;
 		}
 		// Make sure we don't resolve on a subobject of a dying parent actor.
-		if (AActor* OwnerActor = Object->GetTypedOuter<AActor>())
+		if (const AActor* OwnerActor = Object->GetTypedOuter<AActor>())
 		{
 			if (!::IsValid(OwnerActor))
 			{
@@ -307,13 +302,48 @@ UObject* URemoteControlLevelDependantBinding::Resolve() const
 		}
 	}
 
+#if WITH_EDITOR
 	if (bAllowPIE)
 	{
-		return FindObjectInCounterpartWorld(Object, ECounterpartWorldTarget::PIE);
+		if (Preset && Preset->SelectedWorld.IsValid())
+		{
+			if (UObject* ObjectToReturn = ResolveForWorld(Preset->SelectedWorld.Get()))
+			{
+				return ObjectToReturn;
+			}
+		}
 	}
+#endif
 	return Object;
 }
 
+UObject* URemoteControlLevelDependantBinding::ResolveForWorld(const UWorld* InWorld) const
+{
+	if (!InWorld)
+	{
+		return nullptr;
+	}
+
+	if (const UPackage* WorldPackage = InWorld->GetPackage())
+	{
+		const UObject* CurrentResolvedObject = ResolveForCurrentWorld(true).Get();
+		FSoftObjectPath ObjectPath = CurrentResolvedObject;
+		ObjectPath = FSoftObjectPath(WorldPackage->GetFName(), ObjectPath.GetAssetFName(), ObjectPath.GetSubPathString());
+
+		if (UObject* ResolvedObject = ObjectPath.ResolveObject())
+		{
+			return ResolvedObject;
+		}
+
+		if (UObject* LoadedObject = ObjectPath.TryLoad())
+		{
+			return LoadedObject;
+		}
+	}
+
+	return nullptr;
+}
+
 bool URemoteControlLevelDependantBinding::IsValid() const
 {
 	return BoundObjectMapByPath.Num() > 0;
diff --git a/Engine/Plugins/VirtualProduction/RemoteControl/Source/RemoteControl/Private/RemoteControlEntity.cpp b/Engine/Plugins/VirtualProduction/RemoteControl/Source/RemoteControl/Private/RemoteControlEntity.cpp
--- a/Engine/Plugins/VirtualProduction/RemoteControl/Source/RemoteControl/Private/RemoteControlEntity.cpp
+++ b/Engine/Plugins/VirtualProduction/RemoteControl/Source/RemoteControl/Private/RemoteControlEntity.cpp
@@ -30,6 +30,21 @@ UObject* FRemoteControlEntity::GetBoundObject() const
 	return Obj ? *Obj : nullptr;
 }
 
+UObject* FRemoteControlEntity::GetBoundObjectForWorld(const UWorld* InWorld) const
+{
+	if (Bindings.IsEmpty())
+	{
+		return nullptr;
+	}
+
+	if (!Bindings[0].IsValid() || Bindings[0].IsStale())
+	{
+		return nullptr;
+	}
+
+	return Bindings[0]->ResolveForWorld(InWorld);
+}
+
 const TArray<TWeakObjectPtr<URemoteControlBinding>>& FRemoteControlEntity::GetBindings() const
 {
 	return Bindings;
@@ -70,7 +85,6 @@ void FRemoteControlEntity::BindObject(UObject* InObjectToBind)
 		if (URemoteControlBinding* MatchingBinding = Owner->FindMatchingBinding(Bindings[0].Get(), InObjectToBind))
 		{
 			Bindings[0] = MatchingBinding;
-			
 		}
 		else
 		{
diff --git a/Engine/Plugins/VirtualProduction/RemoteControl/Source/RemoteControl/Public/RemoteControlBinding.h b/Engine/Plugins/VirtualProduction/RemoteControl/Source/RemoteControl/Public/RemoteControlBinding.h
--- a/Engine/Plugins/VirtualProduction/RemoteControl/Source/RemoteControl/Public/RemoteControlBinding.h
+++ b/Engine/Plugins/VirtualProduction/RemoteControl/Source/RemoteControl/Public/RemoteControlBinding.h
@@ -53,6 +53,13 @@ public:
 	 */
 	virtual UObject* Resolve() const PURE_VIRTUAL(URemoteControlBinding::Resolve, return nullptr;);
 
+	/**
+	 * Return the Entity Object in the given world
+	 * @param InWorld World where to retrieve the object
+	 * @return The object in the given world or nullptr if not found
+	 */
+	virtual UObject* ResolveForWorld(const UWorld* InWorld) const PURE_VIRTUAL(URemoteControlBinding::ResolveForWorld, return nullptr;);
+
 	/**
 	 * Whether this binding represents a valid object.
 	 */
@@ -103,6 +110,7 @@ public:
 	virtual void SetBoundObject(const TSoftObjectPtr<UObject>& InObject) override;
 	virtual void UnbindObject(const TSoftObjectPtr<UObject>& InBoundObject) override;
 	virtual UObject* Resolve() const override;
+	virtual UObject* ResolveForWorld(const UWorld* InWorld) const override;
 	virtual bool IsValid() const override;
 	virtual bool IsBound(const TSoftObjectPtr<UObject>& Object) const override;
 	virtual bool PruneDeletedObjects() override;
@@ -199,6 +207,7 @@ public:
 	virtual void SetBoundObject(const TSoftObjectPtr<UObject>& BoundObject) override;
 	virtual void UnbindObject(const TSoftObjectPtr<UObject>& BoundObject) override;
 	virtual UObject* Resolve() const override;
+	virtual UObject* ResolveForWorld(const UWorld* InWorld) const override;
 	virtual bool IsValid() const override;
 	virtual bool IsBound(const TSoftObjectPtr<UObject>& Object) const override;
 	virtual bool PruneDeletedObjects() override;
diff --git a/Engine/Plugins/VirtualProduction/RemoteControl/Source/RemoteControl/Public/RemoteControlEntity.h b/Engine/Plugins/VirtualProduction/RemoteControl/Source/RemoteControl/Public/RemoteControlEntity.h
--- a/Engine/Plugins/VirtualProduction/RemoteControl/Source/RemoteControl/Public/RemoteControlEntity.h
+++ b/Engine/Plugins/VirtualProduction/RemoteControl/Source/RemoteControl/Public/RemoteControlEntity.h
@@ -57,6 +57,13 @@ struct REMOTECONTROL_API FRemoteControlEntity
 	 */
 	UObject* GetBoundObject() const;
 
+	/**
+	 * Return the Entity Object in the given world
+	 * @param InWorld World where to retrieve the object
+	 * @return The object in the given world or nullptr if not found
+	 */
+	UObject* GetBoundObjectForWorld(const UWorld* InWorld) const;
+
 	/**
 	 * Get the unresolved bindings of this exposed entity.
 	 */
diff --git a/Engine/Plugins/VirtualProduction/RemoteControl/Source/RemoteControl/Public/RemoteControlPreset.h b/Engine/Plugins/VirtualProduction/RemoteControl/Source/RemoteControl/Public/RemoteControlPreset.h
--- a/Engine/Plugins/VirtualProduction/RemoteControl/Source/RemoteControl/Public/RemoteControlPreset.h
+++ b/Engine/Plugins/VirtualProduction/RemoteControl/Source/RemoteControl/Public/RemoteControlPreset.h
@@ -718,6 +718,10 @@ public:
 	UPROPERTY(EditAnywhere, Category = "Remote Control Preset")
 	TArray<TObjectPtr<URemoteControlBinding>> Bindings;
 
+#if WITH_EDITOR
+	/** Current selected world used only in the editor to let you switch world */
+	TWeakObjectPtr<const UWorld> SelectedWorld = nullptr;
+#endif
 	/** ~~~Virtual Property Wrapper Functions ~~~
 	* 
 	* The goal is to hide the Controller Container and provide a simple interface for Controller access to UI and Web.
diff --git a/Engine/Plugins/VirtualProduction/RemoteControl/Source/RemoteControlUI/Private/UI/MenuContext/RemoteControlPanelMenuContext.h b/Engine/Plugins/VirtualProduction/RemoteControl/Source/RemoteControlUI/Private/UI/MenuContext/RemoteControlPanelMenuContext.h
new file mode 100644
--- /dev/null
+++ b/Engine/Plugins/VirtualProduction/RemoteControl/Source/RemoteControlUI/Private/UI/MenuContext/RemoteControlPanelMenuContext.h
@@ -0,0 +1,17 @@
+// Copyright Epic Games, Inc. All Rights Reserved.
+#pragma once
+
+#include "UObject/Object.h"
+#include "RemoteControlPanelMenuContext.generated.h"
+
+class SRemoteControlPanel;
+
+UCLASS()
+class URemoteControlPanelMenuContext : public UObject
+{
+	GENERATED_BODY()
+
+public:
+	/** RemoteControlPanel of this context */
+	TWeakPtr<SRemoteControlPanel> RemoteControlPanel;
+};
\ No newline at end of file
diff --git a/Engine/Plugins/VirtualProduction/RemoteControl/Source/RemoteControlUI/Private/UI/RCPanelWidgetRegistry.cpp b/Engine/Plugins/VirtualProduction/RemoteControl/Source/RemoteControlUI/Private/UI/RCPanelWidgetRegistry.cpp
--- a/Engine/Plugins/VirtualProduction/RemoteControl/Source/RemoteControlUI/Private/UI/RCPanelWidgetRegistry.cpp
+++ b/Engine/Plugins/VirtualProduction/RemoteControl/Source/RemoteControlUI/Private/UI/RCPanelWidgetRegistry.cpp
@@ -211,9 +211,20 @@ void FRCPanelWidgetRegistry::Refresh(const TSharedPtr<FStructOnScope>& InStruct)
 	}
 }
 
+void FRCPanelWidgetRegistry::UpdateGeneratorAndTreeCache(UObject* InOldObject, UObject* InNewObject, FString InPath)
+{
+	ReplaceGeneratorObject(InOldObject, InNewObject);
+	ReplaceTreeCacheObject(InOldObject, InNewObject, InPath);
+}
+
 void FRCPanelWidgetRegistry::ReplaceGeneratorObject(UObject* InOldObject, UObject* InNewObject)
 {
-	for (TPair<TWeakObjectPtr<UObject>, TSharedPtr<IPropertyRowGenerator>> Generator : ObjectToRowGenerator)
+	if (!InOldObject || !InNewObject)
+	{
+		return;
+	}
+
+	for (TPair<TWeakObjectPtr<UObject>, TSharedPtr<IPropertyRowGenerator>>& Generator : ObjectToRowGenerator)
 	{
 		if (InOldObject && Generator.Value.IsValid() && Generator.Key != InNewObject)
 		{
@@ -226,6 +237,27 @@ void FRCPanelWidgetRegistry::ReplaceGeneratorObject(UObject* InOldObject, UObjec
 	}
 }
 
+void FRCPanelWidgetRegistry::ReplaceTreeCacheObject(UObject* InOldObject, UObject* InNewObject, FString InPath)
+{
+	if (!InOldObject || !InNewObject)
+	{
+		return;
+	}
+
+	const TPair<TWeakObjectPtr<UObject>, FString> OldCacheKey{InOldObject, InPath};
+	if (TreeNodeCache.Contains(OldCacheKey))
+	{
+		if (TSharedPtr<IPropertyRowGenerator>* Generator = ObjectToRowGenerator.Find(InNewObject))
+		{
+			const TPair<TWeakObjectPtr<UObject>, FString> NewCacheKey{InNewObject, InPath};
+			TreeNodeCache.Remove(OldCacheKey);
+			const TSharedPtr<IDetailTreeNode> Node = WidgetRegistryUtils::FindNode((*Generator)->GetRootTreeNodes(), InPath, ERCFindNodeMethod::Path);
+			// Cache the node to avoid having to do the recursive find again.
+			TreeNodeCache.Add(NewCacheKey, Node);
+		}
+	}
+}
+
 void FRCPanelWidgetRegistry::Clear()
 {
 	for (const TPair <TWeakObjectPtr<UObject>, TSharedPtr<IPropertyRowGenerator>>& Pair : ObjectToRowGenerator)
diff --git a/Engine/Plugins/VirtualProduction/RemoteControl/Source/RemoteControlUI/Private/UI/RCPanelWidgetRegistry.h b/Engine/Plugins/VirtualProduction/RemoteControl/Source/RemoteControlUI/Private/UI/RCPanelWidgetRegistry.h
--- a/Engine/Plugins/VirtualProduction/RemoteControl/Source/RemoteControlUI/Private/UI/RCPanelWidgetRegistry.h
+++ b/Engine/Plugins/VirtualProduction/RemoteControl/Source/RemoteControlUI/Private/UI/RCPanelWidgetRegistry.h
@@ -61,9 +61,9 @@ public:
 	void Refresh(const TSharedPtr<FStructOnScope>& InStruct);
 
 	/**
-	 * Replace old object if found with new object
+	 * Update the InOldOObject generator and the tree cache with the InNewObject.
 	 */
-	void ReplaceGeneratorObject(UObject* InOldObject, UObject* InNewObject);
+	void UpdateGeneratorAndTreeCache(UObject* InOldObject, UObject* InNewObject, FString InPath);
 
 	/**
 	 * Clear the registry and its cache.
@@ -86,7 +86,10 @@ private:
 	void OnRowsRefreshed(TSharedPtr<IPropertyRowGenerator> Generator);
 	/** Create a property row generator */
 	TSharedPtr<IPropertyRowGenerator> CreateGenerator(UObject* InObject);
-
+	/** Replace old object if found with new object */
+	void ReplaceGeneratorObject(UObject* InOldObject, UObject* InNewObject);
+	/** Update Cached tree, main use when switching between different world needs to be called after ReplaceGeneratorObject */
+	void ReplaceTreeCacheObject(UObject* InOldObject, UObject* InNewObject, FString InPath);
 private:
 	/** Map of objects to row generator, used to have one row generator per object. */
 	TMap<TWeakObjectPtr<UObject>, TSharedPtr<IPropertyRowGenerator>> ObjectToRowGenerator;
diff --git a/Engine/Plugins/VirtualProduction/RemoteControl/Source/RemoteControlUI/Private/UI/SRCPanelExposedField.cpp b/Engine/Plugins/VirtualProduction/RemoteControl/Source/RemoteControlUI/Private/UI/SRCPanelExposedField.cpp
--- a/Engine/Plugins/VirtualProduction/RemoteControl/Source/RemoteControlUI/Private/UI/SRCPanelExposedField.cpp
+++ b/Engine/Plugins/VirtualProduction/RemoteControl/Source/RemoteControlUI/Private/UI/SRCPanelExposedField.cpp
@@ -423,149 +423,152 @@ TSharedRef<SWidget> SRCPanelExposedField::ConstructWidget()
 	if (TSharedPtr<FRemoteControlField> Field = WeakField.Pin())
 	{
 		// For the moment, just use the first object.
-		UObject* Object = Field->GetBoundObject();
-		if (Object && GetFieldType() == EExposedFieldType::Property)
+		if (const URemoteControlPreset* RCPreset = Preset.Get();
+			RCPreset && RCPreset->SelectedWorld.IsValid())
 		{
-			if (TSharedPtr<FRCPanelWidgetRegistry> Registry = WidgetRegistry.Pin())
+			UObject* Object = Field->GetBoundObjectForWorld(RCPreset->SelectedWorld.Get());
+			if (Object && GetFieldType() == EExposedFieldType::Property)
 			{
-				TRACE_CPUPROFILER_EVENT_SCOPE(SRCPanelExposedField::BeforeGetObjectTreeNode);
-
-				if (TSharedPtr<IDetailTreeNode> Node = Registry->GetObjectTreeNode(Object, Field->FieldPathInfo.ToPathPropertyString(), ERCFindNodeMethod::Path))
+				if (TSharedPtr<FRCPanelWidgetRegistry> Registry = WidgetRegistry.Pin())
 				{
-					TSharedPtr<SWidget> ValueWidget = SNullWidget::NullWidget;
-					TSharedPtr<SWidget> EditConditionWidget = SNullWidget::NullWidget;
-					TSharedPtr<IPropertyHandle> PropertyHandle;
-					bool bIsChildPropertyWidgetCreated = false;
-
-					if (!ExposedFieldUtils::IsFieldPathMatchingWithNodePath(Node, Field))
-					{
-						bIsChildPropertyWidgetCreated = ExposedFieldUtils::TryCreateWidgetForChildProperty(Node, Field, PropertyHandle, ValueWidget, EditConditionWidget);
-					}
+					TRACE_CPUPROFILER_EVENT_SCOPE(SRCPanelExposedField::BeforeGetObjectTreeNode);
 
-					if (!bIsChildPropertyWidgetCreated)
+					if (TSharedPtr<IDetailTreeNode> Node = Registry->GetObjectTreeNode(Object, Field->FieldPathInfo.ToPathPropertyString(), ERCFindNodeMethod::Path))
 					{
-						TArray<TSharedRef<IDetailTreeNode>> ChildNodes;
-						Node->GetChildren(ChildNodes);
-						ChildWidgets.Reset(ChildNodes.Num());
+						TSharedPtr<SWidget> ValueWidget = SNullWidget::NullWidget;
+						TSharedPtr<SWidget> EditConditionWidget = SNullWidget::NullWidget;
+						TSharedPtr<IPropertyHandle> PropertyHandle;
+						bool bIsChildPropertyWidgetCreated = false;
 
-						for (const TSharedRef<IDetailTreeNode>& ChildNode : ChildNodes)
+						if (!ExposedFieldUtils::IsFieldPathMatchingWithNodePath(Node, Field))
 						{
-							ChildWidgets.Add(SNew(SRCPanelFieldChildNode, ChildNode, ColumnSizeData));
+							bIsChildPropertyWidgetCreated = ExposedFieldUtils::TryCreateWidgetForChildProperty(Node, Field, PropertyHandle, ValueWidget, EditConditionWidget);
 						}
 
-						ExposedFieldUtils::CreateWidgets(Node, ValueWidget, EditConditionWidget);
+						if (!bIsChildPropertyWidgetCreated)
+						{
+							TArray<TSharedRef<IDetailTreeNode>> ChildNodes;
+							Node->GetChildren(ChildNodes);
+							ChildWidgets.Reset(ChildNodes.Num());
 
-						PropertyHandle = Node->CreatePropertyHandle();
-					}
+							for (const TSharedRef<IDetailTreeNode>& ChildNode : ChildNodes)
+							{
+								ChildWidgets.Add(SNew(SRCPanelFieldChildNode, ChildNode, ColumnSizeData));
+							}
 
-					TSharedPtr<SHorizontalBox> FieldWidget = SNew(SHorizontalBox);
-					FieldWidget->AddSlot()
-					[
-						ValueWidget.ToSharedRef()
-					];
-					if (PropertyHandle && PropertyHandle->IsValidHandle() && (PropertyHandle->GetParentHandle()->GetPropertyDisplayName().ToString() == FString("Transform")
-						|| PropertyHandle->GetOuterBaseClass() == UMaterialInstanceDynamic::StaticClass()))
-					{
-						// Set up a Zeroed DefaultValue, in case an ExposedEntity doesn't have a native Default. Needed for certain ResetToDefault cases.
-						void* ValuePtr;
-						if (PropertyHandle->GetValueData(ValuePtr) == FPropertyAccess::Result::Success)
-						{
-							DefaultValue.Reset(new uint8[PropertyHandle->GetProperty()->GetSize()]);
-							PropertyHandle->GetProperty()->CopyCompleteValue(DefaultValue.Get(), ValuePtr);
-							PropertyHandle->GetProperty()->ClearValue(DefaultValue.Get());
+							ExposedFieldUtils::CreateWidgets(Node, ValueWidget, EditConditionWidget);
+
+							PropertyHandle = Node->CreatePropertyHandle();
 						}
 
-						TWeakPtr<SRCPanelExposedField> WeakFieldPtr = SharedThis(this);
-						auto IsVisible = [WeakFieldPtr, PropertyHandle]()
+						TSharedPtr<SHorizontalBox> FieldWidget = SNew(SHorizontalBox);
+						FieldWidget->AddSlot()
+						[
+							ValueWidget.ToSharedRef()
+						];
+						if (PropertyHandle && PropertyHandle->IsValidHandle() && (PropertyHandle->GetParentHandle()->GetPropertyDisplayName().ToString() == FString("Transform")
+							|| PropertyHandle->GetOuterBaseClass() == UMaterialInstanceDynamic::StaticClass()))
 						{
-							TSharedPtr<SRCPanelExposedField> FieldPtr = WeakFieldPtr.Pin();
-							if (FieldPtr && PropertyHandle && PropertyHandle->IsValidHandle())
+							// Set up a Zeroed DefaultValue, in case an ExposedEntity doesn't have a native Default. Needed for certain ResetToDefault cases.
+							void* ValuePtr;
+							if (PropertyHandle->GetValueData(ValuePtr) == FPropertyAccess::Result::Success)
 							{
-								void* DataPtr;
-								if (PropertyHandle->GetValueData(DataPtr) == FPropertyAccess::Result::Success)
-								{
-									FProperty* NodeProperty = PropertyHandle->GetProperty();
-									bool bVisible = !NodeProperty->Identical(FieldPtr->DefaultValue.Get(), DataPtr);
-									return bVisible ? EVisibility::Visible : EVisibility::Hidden;
-								}
+								DefaultValue.Reset(new uint8[PropertyHandle->GetProperty()->GetSize()]);
+								PropertyHandle->GetProperty()->CopyCompleteValue(DefaultValue.Get(), ValuePtr);
+								PropertyHandle->GetProperty()->ClearValue(DefaultValue.Get());
 							}
-							return EVisibility::Hidden;
-						};
-
-						ResetButtonWidget = SNew(SBox)
-							.HAlign(HAlign_Center)
-							.VAlign(VAlign_Center)
-							.Padding(2.f, 4.f)
-							[
-								SNew(SButton)
-								.IsFocusable(false)
-								.ButtonStyle(&RCPanelStyle->FlatButtonStyle)
-								.ContentPadding(RCPanelStyle->PanelPadding)
-								.Visibility_Lambda(IsVisible)
-								.OnClicked_Lambda([Node]()
+
+							TWeakPtr<SRCPanelExposedField> WeakFieldPtr = SharedThis(this);
+							auto IsVisible = [WeakFieldPtr, PropertyHandle]()
+							{
+								TSharedPtr<SRCPanelExposedField> FieldPtr = WeakFieldPtr.Pin();
+								if (FieldPtr && PropertyHandle && PropertyHandle->IsValidHandle())
 								{
-									Node->CreatePropertyHandle()->ResetToDefault();
-									return FReply::Handled();
-								})
-								.Content()
+									void* DataPtr;
+									if (PropertyHandle->GetValueData(DataPtr) == FPropertyAccess::Result::Success)
+									{
+										FProperty* NodeProperty = PropertyHandle->GetProperty();
+										bool bVisible = !NodeProperty->Identical(FieldPtr->DefaultValue.Get(), DataPtr);
+										return bVisible ? EVisibility::Visible : EVisibility::Hidden;
+									}
+								}
+								return EVisibility::Hidden;
+							};
+
+							ResetButtonWidget = SNew(SBox)
+								.HAlign(HAlign_Center)
+								.VAlign(VAlign_Center)
+								.Padding(2.f, 4.f)
 								[
-									SNew(SImage)
-									.Image(FAppStyle::GetBrush("PropertyWindow.DiffersFromDefault"))
-									.ColorAndOpacity(FSlateColor::UseForeground())
-								]
-							];
+									SNew(SButton)
+									.IsFocusable(false)
+									.ButtonStyle(&RCPanelStyle->FlatButtonStyle)
+									.ContentPadding(RCPanelStyle->PanelPadding)
+									.Visibility_Lambda(IsVisible)
+									.OnClicked_Lambda([PropertyHandle]()
+									{
+										PropertyHandle->ResetToDefault();
+										return FReply::Handled();
+									})
+									.Content()
+									[
+										SNew(SImage)
+										.Image(FAppStyle::GetBrush("PropertyWindow.DiffersFromDefault"))
+										.ColorAndOpacity(FSlateColor::UseForeground())
+									]
+								];
+						}
+						else
+						{
+							ResetButtonWidget = SNew(SBox)
+								.HAlign(HAlign_Center)
+								.VAlign(VAlign_Center)
+								.Padding(2.f, 4.f)
+								[
+									ConstructResetToDefaultWidget(Object, PropertyHandle)
+								];
+						}
+
+						return MakeFieldWidget(FieldWidget.ToSharedRef(), EditConditionWidget.ToSharedRef());
 					}
 					else
 					{
-						ResetButtonWidget = SNew(SBox)
-							.HAlign(HAlign_Center)
-							.VAlign(VAlign_Center)
-							.Padding(2.f, 4.f)
-							[
-								ConstructResetToDefaultWidget(Object, PropertyHandle)
-							];
-					}
-
-					return MakeFieldWidget(FieldWidget.ToSharedRef(), EditConditionWidget.ToSharedRef());
-				}
-				else
-				{
-					FString PropertyName = Field->FieldPathInfo.ToPathPropertyString();
-					FString TargetName = Object->GetName();
+						FString PropertyName = Field->FieldPathInfo.ToPathPropertyString();
+						FString TargetName = Object->GetName();
 
-					FRCFieldPathInfo& FieldPath = Field->FieldPathInfo;
+						FRCFieldPathInfo& FieldPath = Field->FieldPathInfo;
 
-					if (FieldPath.Resolve(Object))
-					{
-						FRCFieldResolvedData ResolvedData = FieldPath.GetResolvedData();
-						const FRCFieldPathSegment& LastSegment = FieldPath.GetFieldSegment(FieldPath.GetSegmentCount() - 1);
-						if (LastSegment.ArrayIndex != INDEX_NONE)
+						if (FieldPath.Resolve(Object))
 						{
-							if (FArrayProperty* ArrayProperty = CastField<FArrayProperty>(ResolvedData.Field))
+							FRCFieldResolvedData ResolvedData = FieldPath.GetResolvedData();
+							const FRCFieldPathSegment& LastSegment = FieldPath.GetFieldSegment(FieldPath.GetSegmentCount() - 1);
+							if (LastSegment.ArrayIndex != INDEX_NONE)
 							{
-								FScriptArrayHelper_InContainer Helper(ArrayProperty, ResolvedData.ContainerAddress);
-								int32 ArrayNum = Helper.Num();
-								if (!Helper.IsValidIndex(LastSegment.ArrayIndex))
+								if (FArrayProperty* ArrayProperty = CastField<FArrayProperty>(ResolvedData.Field))
 								{
-									PropertyName = LastSegment.ToString();
-
-									if (FieldPath.GetSegmentCount() > 1)
+									FScriptArrayHelper_InContainer Helper(ArrayProperty, ResolvedData.ContainerAddress);
+									int32 ArrayNum = Helper.Num();
+									if (!Helper.IsValidIndex(LastSegment.ArrayIndex))
 									{
-										TargetName = Object->GetName() + TEXT(".") + FieldPath.ToString(FieldPath.GetSegmentCount() - 1);
+										PropertyName = LastSegment.ToString();
+
+										if (FieldPath.GetSegmentCount() > 1)
+										{
+											TargetName = Object->GetName() + TEXT(".") + FieldPath.ToString(FieldPath.GetSegmentCount() - 1);
+										}
 									}
 								}
 							}
 						}
-					}
 
-					FText ErrorText = FText::Format(LOCTEXT("ExposedPropertyInvalidErrorMessage", "Could not find property {0} on object {1}"), FText::FromString(PropertyName), FText::FromString(TargetName));
-					IRemoteControlModule::BroadcastError(ErrorText.ToString());
+						FText ErrorText = FText::Format(LOCTEXT("ExposedPropertyInvalidErrorMessage", "Could not find property {0} on object {1}"), FText::FromString(PropertyName), FText::FromString(TargetName));
+						IRemoteControlModule::BroadcastError(ErrorText.ToString());
 
-					return MakeFieldWidget(CreateInvalidWidget(ErrorText));
+						return MakeFieldWidget(CreateInvalidWidget(ErrorText));
+					}
 				}
 			}
 		}
-
 		return MakeFieldWidget(CreateInvalidWidget());
 	}
 
diff --git a/Engine/Plugins/VirtualProduction/RemoteControl/Source/RemoteControlUI/Private/UI/SRemoteControlPanel.cpp b/Engine/Plugins/VirtualProduction/RemoteControl/Source/RemoteControlUI/Private/UI/SRemoteControlPanel.cpp
--- a/Engine/Plugins/VirtualProduction/RemoteControl/Source/RemoteControlUI/Private/UI/SRemoteControlPanel.cpp
+++ b/Engine/Plugins/VirtualProduction/RemoteControl/Source/RemoteControlUI/Private/UI/SRemoteControlPanel.cpp
@@ -29,6 +29,7 @@
 #include "Kismet/BlueprintFunctionLibrary.h"
 #include "Layout/Visibility.h"
 #include "Materials/Material.h"
+#include "MenuContext/RemoteControlPanelMenuContext.h"
 #include "Modules/ModuleManager.h"
 #include "PropertyCustomizationHelpers.h"
 #include "PropertyEditorModule.h"
@@ -84,6 +85,7 @@
 
 const FName SRemoteControlPanel::DefaultRemoteControlPanelToolBarName("RemoteControlPanel.DefaultToolBar");
 const FName SRemoteControlPanel::AuxiliaryRemoteControlPanelToolBarName("RemoteControlPanel.AuxiliaryToolBar");
+const FName SRemoteControlPanel::TargetWorldRemoteControlPanelMenuName("RemoteControlPanel.TargetWorld");
 const float SRemoteControlPanel::MinimumPanelWidth = 640.f;
 
 TSharedRef<SBox> SRemoteControlPanel::CreateNoneSelectedWidget()
@@ -1446,16 +1448,41 @@ void SRemoteControlPanel::BindRemoteControlCommands()
 
 void SRemoteControlPanel::OnObjectReplaced(const TMap<UObject*, UObject*>& InObjectReplaced)
 {
-	if (!WidgetRegistry.IsValid())
+	if (!WidgetRegistry.IsValid() || !Preset.IsValid() || !Preset->SelectedWorld.IsValid())
 	{
 		return;
 	}
 
-	for (const TPair<UObject*, UObject*> ReplacedObject : InObjectReplaced)
+	for (TWeakPtr<FRemoteControlField> FieldWeak : Preset->GetExposedEntities<FRemoteControlField>())
 	{
-		if (ReplacedObject.Key && ReplacedObject.Value)
+		if (const TSharedPtr<FRemoteControlField> Field = FieldWeak.Pin())
 		{
-			WidgetRegistry->ReplaceGeneratorObject(ReplacedObject.Key, ReplacedObject.Value);
+			const UObject* BoundObject = Field->GetBoundObjectForWorld(Preset->SelectedWorld.Get());
+			for (const TPair<UObject*, UObject*>& ObjectReplaced : InObjectReplaced)
+			{
+				if (ObjectReplaced.Value == BoundObject)
+				{
+					WidgetRegistry->UpdateGeneratorAndTreeCache(ObjectReplaced.Key, ObjectReplaced.Value, Field->FieldPathInfo.ToPathPropertyString());
+				}
+			}
+		}
+	}
+}
+
+void SRemoteControlPanel::OnEndPIE(const bool bInIsSimulating)
+{
+	if (Preset.IsValid() && (!Preset->SelectedWorld.IsValid() || bInIsSimulating == false))
+	{
+		if (Preset->IsEmbeddedPreset() && Preset->SelectedWorld.Get()->WorldType == EWorldType::PIE)
+		{
+			OpenEditorEmbeddedPreset();
+		}
+		else
+		{
+			if (GEditor && IsValid(GEditor->EditorWorld))
+			{
+				UpdatePanelForWorld(GEditor->EditorWorld);
+			}
 		}
 	}
 }
@@ -1481,6 +1508,7 @@ void SRemoteControlPanel::RegisterEvents()
 
 	UMaterial::OnMaterialCompilationFinished().AddSP(this, &SRemoteControlPanel::OnMaterialCompiled);
 	FCoreUObjectDelegates::OnObjectsReplaced.AddSP(this, &SRemoteControlPanel::OnObjectReplaced);
+	FEditorDelegates::EndPIE.AddSP(this, &SRemoteControlPanel::OnEndPIE);
 }
 
 void SRemoteControlPanel::UnregisterEvents()
@@ -1503,6 +1531,7 @@ void SRemoteControlPanel::UnregisterEvents()
 	FEditorDelegates::MapChange.RemoveAll(this);
 	UMaterial::OnMaterialCompilationFinished().RemoveAll(this);
 	FCoreUObjectDelegates::OnObjectsReplaced.RemoveAll(this);
+	FEditorDelegates::EndPIE.RemoveAll(this);
 }
 
 void SRemoteControlPanel::RegisterPanels()
@@ -2054,6 +2083,20 @@ void SRemoteControlPanel::GenerateToolbar()
 	const URemoteControlSettings* Settings = GetMutableDefault<URemoteControlSettings>();
 	const FName& DefaultPanelMode = Settings->DefaultPanelMode;
 
+	if (!Preset->SelectedWorld.IsValid())
+	{
+		Preset->SelectedWorld = Preset->GetWorld(true);
+	}
+
+	if (Preset->SelectedWorld.IsValid())
+	{
+		SelectedWorldName = Preset->SelectedWorld->StreamingLevelsPrefix + Preset->SelectedWorld->GetName();
+	}
+	else
+	{
+		SelectedWorldName = TEXT("WORLD NOT FOUND");
+	}
+
 	Toolbar =
 		SNew(SHorizontalBox)
 		+ SHorizontalBox::Slot()
@@ -2074,6 +2117,27 @@ void SRemoteControlPanel::GenerateToolbar()
 		.Padding(5.f, 0.f)
 		.HAlign(HAlign_Fill)
 		.VAlign(VAlign_Center)
+		.AutoWidth()
+		[
+			SNew(SComboButton)
+			.ToolTipText(LOCTEXT("RCFieldsGroupingTooltip", "Select grouping type for the fields"))
+			.HAlign(HAlign_Center)
+			.VAlign(VAlign_Center)
+			.ButtonStyle(&RCPanelStyle->FlatButtonStyle)
+			.CollapseMenuOnParentFocus(true)
+			.HasDownArrow(false)
+			.ContentPadding(FMargin(4.f, 2.f))
+			.OnGetMenuContent(this, &SRemoteControlPanel::OnGetSelectedWorldButtonContent)
+			.ButtonContent()
+			[
+				SNew(STextBlock) 
+				.Text_Lambda([this] () { return FText::FromString(SelectedWorldName); })
+			]
+		]
+		+ SHorizontalBox::Slot()
+		.Padding(5.f, 0.f)
+		.HAlign(HAlign_Fill)
+		.VAlign(VAlign_Center)
 		.FillWidth(1.f)
 		[
 			SNew(SSpacer)
@@ -2741,6 +2805,174 @@ void SRemoteControlPanel::SaveSettings()
 	}
 }
 
+TSharedRef<SWidget> SRemoteControlPanel::OnGetSelectedWorldButtonContent()
+{
+	UToolMenus* ToolMenus = UToolMenus::Get();
+
+	if (!ensure(ToolMenus))
+	{
+		return SNullWidget::NullWidget;
+	}
+
+	if (!ToolMenus->IsMenuRegistered(TargetWorldRemoteControlPanelMenuName))
+	{
+		UToolMenu* Menu = ToolMenus->RegisterMenu(TargetWorldRemoteControlPanelMenuName, NAME_None, EMultiBoxType::Menu);
+
+		if (!ensure(Menu))
+		{
+			return SNullWidget::NullWidget;
+		}
+
+		Menu->AddDynamicSection(TEXT("Worlds"), FNewToolMenuDelegate::CreateStatic(&CreateTargetWorldButtonDynamicEntries));
+	}
+
+	URemoteControlPanelMenuContext* RemoteControlContext = NewObject<URemoteControlPanelMenuContext>();
+	RemoteControlContext->RemoteControlPanel = SharedThis(this);
+
+	const FToolMenuContext Context(RemoteControlContext);
+	return ToolMenus->GenerateWidget(TargetWorldRemoteControlPanelMenuName, Context);
+}
+
+void SRemoteControlPanel::UpdatePanelForWorld(const UWorld* InWorld)
+{
+	if (!IsValid(InWorld) || !Preset->SelectedWorld.IsValid() || InWorld == Preset->SelectedWorld.Get())
+	{
+		return;
+	}
+
+	if (Preset->IsEmbeddedPreset())
+	{
+		OpenPanelForEmbeddedPreset(InWorld);
+	}
+	else
+	{
+		// Map of old world objects to new target world objects
+		TMap<UObject*, UObject*> OldToNewObject;
+		for (const TObjectPtr<URemoteControlBinding>& Binding : Preset->Bindings)
+		{
+			UObject* OldObject = Binding->ResolveForWorld(Preset->SelectedWorld.Get());
+			UObject* NewObject = Binding->ResolveForWorld(InWorld);
+			OldToNewObject.Add(OldObject, NewObject);
+		}
+
+		for (TWeakPtr<FRemoteControlField> FieldWeak : Preset->GetExposedEntities<FRemoteControlField>())
+		{
+			if (const TSharedPtr<FRemoteControlField> Field = FieldWeak.Pin())
+			{
+				// Get the bound object in the correct world
+				UObject* BoundObject = Field->GetBoundObjectForWorld(Preset->SelectedWorld.Get());
+				if (UObject** NewObject = OldToNewObject.Find(BoundObject))
+				{
+					WidgetRegistry->UpdateGeneratorAndTreeCache(BoundObject, *NewObject, Field->FieldPathInfo.ToPathPropertyString());
+				}
+			}
+		}
+
+		if (EntityList.IsValid())
+		{
+			EntityList->Refresh();
+		}
+
+		SelectedWorldName = InWorld->StreamingLevelsPrefix + InWorld->GetName();
+		Preset->SelectedWorld = InWorld;
+	}
+}
+
+void SRemoteControlPanel::OpenEmbeddedPreset(const FSoftObjectPath& InPresetToOpenPath)
+{
+	URemoteControlPreset* EditorRCPreset = nullptr;
+	if (UObject* EditorPreset = InPresetToOpenPath.ResolveObject())
+	{
+		EditorRCPreset = Cast<URemoteControlPreset>(EditorPreset);
+	}
+	else
+	{
+		EditorRCPreset = Cast<URemoteControlPreset>(InPresetToOpenPath.TryLoad());
+	}
+
+	if (IsValid(EditorRCPreset))
+	{
+		UAssetEditorSubsystem* AssetSubsystem = GEditor ? GEditor->GetEditorSubsystem<UAssetEditorSubsystem>() : nullptr;
+		if (AssetSubsystem)
+		{
+			AssetSubsystem->OpenEditorForAsset(EditorRCPreset);
+		}
+	}
+}
+
+void SRemoteControlPanel::OpenEditorEmbeddedPreset()
+{
+	if (GEditor && GEditor->EditorWorld)
+	{
+		if (const UPackage* EditorWorldPackage = GEditor->EditorWorld.GetPackage())
+		{
+			const FSoftObjectPath RemotePreset = Preset.Get();
+			// Fix path for the editor world
+			const FSoftObjectPath EditorPresetPath = FSoftObjectPath(EditorWorldPackage->GetFName(), RemotePreset.GetAssetFName(), RemotePreset.GetSubPathString());
+			OpenEmbeddedPreset(EditorPresetPath);
+		}
+	}
+}
+
+void SRemoteControlPanel::OpenPanelForEmbeddedPreset(const UWorld* InWorld)
+{
+	if (const UPackage* Package = InWorld->GetPackage())
+	{
+		const int32 PIEInstanceId = Package->GetPIEInstanceID();
+		if (PIEInstanceId != INDEX_NONE)
+		{
+			FSoftObjectPath CurrentPresetSoftPath = Preset.Get();
+			// Fix path for the current selected world
+			CurrentPresetSoftPath = FSoftObjectPath(Package->GetFName(), CurrentPresetSoftPath.GetAssetFName(), CurrentPresetSoftPath.GetSubPathString());
+			OpenEmbeddedPreset(CurrentPresetSoftPath);
+		}
+		else if (InWorld->WorldType == EWorldType::Editor)
+		{
+			OpenEditorEmbeddedPreset();
+		}
+	}
+}
+
+void SRemoteControlPanel::CreateTargetWorldButtonDynamicEntries(UToolMenu* InMenu)
+{
+	const URemoteControlPanelMenuContext* Context = InMenu->FindContext<URemoteControlPanelMenuContext>();
+
+	if (!Context || !Context->RemoteControlPanel.IsValid())
+	{
+		return;
+	}
+
+	const TSharedPtr<SRemoteControlPanel> CurrentPanel = Context->RemoteControlPanel.Pin();
+
+	TSet<const UWorld*> Worlds;
+	if (GEngine)
+	{
+		for (const FWorldContext& WorldContext : GEngine->GetWorldContexts())
+		{
+			const UWorld* CurrentWorld = WorldContext.World();
+			if (CurrentWorld && (CurrentWorld->WorldType == EWorldType::Editor || CurrentWorld->WorldType == EWorldType::PIE))
+			{
+				Worlds.Add(CurrentWorld);
+			}
+		}
+	}
+
+	if (!Worlds.IsEmpty())
+	{
+		FToolMenuSection& Section = InMenu->FindOrAddSection(TEXT("Worlds"));
+		for (const UWorld* World : Worlds)
+		{
+			Section.AddMenuEntry(
+				FName(World->StreamingLevelsPrefix + World->GetName()),
+				FText::FromString(World->StreamingLevelsPrefix + World->GetName()),
+				LOCTEXT("RCSetTargetWorld_Tooltip", "Set this world as the target world for RC"),
+				FSlateIcon(),
+				FExecuteAction::CreateSP(CurrentPanel.Get(), &SRemoteControlPanel::UpdatePanelForWorld, World)
+				);
+		}
+	}
+}
+
 void SRemoteControlPanel::DeleteEntity()
 {
 	if (CanDeleteEntity())
diff --git a/Engine/Plugins/VirtualProduction/RemoteControl/Source/RemoteControlUI/Private/UI/SRemoteControlPanel.h b/Engine/Plugins/VirtualProduction/RemoteControl/Source/RemoteControlUI/Private/UI/SRemoteControlPanel.h
--- a/Engine/Plugins/VirtualProduction/RemoteControl/Source/RemoteControlUI/Private/UI/SRemoteControlPanel.h
+++ b/Engine/Plugins/VirtualProduction/RemoteControl/Source/RemoteControlUI/Private/UI/SRemoteControlPanel.h
@@ -179,6 +179,9 @@ private:
 	/** Called when object are replaced, blueprint for example */
 	void OnObjectReplaced(const TMap<UObject*, UObject*>& InObjectReplaced);
 
+	/** Called when PIE ends */
+	void OnEndPIE(const bool bInIsSimulating);
+
 	/** Register editor events needed to handle reloading objects and blueprint libraries. */
 	void RegisterEvents();
 	/** Unregister editor events */
@@ -365,9 +368,28 @@ private:
 	/** Retrieves active logic panel. */
 	TSharedPtr<SRCLogicPanelBase> GetActiveLogicPanel() const;
 
+	/** Get the menu content for the SelectedWorld button */
+	TSharedRef<SWidget> OnGetSelectedWorldButtonContent();
+
+	/** Update the panel for the new World */
+	void UpdatePanelForWorld(const UWorld* InWorld);
+
+	/** try to Open the given preset path */
+	void OpenEmbeddedPreset(const FSoftObjectPath& InPresetToOpenPath);
+
+	/** try to open the editor preset of the current Preset */
+	void OpenEditorEmbeddedPreset();
+
+	/** Handle the TargetWorld change for embedded presets */
+	void OpenPanelForEmbeddedPreset(const UWorld* World);
+
+	/** Handle the TargetWorld ComboButton entries creation */
+	static void CreateTargetWorldButtonDynamicEntries(UToolMenu* InMenu);
+
 private:
 	static const FName DefaultRemoteControlPanelToolBarName;
 	static const FName AuxiliaryRemoteControlPanelToolBarName;
+	static const FName TargetWorldRemoteControlPanelMenuName;
 	/** Holds the preset asset. */
 	TStrongObjectPtr<URemoteControlPreset> Preset;
 	/** Whether the panel is in protocols mode. */
@@ -432,7 +454,8 @@ private:
 	ERCPanels ActivePanel;
 	/** Input Preprocessor which catches the Delete Key when Docked. */
 	TSharedPtr<IInputProcessor> InputProcessor;
-
+	/** Currently selected world name */
+	FString SelectedWorldName;
 	// ~ Remote Control Logic Panels ~
 
 	/** Controller panel UI widget for Remote Control Logic*/
