[Remote Control] Adding regression test for interception of properties with setter.

Note: extracting ForwardInterceptor so it can be reused in other tests.

diff --git a/Engine/Plugins/VirtualProduction/RemoteControl/Source/RemoteControl/Private/Tests/RemoteControlInterceptionForwardInterceptor.cpp b/Engine/Plugins/VirtualProduction/RemoteControl/Source/RemoteControl/Private/Tests/RemoteControlInterceptionForwardInterceptor.cpp
new file mode 100644
--- /dev/null
+++ b/Engine/Plugins/VirtualProduction/RemoteControl/Source/RemoteControl/Private/Tests/RemoteControlInterceptionForwardInterceptor.cpp
@@ -0,0 +1,68 @@
+// Copyright Epic Games, Inc. All Rights Reserved.
+
+#include "RemoteControlInterceptionForwardInterceptor.h"
+
+#include "Features/IModularFeatures.h"
+
+namespace UE::RemoteControl::Test::Private
+{
+	bool ForEachInterceptionProcessor(TFunctionRef<void(IRemoteControlInterceptionFeatureProcessor*)> InFunction)
+	{
+		IModularFeatures& ModularFeatures = IModularFeatures::Get();
+		const FName ProcessorFeatureName = IRemoteControlInterceptionFeatureProcessor::GetName();
+		const int32 NumProcessors = ModularFeatures.GetModularFeatureImplementationCount(ProcessorFeatureName);
+
+		for (int32 ProcessorIdx = 0; ProcessorIdx < NumProcessors; ++ProcessorIdx)
+		{
+			IRemoteControlInterceptionFeatureProcessor* const Processor = static_cast<IRemoteControlInterceptionFeatureProcessor*>(
+				ModularFeatures.GetModularFeatureImplementation(ProcessorFeatureName, ProcessorIdx));
+
+			if (Processor)
+			{
+				InFunction(Processor);
+			}
+		}
+		
+		return NumProcessors > 0;
+	}
+}
+
+ERCIResponse FRemoteControlInterceptionForwardInterceptor::SetObjectProperties(FRCIPropertiesMetadata& InObjectProperties)
+{
+	using namespace UE::RemoteControl::Test::Private;
+	const bool bForwarded = ForEachInterceptionProcessor([&InObjectProperties](IRemoteControlInterceptionFeatureProcessor* InProcessor)
+	{
+		InProcessor->SetObjectProperties(InObjectProperties);
+	});
+	return bForwarded ? ERCIResponse::Intercept : ERCIResponse::Apply;
+}
+
+ERCIResponse FRemoteControlInterceptionForwardInterceptor::ResetObjectProperties(FRCIObjectMetadata& InObject)
+{
+	using namespace UE::RemoteControl::Test::Private;
+	const bool bForwarded = ForEachInterceptionProcessor([&InObject](IRemoteControlInterceptionFeatureProcessor* InProcessor)
+	{
+		InProcessor->ResetObjectProperties(InObject);
+	});
+	return bForwarded ? ERCIResponse::Intercept : ERCIResponse::Apply;
+}
+
+ERCIResponse FRemoteControlInterceptionForwardInterceptor::InvokeCall(FRCIFunctionMetadata& InFunction)
+{
+	using namespace UE::RemoteControl::Test::Private;
+	const bool bForwarded = ForEachInterceptionProcessor([&InFunction](IRemoteControlInterceptionFeatureProcessor* InProcessor)
+	{
+		InProcessor->InvokeCall(InFunction);
+	});
+	return bForwarded ? ERCIResponse::Intercept : ERCIResponse::Apply;
+}
+
+ERCIResponse FRemoteControlInterceptionForwardInterceptor::SetPresetController(FRCIControllerMetadata& InController)
+{
+	using namespace UE::RemoteControl::Test::Private;
+	const bool bForwarded = ForEachInterceptionProcessor([&InController](IRemoteControlInterceptionFeatureProcessor* InProcessor)
+	{
+		InProcessor->SetPresetController(InController);
+	});
+	return bForwarded ? ERCIResponse::Intercept : ERCIResponse::Apply;
+}
diff --git a/Engine/Plugins/VirtualProduction/RemoteControl/Source/RemoteControl/Private/Tests/RemoteControlInterceptionForwardInterceptor.h b/Engine/Plugins/VirtualProduction/RemoteControl/Source/RemoteControl/Private/Tests/RemoteControlInterceptionForwardInterceptor.h
new file mode 100644
--- /dev/null
+++ b/Engine/Plugins/VirtualProduction/RemoteControl/Source/RemoteControl/Private/Tests/RemoteControlInterceptionForwardInterceptor.h
@@ -0,0 +1,23 @@
+// Copyright Epic Games, Inc. All Rights Reserved.
+
+#pragma once
+
+#include "RemoteControlInterceptionProcessor.h"
+
+/**
+ * Interceptor implementation that forwards data directly to the processor.
+ */
+class FRemoteControlInterceptionForwardInterceptor : public IRemoteControlInterceptionFeatureInterceptor
+{
+public:
+	FRemoteControlInterceptionForwardInterceptor() = default;
+	virtual ~FRemoteControlInterceptionForwardInterceptor() override = default;
+
+protected:
+	//~ Begin IRemoteControlInterceptionCommands
+	virtual ERCIResponse SetObjectProperties(FRCIPropertiesMetadata& InObjectProperties) override;
+	virtual ERCIResponse ResetObjectProperties(FRCIObjectMetadata& InObject) override;
+	virtual ERCIResponse InvokeCall(FRCIFunctionMetadata& InFunction) override;
+	virtual ERCIResponse SetPresetController(FRCIControllerMetadata& InController) override;
+	//~ End IRemoteControlInterceptionCommands
+};
\ No newline at end of file
diff --git a/Engine/Plugins/VirtualProduction/RemoteControl/Source/RemoteControl/Private/Tests/RemoteControlInterceptionTest.cpp b/Engine/Plugins/VirtualProduction/RemoteControl/Source/RemoteControl/Private/Tests/RemoteControlInterceptionTest.cpp
--- a/Engine/Plugins/VirtualProduction/RemoteControl/Source/RemoteControl/Private/Tests/RemoteControlInterceptionTest.cpp
+++ b/Engine/Plugins/VirtualProduction/RemoteControl/Source/RemoteControl/Private/Tests/RemoteControlInterceptionTest.cpp
@@ -18,84 +18,10 @@
 #include "Backends/CborStructDeserializerBackend.h"
 #include "Backends/JsonStructDeserializerBackend.h"
 #include "IRemoteControlModule.h"
-#include "IRemoteControlInterceptionFeature.h"
+#include "RemoteControlInterceptionForwardInterceptor.h"
 #include "Serialization/MemoryWriter.h"
 #include "Serialization/MemoryReader.h"
 
-/**
- * Interceptor interface implementation
- */
-class FRemoteControlTestFeatureInterceptor : public IRemoteControlInterceptionFeatureInterceptor
-{
-public:
-	FRemoteControlTestFeatureInterceptor() = default;
-	virtual ~FRemoteControlTestFeatureInterceptor() = default;
-
-protected:
-	// IRemoteControlInterceptionCommands interface
-	virtual ERCIResponse SetObjectProperties(FRCIPropertiesMetadata& InObjectProperties) override
-	{
-		// Get processor feature
-		IRemoteControlInterceptionFeatureProcessor* const Processor = static_cast<IRemoteControlInterceptionFeatureProcessor*>(
-			IModularFeatures::Get().GetModularFeatureImplementation(IRemoteControlInterceptionFeatureProcessor::GetName(), 0));
-
-		// In case the processor feature has been registered, forward data directly to the processor
-		if (Processor)
-		{
-			Processor->SetObjectProperties(InObjectProperties);
-		}
-
-		return ERCIResponse::Intercept;
-	}
-
-	virtual ERCIResponse ResetObjectProperties(FRCIObjectMetadata& InObject) override
-	{
-		// Get processor feature
-		IRemoteControlInterceptionFeatureProcessor* const Processor = static_cast<IRemoteControlInterceptionFeatureProcessor*>(
-			IModularFeatures::Get().GetModularFeatureImplementation(IRemoteControlInterceptionFeatureProcessor::GetName(), 0));
-
-		// In case the processor feature has been registered, forward data directly to the processor
-		if (Processor)
-		{
-			Processor->ResetObjectProperties(InObject);
-		}
-
-		return ERCIResponse::Intercept;
-	}
-
-	virtual ERCIResponse InvokeCall(FRCIFunctionMetadata& InFunction) override
-	{
-		// Get processor feature
-		IRemoteControlInterceptionFeatureProcessor* const Processor = static_cast<IRemoteControlInterceptionFeatureProcessor*>(
-			IModularFeatures::Get().GetModularFeatureImplementation(IRemoteControlInterceptionFeatureProcessor::GetName(), 0));
-
-		// In case the processor feature has been registered, forward data directly to the processor
-		if (Processor)
-		{
-			Processor->InvokeCall(InFunction);
-		}
-		
-		return ERCIResponse::Intercept;
-	}
-
-	virtual ERCIResponse SetPresetController(FRCIControllerMetadata& InController)
-	{
-		// Get processor feature
-		IRemoteControlInterceptionFeatureProcessor* const Processor = static_cast<IRemoteControlInterceptionFeatureProcessor*>(
-			IModularFeatures::Get().GetModularFeatureImplementation(IRemoteControlInterceptionFeatureProcessor::GetName(), 0));
-
-		// In case the processor feature has been registered, forward data directly to the processor
-		if (Processor)
-		{
-			Processor->SetPresetController(InController);
-		}
-
-		return ERCIResponse::Intercept;
-	}
-
-	// ~IRemoteControlInterceptionCommands interface
-};
-
 struct FRemoteControlInterceptionTest
 {
 	FRemoteControlInterceptionTest() = delete;
@@ -110,9 +36,15 @@ struct FRemoteControlInterceptionTest
 		// Expose property
 		CustomStructProperty = FindFProperty<FProperty>(URemoteControlInterceptionTestObject::StaticClass(),  GET_MEMBER_NAME_CHECKED(URemoteControlInterceptionTestObject,  CustomStruct));
 		Int32ValueProperty   = FindFProperty<FProperty>(FRemoteControlInterceptionTestStruct::StaticStruct(), GET_MEMBER_NAME_CHECKED(FRemoteControlInterceptionTestStruct, Int32Value));
-		FString Int32ValuePropertyFullPath = FString::Printf(TEXT("%s.%s"), *CustomStructProperty->GetName(), *Int32ValueProperty->GetName());
+		const FString Int32ValuePropertyFullPath = FString::Printf(TEXT("%s.%s"), *CustomStructProperty->GetName(), *Int32ValueProperty->GetName());
 		Int32ValuePropertyRCProp = Preset->ExposeProperty(RemoteControlInterceptionTestObject.Get(), FRCFieldPathInfo{ Int32ValuePropertyFullPath }).Pin();
 
+		// Expose property with setter
+		static const FName ValueWithSetterName(TEXT("ValueWithSetter"));
+		ValueWithSetterProperty = FindFProperty<FProperty>(URemoteControlInterceptionTestObject::StaticClass(), ValueWithSetterName);
+		check(ValueWithSetterProperty);
+		ValueWithSetterRCProp = Preset->ExposeProperty(RemoteControlInterceptionTestObject.Get(), FRCFieldPathInfo{ ValueWithSetterProperty->GetName() }).Pin();
+		
 		// Expose function
 		UClass* TestObjectClass = RemoteControlInterceptionTestObject->StaticClass();
 		UFunction* TestFunctionClass = TestObjectClass->FindFunctionByName("TestFunction");
@@ -122,7 +54,7 @@ struct FRemoteControlInterceptionTest
 		RemoteControlInterceptionTestObject->CustomStruct.Int32Value = FRemoteControlInterceptionTestStruct::Int32ValueDefault;
 
 		// Instantiate the interceptor and processor features
-		FeatureInterceptor = MakeUnique<FRemoteControlTestFeatureInterceptor>();
+		FeatureInterceptor = MakeUnique<FRemoteControlInterceptionForwardInterceptor>();
 		FeatureProcessor   = MakeUnique<FRemoteControlInterceptionProcessor>();
 
 		// Register the features
@@ -212,7 +144,7 @@ struct FRemoteControlInterceptionTest
 		FJsonStructSerializerBackend SerializerBackend(Writer, EStructSerializerBackendFlags::Default);
 		FJsonStructDeserializerBackend DeserializerBackend(Reader);
 
-		// Serialize with test value and change the to default value back
+		// Serialize with test value and change the value back to default
 		RemoteControlInterceptionTestObject->CustomStruct.Int32Value = Int32ValueTest; // Set test value before serialization
 		FStructSerializer::SerializeElement(&RemoteControlInterceptionTestObject->CustomStruct, Int32ValueProperty, INDEX_NONE, SerializerBackend, FStructSerializerPolicies());
 		RemoteControlInterceptionTestObject->CustomStruct.Int32Value = FRemoteControlInterceptionTestStruct::Int32ValueDefault; // Set default value after serialization
@@ -226,11 +158,44 @@ struct FRemoteControlInterceptionTest
 		for (UObject* Object : Int32ValuePropertyRCProp->GetBoundObjects())
 		{
 			IRemoteControlModule::Get().ResolveObjectProperty(ObjectRef.Access, Object, ObjectRef.PropertyPathInfo, ObjectRef);
-			// Arhive should be intercepted inside SetObjectProperties
+			// Archive should be intercepted inside SetObjectProperties
 			IRemoteControlModule::Get().SetObjectProperties(ObjectRef, DeserializerBackend, ERCPayloadType::Json, JsonBuffer);
 		}
 	}
 
+	void TestValueWithSetter_Json()
+	{
+		// Set testing Archives
+		TArray<uint8> JsonBuffer;
+		FMemoryReader Reader(JsonBuffer);
+		FMemoryWriter Writer(JsonBuffer);
+
+		// Set serializers
+		FJsonStructSerializerBackend SerializerBackend(Writer, EStructSerializerBackendFlags::Default);
+		FJsonStructDeserializerBackend DeserializerBackend(Reader);
+
+		// Serialize with test value and change the value back to default
+		const FString TestValue = TEXT("ValueWithSetter_TestValue");
+		const FString DefaultValue = RemoteControlInterceptionTestObject->GetValueWithSetter();
+		RemoteControlInterceptionTestObject->SetValueWithSetter(TestValue); // Set test value before serialization
+		FStructSerializer::SerializeElement(RemoteControlInterceptionTestObject.Get(), ValueWithSetterProperty, INDEX_NONE, SerializerBackend, FStructSerializerPolicies());
+		RemoteControlInterceptionTestObject->SetValueWithSetter(DefaultValue); // Set default value after serialization
+
+		// Set object reference
+		FRCObjectReference ObjectRef;
+		ObjectRef.Property = ValueWithSetterRCProp->GetProperty();
+		ObjectRef.Access = ERCAccess::WRITE_TRANSACTION_ACCESS;
+		ObjectRef.PropertyPathInfo = ValueWithSetterRCProp->FieldPathInfo.ToString();
+
+		for (UObject* Object : ValueWithSetterRCProp->GetBoundObjects())
+		{
+			IRemoteControlModule::Get().ResolveObjectProperty(ObjectRef.Access, Object, ObjectRef.PropertyPathInfo, ObjectRef);
+			IRemoteControlModule::Get().SetObjectProperties(ObjectRef, DeserializerBackend, ERCPayloadType::Json, JsonBuffer);
+		}
+
+		Test->TestEqual(TEXT("Property value with setter different than expected"), RemoteControlInterceptionTestObject->GetValueWithSetter(), TestValue);
+	}
+	
 	void ResetObjectProperties()
 	{
 		// Set reset value before reset object property
@@ -383,6 +348,9 @@ private:
 	FProperty* Int32ValueProperty = nullptr;
 	const int32 Int32ValueTest = 34780;
 
+	TSharedPtr<FRemoteControlProperty> ValueWithSetterRCProp;
+	FProperty* ValueWithSetterProperty = nullptr;
+
 	TUniquePtr<IRemoteControlInterceptionFeatureInterceptor> FeatureInterceptor;
 	TUniquePtr<IRemoteControlInterceptionFeatureProcessor>   FeatureProcessor;
 };
@@ -397,6 +365,7 @@ bool FRemoteControlPresetInterceptionTest::RunTest(const FString& Parameters)
 	RemoteControlInterceptionTest.TestJson();
 	RemoteControlInterceptionTest.ResetObjectProperties();
 	RemoteControlInterceptionTest.InvokeCall();
+	RemoteControlInterceptionTest.TestValueWithSetter_Json();
 	
 	return true;
 }
diff --git a/Engine/Plugins/VirtualProduction/RemoteControl/Source/RemoteControl/Private/Tests/RemoteControlInterceptionTestData.h b/Engine/Plugins/VirtualProduction/RemoteControl/Source/RemoteControl/Private/Tests/RemoteControlInterceptionTestData.h
--- a/Engine/Plugins/VirtualProduction/RemoteControl/Source/RemoteControl/Private/Tests/RemoteControlInterceptionTestData.h
+++ b/Engine/Plugins/VirtualProduction/RemoteControl/Source/RemoteControl/Private/Tests/RemoteControlInterceptionTestData.h
@@ -66,5 +66,20 @@ public:
 
 		return Copy;
 	}
+	
+	void SetValueWithSetter(const FString& InString)
+	{
+		ValueWithSetter = InString;
+	}
+
+	const FString& GetValueWithSetter() const
+	{
+		return ValueWithSetter;
+	}
+
+private:
+	// Testing private property with a setter.
+	UPROPERTY(EditAnywhere, Setter = "SetValueWithSetter", Category = "RC")
+	FString ValueWithSetter;
 };
 
