Remote Control: Playing on PIE to automatically select the PIE world.

diff --git a/Engine/Plugins/VirtualProduction/RemoteControl/Source/RemoteControlUI/Private/UI/SRemoteControlPanel.cpp b/Engine/Plugins/VirtualProduction/RemoteControl/Source/RemoteControlUI/Private/UI/SRemoteControlPanel.cpp
--- a/Engine/Plugins/VirtualProduction/RemoteControl/Source/RemoteControlUI/Private/UI/SRemoteControlPanel.cpp
+++ b/Engine/Plugins/VirtualProduction/RemoteControl/Source/RemoteControlUI/Private/UI/SRemoteControlPanel.cpp
@@ -1469,21 +1469,50 @@ void SRemoteControlPanel::OnObjectReplaced(const TMap<UObject*, UObject*>& InObj
 	}
 }
 
+void SRemoteControlPanel::PostPIEStarted(const bool bInIsSimulating)
+{
+	// Don't do anything when it's Simulate In Editor. This panel will only react when it's a proper Play In Editor
+	if (!Preset.IsValid() || bInIsSimulating)
+	{
+		return;
+	}
+
+	// If the currently selected preset is already a PIE world, no need to override what's there. Return early
+	const UWorld* PresetWorld = Preset->SelectedWorld.Get();
+	if (PresetWorld && PresetWorld->WorldType == EWorldType::PIE)
+	{
+		return;
+	}
+
+	// Default to the first PIE World
+	if (FWorldContext* PIEWorldContext = GEditor->GetWorldContextFromPIEInstance(0))
+	{
+		UpdatePanelForWorld(PIEWorldContext->World());
+	}
+}
+
 void SRemoteControlPanel::OnEndPIE(const bool bInIsSimulating)
 {
-	if (Preset.IsValid() && (!Preset->SelectedWorld.IsValid() || bInIsSimulating == false))
+	// Don't do anything when it's Simulate In Editor. This panel will only react when it's a proper Play In Editor
+	if (!Preset.IsValid() || bInIsSimulating)
 	{
-		if (Preset->IsEmbeddedPreset() && Preset->SelectedWorld.Get()->WorldType == EWorldType::PIE)
-		{
-			OpenEditorEmbeddedPreset();
-		}
-		else
-		{
-			if (GEditor && IsValid(GEditor->EditorWorld))
-			{
-				UpdatePanelForWorld(GEditor->EditorWorld);
-			}
-		}
+		return;
+	}
+
+	// If the currently selected preset is not a PIE world when ending, no need to override what's there. Return early
+	const UWorld* PresetWorld = Preset->SelectedWorld.Get();
+	if (PresetWorld && PresetWorld->WorldType != EWorldType::PIE)
+	{
+		return;
+	}
+
+	if (Preset->IsEmbeddedPreset())
+	{
+		OpenEditorEmbeddedPreset();
+	}
+	else if (GEditor)
+	{
+		UpdatePanelForWorld(GEditor->EditorWorld);
 	}
 }
 
@@ -1508,6 +1537,8 @@ void SRemoteControlPanel::RegisterEvents()
 
 	UMaterial::OnMaterialCompilationFinished().AddSP(this, &SRemoteControlPanel::OnMaterialCompiled);
 	FCoreUObjectDelegates::OnObjectsReplaced.AddSP(this, &SRemoteControlPanel::OnObjectReplaced);
+
+	FEditorDelegates::PostPIEStarted.AddSP(this, &SRemoteControlPanel::PostPIEStarted);
 	FEditorDelegates::EndPIE.AddSP(this, &SRemoteControlPanel::OnEndPIE);
 }
 
@@ -1531,6 +1562,8 @@ void SRemoteControlPanel::UnregisterEvents()
 	FEditorDelegates::MapChange.RemoveAll(this);
 	UMaterial::OnMaterialCompilationFinished().RemoveAll(this);
 	FCoreUObjectDelegates::OnObjectsReplaced.RemoveAll(this);
+
+	FEditorDelegates::PostPIEStarted.RemoveAll(this);
 	FEditorDelegates::EndPIE.RemoveAll(this);
 }
 
@@ -2835,7 +2868,7 @@ TSharedRef<SWidget> SRemoteControlPanel::OnGetSelectedWorldButtonContent()
 
 void SRemoteControlPanel::UpdatePanelForWorld(const UWorld* InWorld)
 {
-	if (!IsValid(InWorld) || !Preset->SelectedWorld.IsValid() || InWorld == Preset->SelectedWorld.Get())
+	if (!IsValid(InWorld) || !Preset.IsValid() || InWorld == Preset->SelectedWorld.Get())
 	{
 		return;
 	}
diff --git a/Engine/Plugins/VirtualProduction/RemoteControl/Source/RemoteControlUI/Private/UI/SRemoteControlPanel.h b/Engine/Plugins/VirtualProduction/RemoteControl/Source/RemoteControlUI/Private/UI/SRemoteControlPanel.h
--- a/Engine/Plugins/VirtualProduction/RemoteControl/Source/RemoteControlUI/Private/UI/SRemoteControlPanel.h
+++ b/Engine/Plugins/VirtualProduction/RemoteControl/Source/RemoteControlUI/Private/UI/SRemoteControlPanel.h
@@ -179,6 +179,9 @@ private:
 	/** Called when object are replaced, blueprint for example */
 	void OnObjectReplaced(const TMap<UObject*, UObject*>& InObjectReplaced);
 
+	/** Called when PIE begins */
+	void PostPIEStarted(const bool bInIsSimulating);
+
 	/** Called when PIE ends */
 	void OnEndPIE(const bool bInIsSimulating);
 
