Fix warning V547: Expression 'InSelectInfo == ESelectInfo::Direct' is always true.
#rnx

diff --git a/Engine/Plugins/VirtualProduction/RemoteControl/Source/RemoteControlUI/Private/UI/Behaviour/SRCBehaviourPanel.cpp b/Engine/Plugins/VirtualProduction/RemoteControl/Source/RemoteControlUI/Private/UI/Behaviour/SRCBehaviourPanel.cpp
--- a/Engine/Plugins/VirtualProduction/RemoteControl/Source/RemoteControlUI/Private/UI/Behaviour/SRCBehaviourPanel.cpp
+++ b/Engine/Plugins/VirtualProduction/RemoteControl/Source/RemoteControlUI/Private/UI/Behaviour/SRCBehaviourPanel.cpp
@@ -70,11 +70,8 @@ void SRCBehaviourPanel::OnControllerSelectionChanged(TSharedPtr<FRCControllerMod
 	TArray<TSharedPtr<FRCBehaviourModel>> CurrentSelectedItems;
 	if (BehaviourPanelList.IsValid() && InSelectInfo == ESelectInfo::Direct)
 	{
-		if (InSelectInfo == ESelectInfo::Direct)
-		{
-			CurrentSelectedItems = BehaviourPanelList->GetSelectedBehaviourItems();
-			bShouldRestoreSelection = true;
-		}
+		CurrentSelectedItems = BehaviourPanelList->GetSelectedBehaviourItems();
+		bShouldRestoreSelection = true;
 	}
 
 	UpdateWrappedWidget(InControllerItem);
