Remote Control: UI Improvement for Entity Groups

diff --git a/Engine/Plugins/VirtualProduction/RemoteControl/Source/RemoteControlUI/Private/UI/SRCPanelExposedEntitiesGroup.cpp b/Engine/Plugins/VirtualProduction/RemoteControl/Source/RemoteControlUI/Private/UI/SRCPanelExposedEntitiesGroup.cpp
--- a/Engine/Plugins/VirtualProduction/RemoteControl/Source/RemoteControlUI/Private/UI/SRCPanelExposedEntitiesGroup.cpp
+++ b/Engine/Plugins/VirtualProduction/RemoteControl/Source/RemoteControlUI/Private/UI/SRCPanelExposedEntitiesGroup.cpp
@@ -7,6 +7,8 @@
 #include "Widgets/Layout/SBox.h"
 #include "Widgets/Input/SEditableTextBox.h"
 
+#define LOCTEXT_NAMESPACE "SRCPanelExposedEntitiesGroup"
+
 void SRCPanelExposedEntitiesGroup::Construct(const FArguments& InArgs, EFieldGroupType InFieldGroupType, URemoteControlPreset* Preset)
 {
 	FieldKey = InArgs._FieldKey;
@@ -22,52 +24,47 @@ SRCPanelTreeNode::FMakeNodeWidgetArgs SRCPanelExposedEntitiesGroup::CreateNodeWi
 {
 	FMakeNodeWidgetArgs Args;
 
-	if (GroupType == EFieldGroupType::PropertyId)
-	{
-		Args.PropertyIdWidget = SNew(SBox)
+	Args.PropertyIdWidget = SNew(SBox)
 		[
 			SNew(SEditableTextBox)
-			.Justification(ETextJustify::Center)
+			.Justification(ETextJustify::Left)
 			.MinDesiredWidth(50.f)
 			.SelectAllTextWhenFocused(true)
 			.RevertTextOnEscape(true)
 			.ClearKeyboardFocusOnCommit(true)
-			.Text_Lambda([this] () { return FText::FromName(FieldKey); })
+			.Text_Lambda([this]{ return FText::FromName(PropertyIdName); })
 			.OnTextCommitted(this, &SRCPanelExposedEntitiesGroup::OnPropertyIdTextCommitted)
 		];
 
-		Args.NameWidget = SNew(SBox).HeightOverride(25).VAlign(VAlign_Center)
-			[
-				SNew(STextBlock)
-				.Text(FText::FromString("Property Id Group"))
-			];
+	Args.OwnerNameWidget = SNew(SBox)
+		.HeightOverride(25)
+		.VAlign(VAlign_Center)
+		[
+			SNew(STextBlock)
+			.Text_Lambda([this]{ return FText::FromName(OwnerName); })
+		];
 
-		Args.OwnerNameWidget = SNullWidget::NullWidget;
-	}
-	else if (GroupType == EFieldGroupType::Owner)
-	{
-		Args.OwnerNameWidget = SNew(SBox)
-			.HeightOverride(25)
-			.VAlign(VAlign_Center)
-			[
-				SNew(STextBlock)
-				.Text(FText::FromName(FieldKey))
-			];
-
-		Args.NameWidget = SNew(SBox).HeightOverride(25).VAlign(VAlign_Center)
-			[
-				SNew(STextBlock)
-				.Text(FText::FromString("Owner Group"))
-			];
-
-		Args.PropertyIdWidget = SNullWidget::NullWidget;
-	}
-	else
+	FText GroupText;
+
+	switch (GroupType)
 	{
-		Args.PropertyIdWidget = SNullWidget::NullWidget;
-		Args.OwnerNameWidget = SNullWidget::NullWidget;
+	case EFieldGroupType::PropertyId:
+		GroupText = LOCTEXT("GroupPropertyId", "Group by Id");
+		break;
+
+	case EFieldGroupType::Owner:
+		GroupText = LOCTEXT("GroupOwner", "Group by Owner");
+		break;
 	}
 
+	Args.NameWidget = SNew(SBox)
+		.HeightOverride(25)
+		.VAlign(VAlign_Center)
+		[
+			SNew(STextBlock)
+			.Text(GroupText)
+		];
+
 	Args.SubObjectPathWidget = SNullWidget::NullWidget;
 	Args.ValueWidget = SNullWidget::NullWidget;
 	Args.ResetButton = SNullWidget::NullWidget;
@@ -91,6 +88,7 @@ void SRCPanelExposedEntitiesGroup::OnPropertyIdTextCommitted(const FText& InText
 
 	URemoteControlPreset* Preset = PresetWeak.Get();
 	FieldKey = NewId;
+
 	for (const TSharedPtr<SRCPanelTreeNode>& Child : ChildWidgets)
 	{
 		TWeakPtr<FRemoteControlField> ExposedField = Preset->GetExposedEntity<FRemoteControlField>(Child->GetRCId());
@@ -101,36 +99,44 @@ void SRCPanelExposedEntitiesGroup::OnPropertyIdTextCommitted(const FText& InText
 			Preset->UpdateIdentifiedField(ExposedField.Pin().ToSharedRef());
 		}
 	}
-	
+
 	OnGroupPropertyIdChangedDelegate.ExecuteIfBound();
 }
 
 void SRCPanelExposedEntitiesGroup::AssignChildren(const TArray<TSharedPtr<SRCPanelTreeNode>>& InFieldEntities)
 {
 	ChildWidgets.Empty();
-	if (GroupType == EFieldGroupType::PropertyId)
-	{
-		for (const TSharedPtr<SRCPanelTreeNode>& Entity : InFieldEntities)
+
+	OwnerName = NAME_None;
+	PropertyIdName = NAME_None;
+
+	auto CompareUpdateName = [](FName& InOutName, const FName& InOtherName)
 		{
-			if (const TSharedPtr<SRCPanelExposedField> ExposedField = StaticCastSharedPtr<SRCPanelExposedField>(Entity))
+			if (InOutName == NAME_None)
 			{
-				if (ExposedField->GetPropertyId() == FieldKey)
-				{
-					ChildWidgets.Add(Entity);
-				}
+				InOutName = InOtherName;
 			}
-		}
-	}
-	else if (GroupType == EFieldGroupType::Owner)
+			else if (InOutName != InOtherName)
+			{
+				InOutName = TEXT("Multiple Values");
+			}
+		};
+
+	for (const TSharedPtr<SRCPanelTreeNode>& Entity : InFieldEntities)
 	{
-		for (const TSharedPtr<SRCPanelTreeNode>& Entity : InFieldEntities)
+		if (const TSharedPtr<SRCPanelExposedField> ExposedField = StaticCastSharedPtr<SRCPanelExposedField>(Entity))
 		{
-			if (const TSharedPtr<SRCPanelExposedField> ExposedField = StaticCastSharedPtr<SRCPanelExposedField>(Entity))
+			const FName FieldOwnerName  = ExposedField->GetOwnerName();
+			const FName FieldPropertyId = ExposedField->GetPropertyId();
+
+			const bool bOwnerMatch = GroupType == EFieldGroupType::Owner && FieldOwnerName == FieldKey;
+			const bool bPropertyIdMatch = GroupType == EFieldGroupType::PropertyId && FieldPropertyId == FieldKey;
+
+			if (bOwnerMatch || bPropertyIdMatch)
 			{
-				if (ExposedField->GetOwnerName() == FieldKey)
-				{
-					ChildWidgets.Add(Entity);
-				}
+				ChildWidgets.Add(Entity);
+				CompareUpdateName(OwnerName, FieldOwnerName);
+				CompareUpdateName(PropertyIdName, FieldPropertyId);
 			}
 		}
 	}
@@ -141,3 +147,5 @@ void SRCPanelExposedEntitiesGroup::GetNodeChildren(TArray<TSharedPtr<SRCPanelTre
 	OutChildren.Append(ChildWidgets);
 	SRCPanelTreeNode::GetNodeChildren(OutChildren);
 }
+
+#undef LOCTEXT_NAMESPACE
diff --git a/Engine/Plugins/VirtualProduction/RemoteControl/Source/RemoteControlUI/Private/UI/SRCPanelExposedEntitiesGroup.h b/Engine/Plugins/VirtualProduction/RemoteControl/Source/RemoteControlUI/Private/UI/SRCPanelExposedEntitiesGroup.h
--- a/Engine/Plugins/VirtualProduction/RemoteControl/Source/RemoteControlUI/Private/UI/SRCPanelExposedEntitiesGroup.h
+++ b/Engine/Plugins/VirtualProduction/RemoteControl/Source/RemoteControlUI/Private/UI/SRCPanelExposedEntitiesGroup.h
@@ -10,9 +10,14 @@
 /** Grouping type available */
 enum class EFieldGroupType
 {
+	/** No Grouping. Will be shown as a flat list */
 	None,
+
+	/** Fields with same property id are grouped together */
 	PropertyId,
-	Owner
+
+	/** Fields with same owner are grouped together */
+	Owner,
 };
 
 ENUM_CLASS_FLAGS(EFieldGroupType)
@@ -76,6 +81,12 @@ private:
 	/** Field key value of this group */
 	FName FieldKey;
 
+	/** Shared owner of the entity group */
+	FName OwnerName;
+
+	/** Shared property id of the entity group */
+	FName PropertyIdName;
+
 	/** Field group type of this group */
 	EFieldGroupType GroupType = EFieldGroupType::None;
 
diff --git a/Engine/Plugins/VirtualProduction/RemoteControl/Source/RemoteControlUI/Private/UI/SRCPanelExposedEntitiesList.cpp b/Engine/Plugins/VirtualProduction/RemoteControl/Source/RemoteControlUI/Private/UI/SRCPanelExposedEntitiesList.cpp
--- a/Engine/Plugins/VirtualProduction/RemoteControl/Source/RemoteControlUI/Private/UI/SRCPanelExposedEntitiesList.cpp
+++ b/Engine/Plugins/VirtualProduction/RemoteControl/Source/RemoteControlUI/Private/UI/SRCPanelExposedEntitiesList.cpp
@@ -84,41 +84,34 @@ public:
 		ActiveProtocol = InArgs._ActiveProtocol;
 		Entity = InArgs._Entity;
 
-		FSuperRowType::FArguments SuperArgs = FSuperRowType::FArguments();
-
-		SuperArgs.OnCanAcceptDrop_Lambda([this] (const FDragDropEvent& InDragDropEvent, EItemDropZone InDropZone, TSharedPtr<SRCPanelTreeNode> Node)
-		{
-			return InDropZone;
-		});
-
-		SuperArgs.OnAcceptDrop(InArgs._OnAcceptDrop);
-		SuperArgs.OnDragDetected(InArgs._OnDragDetected);
-		SuperArgs.OnDragEnter(InArgs._OnDragEnter);
-		SuperArgs.OnDragLeave(InArgs._OnDragLeave);
-
-		SuperArgs.ExpanderStyleSet(&FCoreStyle::Get());
-		SuperArgs.Padding(InArgs._Padding);
-		SuperArgs.ShowWires(false);
-		SuperArgs.Style(InArgs._Style);
+		FSuperRowType::FArguments SuperArgs = FSuperRowType::FArguments()
+			.OnCanAcceptDrop_Lambda(
+				[](const FDragDropEvent& InDragDropEvent, EItemDropZone InDropZone, TSharedPtr<SRCPanelTreeNode> Node)
+				{
+					return InDropZone;
+				})
+			.OnAcceptDrop(InArgs._OnAcceptDrop)
+			.OnDragDetected(InArgs._OnDragDetected)
+			.OnDragEnter(InArgs._OnDragEnter)
+			.OnDragLeave(InArgs._OnDragLeave)
+			.ExpanderStyleSet(&FCoreStyle::Get())
+			.Padding(InArgs._Padding)
+			.ShowWires(false)
+			.Style(InArgs._Style);
 
 		FSuperRowType::Construct(SuperArgs, OwnerTableView);
 	}
 
 	TSharedRef<SWidget> GenerateWidgetForColumn(const FName& InColumnName) override
 	{
-		const FName& ActiveProtocolName = ActiveProtocol.Get(NAME_None);
-
 		if (Entity.IsValid())
 		{
-			if (Entity->HasChildren() && InColumnName == RemoteControlPresetColumns::Description)
-			{
-				// -- Row is for TreeView --
-				SHorizontalBox::FSlot* InnerContentSlotNativePtr = nullptr;
+			TSharedRef<SHorizontalBox> HorizontalBox = SNew(SHorizontalBox)
+				.Clipping(EWidgetClipping::OnDemand);
 
-				TSharedRef<SWidget> TreeNode = SNew(SHorizontalBox)
-					.Clipping(EWidgetClipping::OnDemand)
-					
-					+ SHorizontalBox::Slot()
+			if (InColumnName == RemoteControlPresetColumns::OwnerName)
+			{
+				HorizontalBox->AddSlot()
 					.AutoWidth()
 					.HAlign(HAlign_Right)
 					.VAlign(VAlign_Fill)
@@ -126,35 +119,30 @@ public:
 						SAssignNew(ExpanderArrowWidget, SExpanderArrow, SharedThis(this))
 						.StyleSet(ExpanderStyleSet)
 						.ShouldDrawWires(false)
-					]
-
-					+ SHorizontalBox::Slot()
-					.FillWidth(1)
-					.Expose(InnerContentSlotNativePtr)
-					.Padding(8.f, 2.f)
-					[
-						Entity->GetWidget(InColumnName, ActiveProtocolName)
 					];
+			}
 
-				InnerContentSlot = InnerContentSlotNativePtr;
-
-				return TreeNode;
-
+			// todo: this alignment should be handled by the generated widgets for each column themselves
+			EHorizontalAlignment HorizontalAlignment = HAlign_Fill;
+			if (InColumnName == RemoteControlPresetColumns::Status || InColumnName == RemoteControlPresetColumns::BindingStatus)
+			{
+				HorizontalAlignment = HAlign_Center;
 			}
-			else
+			else if (InColumnName == RemoteControlPresetColumns::Mask)
 			{
-				const bool bAlignCenter = InColumnName == RemoteControlPresetColumns::Status || InColumnName == RemoteControlPresetColumns::BindingStatus;
-				
-				const bool bAlignLeft = InColumnName == RemoteControlPresetColumns::Mask;
-
-				return SNew(SBox)
-					.Padding(FMargin(4.f, 2.f))
-					.HAlign(bAlignCenter ? HAlign_Center : bAlignLeft ? HAlign_Left : HAlign_Fill)
-					.VAlign(VAlign_Center)
-					[
-						Entity->GetWidget(InColumnName, ActiveProtocolName)
-					];
+				HorizontalAlignment = HAlign_Left;
 			}
+
+			HorizontalBox->AddSlot()
+				.FillWidth(1.f)
+				.Padding(2.f, 0.f)
+				.HAlign(HorizontalAlignment)
+				.VAlign(VAlign_Center)
+				[
+					Entity->GetWidget(InColumnName, ActiveProtocol.Get(NAME_None))
+				];
+
+			return HorizontalBox;
 		}
 
 		return SNullWidget::NullWidget;
@@ -369,7 +357,7 @@ void SRCPanelExposedEntitiesList::Construct(const FArguments& InArgs, URemoteCon
 			.HeaderContentPadding(RCPanelStyle->HeaderRowPadding)
 
 			+ SRCHeaderRow::Column(RemoteControlPresetColumns::OwnerName)
-			.DefaultLabel(LOCTEXT("RCPresetOwnerNameColumnHeader", "Owner Name"))
+			.DefaultLabel(LOCTEXT("RCPresetOwnerNameColumnHeader", "Owner"))
 			.HAlignHeader(HAlign_Center)
 			.FillWidth(0.1f)
 			.HeaderContentPadding(RCPanelStyle->HeaderRowPadding)
@@ -1412,6 +1400,17 @@ TSharedRef<SWidget> SRCPanelExposedEntitiesList::GetGroupMenuContentWidget()
 	// Grouping type section
 	MenuBuilder.BeginSection(FName(TEXT("RCGroupingType")), LOCTEXT("RCGroupingTypeHeader", "Group by"));
 
+	MenuBuilder.AddMenuEntry(
+		LOCTEXT("RCNoGroupingLabel", "No Grouping"),
+		LOCTEXT("RCNoGroupingTooltip", "No grouping set. Show items as a flat list"),
+		FSlateIcon(),
+		FUIAction(
+		FExecuteAction::CreateSP(this, &SRCPanelExposedEntitiesList::OnCreateFieldGroup, EFieldGroupType::None),
+		FCanExecuteAction(),
+		FIsActionChecked::CreateSPLambda(this, [this]{ return CurrentGroupType == EFieldGroupType::None; })),
+		NAME_None,
+		EUserInterfaceActionType::RadioButton);
+
 	MenuBuilder.AddMenuEntry(
 		LOCTEXT("RCGroupPropertyIdLabel", "Property Id"),
 		LOCTEXT("RCGroupPropertyIdTooltip", "Group by Property Id"),
@@ -1419,7 +1418,7 @@ TSharedRef<SWidget> SRCPanelExposedEntitiesList::GetGroupMenuContentWidget()
 		FUIAction(
 		FExecuteAction::CreateSP(this, &SRCPanelExposedEntitiesList::OnCreateFieldGroup, EFieldGroupType::PropertyId),
 		FCanExecuteAction(),
-		FIsActionChecked::CreateLambda([this]() { return CurrentGroupType == EFieldGroupType::PropertyId; })
+		FIsActionChecked::CreateSPLambda(this, [this]{ return CurrentGroupType == EFieldGroupType::PropertyId; })
 		),
 		NAME_None,
 		EUserInterfaceActionType::RadioButton);
@@ -1431,7 +1430,7 @@ TSharedRef<SWidget> SRCPanelExposedEntitiesList::GetGroupMenuContentWidget()
 		FUIAction(
 		FExecuteAction::CreateSP(this, &SRCPanelExposedEntitiesList::OnCreateFieldGroup, EFieldGroupType::Owner),
 		FCanExecuteAction(),
-		FIsActionChecked::CreateLambda([this]() { return CurrentGroupType == EFieldGroupType::Owner; })
+		FIsActionChecked::CreateSPLambda(this, [this]{ return CurrentGroupType == EFieldGroupType::Owner; })
 		),
 		NAME_None,
 		EUserInterfaceActionType::RadioButton);
@@ -1448,7 +1447,7 @@ TSharedRef<SWidget> SRCPanelExposedEntitiesList::GetGroupMenuContentWidget()
 		FUIAction(
 		FExecuteAction::CreateSP(this, &SRCPanelExposedEntitiesList::OnGroupOrderChanged, ERCGroupOrder::Ascending),
 		FCanExecuteAction(),
-		FIsActionChecked::CreateLambda([this] () { return CurrentGroupSortType == ERCGroupOrder::Ascending; })
+		FIsActionChecked::CreateSPLambda(this, [this]{ return CurrentGroupSortType == ERCGroupOrder::Ascending; })
 		),
 		NAME_None,
 		EUserInterfaceActionType::RadioButton);
@@ -1460,7 +1459,7 @@ TSharedRef<SWidget> SRCPanelExposedEntitiesList::GetGroupMenuContentWidget()
 		FUIAction(
 		FExecuteAction::CreateSP(this, &SRCPanelExposedEntitiesList::OnGroupOrderChanged, ERCGroupOrder::Descending),
 		FCanExecuteAction(),
-		FIsActionChecked::CreateLambda([this]() { return CurrentGroupSortType == ERCGroupOrder::Descending; })
+		FIsActionChecked::CreateSPLambda(this, [this]{ return CurrentGroupSortType == ERCGroupOrder::Descending; })
 		),
 		NAME_None,
 		EUserInterfaceActionType::RadioButton);
@@ -1472,7 +1471,7 @@ TSharedRef<SWidget> SRCPanelExposedEntitiesList::GetGroupMenuContentWidget()
 
 void SRCPanelExposedEntitiesList::OnCreateFieldGroup(EFieldGroupType InFieldGroupType)
 {
-	CurrentGroupType = CurrentGroupType == InFieldGroupType? EFieldGroupType::None : InFieldGroupType;
+	CurrentGroupType = InFieldGroupType;
 	CreateFieldGroup();
 	OrderGroups();
 }
