Remote Control: Added Signatures.
Signatures are 'recipes' for properties that can be made by dragging Property Entities into a new signature entry.
These signatures can then be applied by dragging any number of actors to it.

diff --git a/Engine/Plugins/VirtualProduction/RemoteControl/Source/RemoteControl/Private/RemoteControlPreset.cpp b/Engine/Plugins/VirtualProduction/RemoteControl/Source/RemoteControl/Private/RemoteControlPreset.cpp
--- a/Engine/Plugins/VirtualProduction/RemoteControl/Source/RemoteControl/Private/RemoteControlPreset.cpp
+++ b/Engine/Plugins/VirtualProduction/RemoteControl/Source/RemoteControl/Private/RemoteControlPreset.cpp
@@ -28,6 +28,7 @@
 #include "RemoteControlLogger.h"
 #include "RemoteControlObjectVersion.h"
 #include "RemoteControlPresetRebindingManager.h"
+#include "RemoteControlSignatureRegistry.h"
 #include "RemoteControlTransactionListenerHelper.h"
 
 #include "UObject/Object.h"
@@ -673,6 +674,8 @@ URemoteControlPreset::URemoteControlPreset()
 {
 	Registry = CreateDefaultSubobject<URemoteControlExposeRegistry>(FName("ExposeRegistry"));
 
+	SignatureRegistry = CreateDefaultSubobject<URemoteControlSignatureRegistry>(TEXT("SignatureRegistry"));
+
 	PropertyIdRegistry = CreateDefaultSubobject<URemoteControlPropertyIdRegistry>(FName("PropertyIdRegistry"));
 	PropertyIdRegistry->Initialize();
 }
diff --git a/Engine/Plugins/VirtualProduction/RemoteControl/Source/RemoteControl/Private/RemoteControlSignature.cpp b/Engine/Plugins/VirtualProduction/RemoteControl/Source/RemoteControl/Private/RemoteControlSignature.cpp
new file mode 100644
--- /dev/null
+++ b/Engine/Plugins/VirtualProduction/RemoteControl/Source/RemoteControl/Private/RemoteControlSignature.cpp
@@ -0,0 +1,122 @@
+// Copyright Epic Games, Inc. All Rights Reserved.
+
+#include "RemoteControlSignature.h"
+#include "GameFramework/Actor.h"
+#include "RemoteControlField.h"
+#include "RemoteControlPreset.h"
+
+int32 FRCSignature::AddFieldsFromEntities(URemoteControlPreset* InPreset, TConstArrayView<FGuid> InFieldEntityIds)
+{
+	if (!InPreset || InFieldEntityIds.IsEmpty())
+	{
+		return 0;
+	}
+
+	const int32 PreviousFieldCount = Fields.Num();
+
+	Fields.Reserve(PreviousFieldCount + InFieldEntityIds.Num());
+
+	for (const FGuid& EntityId : InFieldEntityIds)
+	{
+		TSharedPtr<FRemoteControlProperty> ExposedProperty = InPreset->GetExposedEntity<FRemoteControlProperty>(EntityId).Pin();
+		if (!ExposedProperty.IsValid())
+		{
+			continue;
+		}
+
+		FRCSignatureField Field;
+		Field.bEnabled = true;
+		Field.FieldPath = ExposedProperty->FieldPathInfo;
+		Field.SupportedClass = ExposedProperty->GetSupportedBindingClass();
+
+		if (UObject* BoundObject = ExposedProperty->GetBoundObject())
+		{
+			if (AActor* ActorOwner = BoundObject->GetTypedOuter<AActor>())
+			{
+				Field.ObjectRelativePath = BoundObject->GetPathName(ActorOwner);
+			}
+		}
+
+		Fields.AddUnique(MoveTemp(Field));
+	}
+
+	return Fields.Num() - PreviousFieldCount;
+}
+
+int32 FRCSignature::ApplySignature(URemoteControlPreset* InPreset, TConstArrayView<TWeakObjectPtr<AActor>> InActors) const
+{
+	if (!InPreset || InActors.IsEmpty())
+	{
+		return 0;
+	}
+
+	// Resolve Weak Actors
+	TArray<AActor*> ResolvedActors;
+	ResolvedActors.Reserve(InActors.Num());
+	for (const TWeakObjectPtr<AActor>& ActorWeak : InActors)
+	{
+		if (AActor* Actor = ActorWeak.Get())
+		{
+			ResolvedActors.Add(Actor);
+		}
+	}
+
+	if (ResolvedActors.IsEmpty())
+	{
+		return false;
+	}
+
+	int32 ExposeCount = 0;
+
+	FRemoteControlPresetExposeArgs ExposeArgs;
+
+	for (const FRCSignatureField& Field : Fields)
+	{
+		if (!Field.bEnabled)
+		{
+			continue;
+		}
+
+		// Resolve, not load. It should be loaded already if relevant to the Actors to apply this signature to
+		UClass* SupportedClass = Field.SupportedClass.ResolveClass();
+		UClass* FindClass = SupportedClass ? SupportedClass : UObject::StaticClass();
+
+		FRCFieldPathInfo Path = Field.FieldPath;
+		for (AActor* Actor : ResolvedActors)
+		{
+			UObject* Context;
+			if (Field.ObjectRelativePath.IsEmpty())
+			{
+				Context = Actor;
+			}
+			else
+			{
+				Context = StaticFindObject(FindClass, Actor, *Field.ObjectRelativePath);
+
+				// Slow Path: if Subobject Path did not find the object, try to find the first sub-object of the actor matching the class.
+				if (!Context && SupportedClass)
+				{
+					ForEachObjectWithOuterBreakable(Actor, [&Context, SupportedClass](UObject* InSubobject)->bool
+						{
+							UClass* SubobjectClass = InSubobject->GetClass();
+							if (SubobjectClass && SubobjectClass->IsChildOf(SupportedClass))
+							{
+								Context = InSubobject;
+								return false;
+							}
+							return true;
+						}
+						, /*bIncludeNestedObjects*/true);
+				}
+			}
+
+			if (Context && Path.Resolve(Context))
+			{
+				InPreset->ExposeProperty(Context, Field.FieldPath, ExposeArgs);
+				++ExposeCount;
+			}
+		}
+	}
+
+	return ExposeCount;
+}
diff --git a/Engine/Plugins/VirtualProduction/RemoteControl/Source/RemoteControl/Private/RemoteControlSignatureRegistry.cpp b/Engine/Plugins/VirtualProduction/RemoteControl/Source/RemoteControl/Private/RemoteControlSignatureRegistry.cpp
new file mode 100644
--- /dev/null
+++ b/Engine/Plugins/VirtualProduction/RemoteControl/Source/RemoteControl/Private/RemoteControlSignatureRegistry.cpp
@@ -0,0 +1,38 @@
+// Copyright Epic Games, Inc. All Rights Reserved.
+
+#include "RemoteControlSignatureRegistry.h"
+
+#define LOCTEXT_NAMESPACE "RemoteControlSignatureRegistry"
+
+const FRCSignature* URemoteControlSignatureRegistry::FindSignature(const FGuid& InSignatureId) const
+{
+	return Signatures.FindByKey(InSignatureId);
+}
+
+FRCSignature* URemoteControlSignatureRegistry::FindSignatureMutable(const FGuid& InSignatureId)
+{
+	return Signatures.FindByKey(InSignatureId);
+}
+
+FRCSignature& URemoteControlSignatureRegistry::AddSignature()
+{
+	FRCSignature& Signature = Signatures.AddDefaulted_GetRef();
+	Signature.DisplayName = LOCTEXT("NewSignatureDisplayName", "New Signature");
+	Signature.Id = FGuid::NewGuid();
+	return Signature;
+}
+
+int32 URemoteControlSignatureRegistry::RemoveSignature(const FGuid& InSignatureId)
+{
+	return Signatures.RemoveAll([&InSignatureId](const FRCSignature& InSignature)
+		{
+			return InSignature.Id == InSignatureId;
+		});
+}
+
+void URemoteControlSignatureRegistry::EmptySignatures()
+{
+	Signatures.Empty();
+}
+
+#undef LOCTEXT_NAMESPACE
diff --git a/Engine/Plugins/VirtualProduction/RemoteControl/Source/RemoteControl/Public/RemoteControlPreset.h b/Engine/Plugins/VirtualProduction/RemoteControl/Source/RemoteControl/Public/RemoteControlPreset.h
--- a/Engine/Plugins/VirtualProduction/RemoteControl/Source/RemoteControl/Public/RemoteControlPreset.h
+++ b/Engine/Plugins/VirtualProduction/RemoteControl/Source/RemoteControl/Public/RemoteControlPreset.h
@@ -29,6 +29,7 @@ class URCVirtualPropertyContainerBase;
 class URCVirtualPropertyInContainer;
 class URCVirtualPropertySelfContainer;
 class URemoteControlExposeRegistry;
+class URemoteControlSignatureRegistry;
 class URemoteControlBinding;
 class URemoteControlPreset;
 
@@ -705,6 +706,12 @@ public:
 	/** Generate label from Registry */
 	FName GenerateUniqueLabel(const FName InDesiredName) const;
 
+	/** Returns the preset's signature registry */
+	URemoteControlSignatureRegistry* GetSignatureRegistry() const
+	{
+		return SignatureRegistry;
+	}
+
 public:
 	/** The visual layout for this preset. */
 	UPROPERTY()
@@ -932,10 +939,14 @@ private:
 	UPROPERTY(Transient)
 	TMap<FName, FGuid> NameToGuidMap;
 
-	UPROPERTY(Instanced)
 	/** Holds exposed entities on the preset. */
+	UPROPERTY(Instanced)
 	TObjectPtr<URemoteControlExposeRegistry> Registry = nullptr;
 
+	/** Holds all signatures on the preset. */
+	UPROPERTY(Instanced)
+	TObjectPtr<URemoteControlSignatureRegistry> SignatureRegistry = nullptr;
+
 	/** Holds identities of exposed entities on the preset. */
 	UPROPERTY(Transient)
 	TObjectPtr<URemoteControlPropertyIdRegistry> PropertyIdRegistry = nullptr;
diff --git a/Engine/Plugins/VirtualProduction/RemoteControl/Source/RemoteControl/Public/RemoteControlSignature.h b/Engine/Plugins/VirtualProduction/RemoteControl/Source/RemoteControl/Public/RemoteControlSignature.h
new file mode 100644
--- /dev/null
+++ b/Engine/Plugins/VirtualProduction/RemoteControl/Source/RemoteControl/Public/RemoteControlSignature.h
@@ -0,0 +1,80 @@
+// Copyright Epic Games, Inc. All Rights Reserved.
+
+#pragma once
+
+#include "RemoteControlFieldPath.h"
+#include "RemoteControlSignature.generated.h"
+
+class URemoteControlPreset;
+
+USTRUCT()
+struct FRCSignatureField
+{
+	GENERATED_BODY()
+
+	FRCSignatureField() = default;
+
+	bool operator==(const FRCSignatureField& InOtherField) const
+	{
+		return FieldPath == InOtherField.FieldPath
+			&& ObjectRelativePath == InOtherField.ObjectRelativePath
+			&& SupportedClass == InOtherField.SupportedClass;
+	}
+
+	/** Path Info for this Field */
+	UPROPERTY()
+	FRCFieldPathInfo FieldPath;
+
+	/** Optional: Relative path from an owner (e.g. Actor) to the object owning the property (e.g. an Actor Component) */
+	UPROPERTY()
+	FString ObjectRelativePath;
+
+	/** Object class holding the property */
+	UPROPERTY()
+	FSoftClassPath SupportedClass;
+
+	/** Whether to consider this field when applying a Signature */
+	UPROPERTY()
+	bool bEnabled = true;
+};
+
+USTRUCT()
+struct FRCSignature
+{
+	GENERATED_BODY()
+
+	bool operator==(const FGuid& InSignatureId) const
+	{
+		return Id == InSignatureId;
+	}
+
+	/**
+	 * Gathers the Field Entities from the given preset and field ids and adds them as Signature Fields to this instance
+	 * @param InPreset the preset to check for
+	 * @param InFieldEntityIds the ids of the entities to find and add
+	 */
+	REMOTECONTROL_API int32 AddFieldsFromEntities(URemoteControlPreset* InPreset, TConstArrayView<FGuid> InFieldEntityIds);
+
+	/**
+	 * Applies this Signature to the given Actors by exposing all this Signature's fields to the given preset
+	 * @param InPreset the preset where these properties will be exposed to
+	 * @param InActors the actors whose properties to expose
+	 */
+	REMOTECONTROL_API int32 ApplySignature(URemoteControlPreset* InPreset, TConstArrayView<TWeakObjectPtr<AActor>> InActors) const;
+
+	/** User facing friendly name. Used as the Label when exposing */
+	UPROPERTY()
+	FText DisplayName;
+
+	/** Unique Id identifying this Signature */
+	UPROPERTY(meta=(IgnoreForMemberInitializationTest))
+	FGuid Id;
+
+	/** The fields owned by this Signature */
+	UPROPERTY()
+	TArray<FRCSignatureField> Fields;
+
+	/** Whether this Signature can be applied */
+	UPROPERTY()
+	bool bEnabled = true;
+};
\ No newline at end of file
diff --git a/Engine/Plugins/VirtualProduction/RemoteControl/Source/RemoteControl/Public/RemoteControlSignatureRegistry.h b/Engine/Plugins/VirtualProduction/RemoteControl/Source/RemoteControl/Public/RemoteControlSignatureRegistry.h
new file mode 100644
--- /dev/null
+++ b/Engine/Plugins/VirtualProduction/RemoteControl/Source/RemoteControl/Public/RemoteControlSignatureRegistry.h
@@ -0,0 +1,51 @@
+// Copyright Epic Games, Inc. All Rights Reserved.
+
+#pragma once
+
+#include "RemoteControlSignature.h"
+#include "UObject/Object.h"
+#include "RemoteControlSignatureRegistry.generated.h"
+
+UCLASS(MinimalAPI)
+class URemoteControlSignatureRegistry : public UObject
+{
+	GENERATED_BODY()
+
+public:
+	TConstArrayView<FRCSignature> GetSignatures() const
+	{
+		return Signatures;
+	}
+
+	/**
+	 * Find an existing Signature matching the given id
+	 * @param InSignatureId the id of the signature to find
+	 * @return const pointer to the signature if found, nullptr otherwise 
+	 */
+	REMOTECONTROL_API const FRCSignature* FindSignature(const FGuid& InSignatureId) const;
+
+	/**
+	 * Find an existing Signature matching the given id
+	 * @param InSignatureId the id of the signature to find
+	 * @return pointer to the signature if found, nullptr otherwise 
+	 */
+	REMOTECONTROL_API FRCSignature* FindSignatureMutable(const FGuid& InSignatureId);
+
+	/** Adds a new Signature to this registry */
+	REMOTECONTROL_API FRCSignature& AddSignature();
+
+	/**
+	 * Removes the Signature with the matching Signature Id
+	 * @param InSignatureId the id of the signature to remove
+	 * @return the number of signatures removed
+	 */
+	REMOTECONTROL_API int32 RemoveSignature(const FGuid& InSignatureId);
+
+	/** Removes all the Signature in this registry */
+	REMOTECONTROL_API void EmptySignatures();
+
+private:
+	/** Holds all the Signatures in this Registry */
+	UPROPERTY()
+	TArray<FRCSignature> Signatures;
+};
diff --git a/Engine/Plugins/VirtualProduction/RemoteControl/Source/RemoteControlUI/Private/Commands/RemoteControlCommands.cpp b/Engine/Plugins/VirtualProduction/RemoteControl/Source/RemoteControlUI/Private/Commands/RemoteControlCommands.cpp
--- a/Engine/Plugins/VirtualProduction/RemoteControl/Source/RemoteControlUI/Private/Commands/RemoteControlCommands.cpp
+++ b/Engine/Plugins/VirtualProduction/RemoteControl/Source/RemoteControlUI/Private/Commands/RemoteControlCommands.cpp
@@ -31,6 +31,9 @@ void FRemoteControlCommands::RegisterCommands()
 	// Toggle Logic Editor
 	UI_COMMAND(ToggleLogicEditor, "Logic", "View the logic applied to active selection.", EUserInterfaceActionType::ToggleButton, FInputChord());
 
+	// Toggle Signature Editor
+	UI_COMMAND(ToggleSignatureEditor, "Signature", "View the current signatures available", EUserInterfaceActionType::ToggleButton, FInputChord());
+
 	// Delete Entity
 	UI_COMMAND(DeleteEntity, "Delete", "Delete the selected group/exposed entity from the list.", EUserInterfaceActionType::Button, FInputChord(EKeys::Delete));
 
diff --git a/Engine/Plugins/VirtualProduction/RemoteControl/Source/RemoteControlUI/Private/Commands/RemoteControlCommands.h b/Engine/Plugins/VirtualProduction/RemoteControl/Source/RemoteControlUI/Private/Commands/RemoteControlCommands.h
--- a/Engine/Plugins/VirtualProduction/RemoteControl/Source/RemoteControlUI/Private/Commands/RemoteControlCommands.h
+++ b/Engine/Plugins/VirtualProduction/RemoteControl/Source/RemoteControlUI/Private/Commands/RemoteControlCommands.h
@@ -35,12 +35,17 @@ public:
 	 * Holds the information about UI Command that brings up a panel which holds the active protocol mappings.
 	 */
 	TSharedPtr<FUICommandInfo> ToggleProtocolMappings;
-	
+
 	/**
 	 * Holds the information about UI Command that brings up a panel which enables the RC Logical Behaviour.
 	 */
 	TSharedPtr<FUICommandInfo> ToggleLogicEditor;
 
+	/**
+	 * Holds the information about UI Command that brings up the RC Signature UI panel.
+	 */
+	TSharedPtr<FUICommandInfo> ToggleSignatureEditor;
+
 	/**
 	 * Holds the information about UI Command that deletes currently selected group/exposed entity.
 	 */
diff --git a/Engine/Plugins/VirtualProduction/RemoteControl/Source/RemoteControlUI/Private/RemoteControlUIModule.cpp b/Engine/Plugins/VirtualProduction/RemoteControl/Source/RemoteControlUI/Private/RemoteControlUIModule.cpp
--- a/Engine/Plugins/VirtualProduction/RemoteControl/Source/RemoteControlUI/Private/RemoteControlUIModule.cpp
+++ b/Engine/Plugins/VirtualProduction/RemoteControl/Source/RemoteControlUI/Private/RemoteControlUIModule.cpp
@@ -611,6 +611,22 @@ uint32 FRemoteControlUIModule::GetRemoteControlAssetCategory() const
 	return RemoteControlAssetCategoryBit;
 }
 
+void FRemoteControlUIModule::RegisterSignatureCustomization(const TSharedPtr<IRCSignatureCustomization>& InCustomization)
+{
+	if (InCustomization.IsValid())
+	{
+		SignatureCustomizations.AddUnique(InCustomization.ToSharedRef());
+	}
+}
+
+void FRemoteControlUIModule::UnregisterSignatureCustomization(const TSharedPtr<IRCSignatureCustomization>& InCustomization)
+{
+	if (InCustomization.IsValid())
+	{
+		SignatureCustomizations.Remove(InCustomization.ToSharedRef());
+	}
+}
+
 void FRemoteControlUIModule::RegisterAssetTools()
 {
 	if (FAssetToolsModule* AssetToolsModule = FModuleManager::GetModulePtr<FAssetToolsModule>("AssetTools"))
@@ -650,6 +666,7 @@ void FRemoteControlUIModule::UnbindRemoteControlCommands()
 		ActionList.UnmapAction(Commands.FindPresetInContentBrowser);
 		ActionList.UnmapAction(Commands.ToggleProtocolMappings);
 		ActionList.UnmapAction(Commands.ToggleLogicEditor);
+		ActionList.UnmapAction(Commands.ToggleSignatureEditor);
 		ActionList.UnmapAction(Commands.DeleteEntity);
 		ActionList.UnmapAction(Commands.RenameEntity);
 		ActionList.UnmapAction(Commands.CopyItem);
diff --git a/Engine/Plugins/VirtualProduction/RemoteControl/Source/RemoteControlUI/Private/RemoteControlUIModule.h b/Engine/Plugins/VirtualProduction/RemoteControl/Source/RemoteControlUI/Private/RemoteControlUIModule.h
--- a/Engine/Plugins/VirtualProduction/RemoteControl/Source/RemoteControlUI/Private/RemoteControlUIModule.h
+++ b/Engine/Plugins/VirtualProduction/RemoteControl/Source/RemoteControlUI/Private/RemoteControlUIModule.h
@@ -56,6 +56,8 @@ public:
 	virtual void UnregisterMetadataCustomization(FName MetadataKey) override;
 	virtual URemoteControlPreset* GetActivePreset() const override;
 	virtual uint32 GetRemoteControlAssetCategory() const override;
+	virtual void RegisterSignatureCustomization(const TSharedPtr<IRCSignatureCustomization>& InCustomization) override;
+	virtual void UnregisterSignatureCustomization(const TSharedPtr<IRCSignatureCustomization>& InCustomization) override;
 	virtual void RegisterWidgetFactoryForType(UScriptStruct* RemoteControlEntityType, const FOnGenerateRCWidget& OnGenerateRCWidgetDelegate) override;
 	virtual void UnregisterWidgetFactoryForType(UScriptStruct* RemoteControlEntityType) override;
 	virtual void HighlightPropertyInDetailsPanel(const FPropertyPath& Path) const override;
@@ -80,6 +82,11 @@ public:
 		return ExternalEntityMetadataCustomizations;
 	}
 
+	TConstArrayView<TSharedRef<IRCSignatureCustomization>> GetSignatureCustomizations() const
+	{
+		return SignatureCustomizations;
+	}
+
 public:
 	static const FName RemoteControlPanelTabName;
 
@@ -268,6 +275,9 @@ private:
 	/** Map of metadata key to customization handler. */
 	TMap<FName, FOnCustomizeMetadataEntry> ExternalEntityMetadataCustomizations;
 
+	/** Registered Signature Customizations */
+	TArray<TSharedRef<IRCSignatureCustomization>> SignatureCustomizations; 
+
 	TMap<TWeakObjectPtr<UScriptStruct>, FOnGenerateRCWidget> GenerateWidgetDelegates;
 
 	/** Holds the advanced asset type category bit registered by Remote Control. */
diff --git a/Engine/Plugins/VirtualProduction/RemoteControl/Source/RemoteControlUI/Private/UI/BaseLogicUI/SRCLogicPanelBase.h b/Engine/Plugins/VirtualProduction/RemoteControl/Source/RemoteControlUI/Private/UI/BaseLogicUI/SRCLogicPanelBase.h
--- a/Engine/Plugins/VirtualProduction/RemoteControl/Source/RemoteControlUI/Private/UI/BaseLogicUI/SRCLogicPanelBase.h
+++ b/Engine/Plugins/VirtualProduction/RemoteControl/Source/RemoteControlUI/Private/UI/BaseLogicUI/SRCLogicPanelBase.h
@@ -39,9 +39,15 @@ public:
 	/** "Delete Items" UI command implementation for panels*/
 	virtual void DeleteSelectedPanelItems() = 0;
 
+	/** Whether this panel can Duplicate Items */
+	virtual bool CanDuplicateItems() const { return !GetSelectedLogicItems().IsEmpty(); }
+
 	/** "Duplicate Items" UI command implementation for action panels */
 	virtual void DuplicateSelectedPanelItems() {}
 
+	/** Whether this panel can Copy Items */
+	virtual bool CanCopyItems() const { return !GetSelectedLogicItems().IsEmpty(); }
+
 	/** "Copy Items" UI command implementation for panels*/
 	virtual void CopySelectedPanelItems() {}
 
diff --git a/Engine/Plugins/VirtualProduction/RemoteControl/Source/RemoteControlUI/Private/UI/Drawers/SRCPanelDrawer.h b/Engine/Plugins/VirtualProduction/RemoteControl/Source/RemoteControlUI/Private/UI/Drawers/SRCPanelDrawer.h
--- a/Engine/Plugins/VirtualProduction/RemoteControl/Source/RemoteControlUI/Private/UI/Drawers/SRCPanelDrawer.h
+++ b/Engine/Plugins/VirtualProduction/RemoteControl/Source/RemoteControlUI/Private/UI/Drawers/SRCPanelDrawer.h
@@ -26,6 +26,8 @@ enum class ERCPanels : uint8
 	RCP_Live,
 
 	RCP_Count,
+
+	RCP_Signature,
 };
 
 /** State of the RC Panels. */
diff --git a/Engine/Plugins/VirtualProduction/RemoteControl/Source/RemoteControlUI/Private/UI/SRemoteControlPanel.cpp b/Engine/Plugins/VirtualProduction/RemoteControl/Source/RemoteControlUI/Private/UI/SRemoteControlPanel.cpp
--- a/Engine/Plugins/VirtualProduction/RemoteControl/Source/RemoteControlUI/Private/UI/SRemoteControlPanel.cpp
+++ b/Engine/Plugins/VirtualProduction/RemoteControl/Source/RemoteControlUI/Private/UI/SRemoteControlPanel.cpp
@@ -55,6 +55,7 @@
 #include "SceneOutlinerModule.h"
 #include "SceneOutlinerPublicTypes.h"
 #include "ScopedTransaction.h"
+#include "Signature/SRCSignaturePanel.h"
 #include "Styling/RemoteControlStyles.h"
 #include "Styling/ToolBarStyle.h"
 #include "Subsystems/Subsystem.h"
@@ -487,12 +488,15 @@ void SRemoteControlPanel::Construct(const FArguments& InArgs, URemoteControlPres
 		.EnableFooter(false);
 
 	ControllerPanel = SNew(SRCControllerPanel, SharedThis(this))
-		.LiveMode_Lambda([this]() { return bIsInLiveMode; })
+		.LiveMode_Lambda([this] { return bIsInLiveMode; })
 		.Visibility_Lambda([this] { return (bIsLogicPanelEnabled || bIsInLiveMode) && (ActivePanel == ERCPanels::RCP_Properties || ActivePanel == ERCPanels::RCP_Live || ActivePanel == ERCPanels::RCP_None) ? EVisibility::Visible : EVisibility::Collapsed; });
 
 	BehaviourPanel = SNew(SRCBehaviourPanel, SharedThis(this))
 		.Visibility_Lambda([this] { return !bIsInLiveMode && bIsLogicPanelEnabled && (ActivePanel == ERCPanels::RCP_Properties || ActivePanel == ERCPanels::RCP_None) ? EVisibility::Visible : EVisibility::Collapsed; });
 
+	SignaturePanel = SNew(SRCSignaturePanel, SharedThis(this))
+		.LiveMode_Lambda([this] { return bIsInLiveMode; });
+
 	ControllersAndBehavioursPanel->AddPanel(ControllerPanel.ToSharedRef(), 0.8f);
 
 	constexpr bool bResizable = false;
@@ -710,6 +714,31 @@ void SRemoteControlPanel::Construct(const FArguments& InArgs, URemoteControlPres
 					LogicPanel
 				]
 			]
+
+			// 7. (unused) Count (Index : 6)
+			+ SWidgetSwitcher::Slot()
+			[
+				SNullWidget::NullWidget
+			]
+
+			// 8. Signatures (Index : 7).
+			+ SWidgetSwitcher::Slot()
+			[
+				SNew(SSplitter)
+				.Orientation(Orient_Horizontal)
+
+				+SSplitter::Slot()
+				.Value(0.4)
+				[
+					EntityList.ToSharedRef()
+				]
+
+				+SSplitter::Slot()
+				.Value(0.6)
+				[
+					SignaturePanel.ToSharedRef()
+				]
+			]
 		];
 
 	ChildSlot
@@ -1410,6 +1439,13 @@ void SRemoteControlPanel::BindRemoteControlCommands()
 		FIsActionChecked::CreateSP(this, &SRemoteControlPanel::IsLogicPanelEnabled),
 		FIsActionButtonVisible::CreateSP(this, &SRemoteControlPanel::CanToggleLogicPanel));
 
+	ActionList.MapAction(
+		Commands.ToggleSignatureEditor,
+		FExecuteAction::CreateSP(this, &SRemoteControlPanel::ToggleSignatureEditor_Execute),
+		FCanExecuteAction::CreateSP(this, &SRemoteControlPanel::CanToggleSignaturePanel),
+		FIsActionChecked::CreateSP(this, &SRemoteControlPanel::IsSignaturePanelEnabled),
+		FIsActionButtonVisible::CreateSP(this, &SRemoteControlPanel::CanToggleSignaturePanel));
+
 	ActionList.MapAction(
 		Commands.DeleteEntity,
 		FExecuteAction::CreateSP(this, &SRemoteControlPanel::DeleteEntity_Execute),
@@ -1428,22 +1464,30 @@ void SRemoteControlPanel::BindRemoteControlCommands()
 	ActionList.MapAction(
 		Commands.CopyItem,
 		FExecuteAction::CreateSP(this, &SRemoteControlPanel::CopyItem_Execute),
-		FCanExecuteAction::CreateSP(this, &SRemoteControlPanel::CanCopyItem));
+		FCanExecuteAction(),
+		FGetActionCheckState(),
+		FIsActionButtonVisible::CreateSP(this, &SRemoteControlPanel::CanCopyItem));
 
 	ActionList.MapAction(
 		Commands.PasteItem,
 		FExecuteAction::CreateSP(this, &SRemoteControlPanel::PasteItem_Execute),
-		FCanExecuteAction::CreateSP(this, &SRemoteControlPanel::CanPasteItem));
+		FCanExecuteAction(),
+		FGetActionCheckState(),
+		FIsActionButtonVisible::CreateSP(this, &SRemoteControlPanel::CanPasteItem));
 
 	ActionList.MapAction(
 		Commands.DuplicateItem,
 		FExecuteAction::CreateSP(this, &SRemoteControlPanel::DuplicateItem_Execute),
-		FCanExecuteAction::CreateSP(this, &SRemoteControlPanel::CanDuplicateItem));
+		FCanExecuteAction(),
+		FGetActionCheckState(),
+		FIsActionButtonVisible::CreateSP(this, &SRemoteControlPanel::CanDuplicateItem));
 
 	ActionList.MapAction(
 		Commands.UpdateValue,
 		FExecuteAction::CreateSP(this, &SRemoteControlPanel::UpdateValue_Execute),
-		FCanExecuteAction::CreateSP(this, &SRemoteControlPanel::CanUpdateValue));
+		FCanExecuteAction(),
+		FGetActionCheckState(),
+		FIsActionButtonVisible::CreateSP(this, &SRemoteControlPanel::CanUpdateValue));
 }
 
 void SRemoteControlPanel::OnObjectReplaced(const TMap<UObject*, UObject*>& InObjectReplaced)
@@ -1593,6 +1637,21 @@ void SRemoteControlPanel::RegisterPanels()
 			RegisteredDrawers.Add(PropertiesPanel->GetPanelID(), PropertiesPanel);
 		}
 
+		{// Signature Panel
+			TSharedRef<FRCPanelDrawerArgs> SignaturePanelDrawer = MakeShared<FRCPanelDrawerArgs>(ERCPanels::RCP_Signature);
+
+			SignaturePanelDrawer->bDrawnByDefault = true;
+			SignaturePanelDrawer->bRotateIconBy90 = false;
+			SignaturePanelDrawer->DrawerVisibility = bIsSignaturePanelEnabled ? EVisibility::Visible : EVisibility::Collapsed;
+			SignaturePanelDrawer->Label = LOCTEXT("SignaturePanelLabel", "Signature");
+			SignaturePanelDrawer->ToolTip = LOCTEXT("SignaturePanelTooltip", "Open Signature panel.");
+			SignaturePanelDrawer->Icon = FSlateIcon(FAppStyle::GetAppStyleSetName(), TEXT("GraphEditor.Function_16x"));
+
+			PanelDrawer->RegisterPanel(SignaturePanelDrawer);
+
+			RegisteredDrawers.Add(SignaturePanelDrawer->GetPanelID(), SignaturePanelDrawer);
+		}
+
 		{// Properties With Details Panel
 			TSharedRef<FRCPanelDrawerArgs> EntityDetailsPanel = MakeShared<FRCPanelDrawerArgs>(ERCPanels::RCP_EntityDetails);
 
@@ -2153,7 +2212,7 @@ void SRemoteControlPanel::GenerateToolbar()
 		.AutoWidth()
 		[
 			SNew(SComboButton)
-			.ToolTipText(LOCTEXT("RCFieldsGroupingTooltip", "Select grouping type for the fields"))
+			.ToolTipText(LOCTEXT("RSWorldSelectorTooltip", "Select World to Control"))
 			.HAlign(HAlign_Center)
 			.VAlign(VAlign_Center)
 			.ButtonStyle(&RCPanelStyle->FlatButtonStyle)
@@ -2323,8 +2382,21 @@ void SRemoteControlPanel::RegisterAuxiliaryToolBar()
 
 			const FRemoteControlCommands& Commands = FRemoteControlCommands::Get();
 
-			ToolsSection.AddEntry(FToolMenuEntry::InitWidget("Logic"
-			, SNew(SVerticalBox)
+			ToolsSection.AddEntry(FToolMenuEntry::InitWidget("Signature",
+				SNew(SVerticalBox)
+				+ SVerticalBox::Slot()
+				.AutoHeight()
+				[
+					SNew(SAutoResizeButton)
+					.UICommand(FRemoteControlCommands::Get().ToggleSignatureEditor)
+					.ForceSmallIcons_Static(SRemoteControlPanel::ShouldForceSmallIcons)
+					.IconOverride(FSlateIcon(FAppStyle::GetAppStyleSetName(), TEXT("GraphEditor.Function_16x")))
+				]
+				, Commands.ToggleSignatureEditor->GetLabel())
+			);
+
+			ToolsSection.AddEntry(FToolMenuEntry::InitWidget("Logic",
+				SNew(SVerticalBox)
 				+ SVerticalBox::Slot()
 				.AutoHeight()
 				[
@@ -2333,12 +2405,11 @@ void SRemoteControlPanel::RegisterAuxiliaryToolBar()
 					.ForceSmallIcons_Static(SRemoteControlPanel::ShouldForceSmallIcons)
 					.IconOverride(FSlateIcon(FAppStyle::GetAppStyleSetName(), TEXT("GraphEditor.StateMachine_16x")))
 				]
-				, Commands.ToggleLogicEditor->GetLabel()
-			)
+				, Commands.ToggleLogicEditor->GetLabel())
 			);
 
-			ToolsSection.AddEntry(FToolMenuEntry::InitWidget("Protocols"
-			, SNew(SVerticalBox)
+			ToolsSection.AddEntry(FToolMenuEntry::InitWidget("Protocols",
+				SNew(SVerticalBox)
 				+ SVerticalBox::Slot()
 				.AutoHeight()
 				[
@@ -2347,8 +2418,7 @@ void SRemoteControlPanel::RegisterAuxiliaryToolBar()
 					.ForceSmallIcons_Static(SRemoteControlPanel::ShouldForceSmallIcons)
 					.IconOverride(FSlateIcon(FAppStyle::GetAppStyleSetName(), "LevelEditor.Tabs.StatsViewer"))
 				]
-				, Commands.ToggleProtocolMappings->GetLabel()
-			)
+				, Commands.ToggleProtocolMappings->GetLabel())
 			);
 		}
 	}
@@ -2528,6 +2598,35 @@ bool SRemoteControlPanel::IsLogicPanelEnabled() const
 	return bIsLogicPanelEnabled;
 }
 
+void SRemoteControlPanel::ToggleSignatureEditor_Execute()
+{
+	bIsSignaturePanelEnabled = !bIsSignaturePanelEnabled;
+
+	if (PanelDrawer.IsValid())
+	{
+		TSharedRef<FRCPanelDrawerArgs> SignaturePanelDrawer = RegisteredDrawers.FindChecked(ERCPanels::RCP_Signature);
+
+		SignaturePanelDrawer->DrawerVisibility = bIsSignaturePanelEnabled ? EVisibility::Visible : EVisibility::Collapsed;
+		PanelDrawer->TogglePanel(SignaturePanelDrawer, !bIsSignaturePanelEnabled);
+
+		if (!bIsSignaturePanelEnabled)
+		{
+			TSharedRef<FRCPanelDrawerArgs> PropertiesPanel = RegisteredDrawers.FindChecked(ERCPanels::RCP_Properties);
+			PanelDrawer->TogglePanel(PropertiesPanel);
+		}
+	}
+}
+
+bool SRemoteControlPanel::CanToggleSignaturePanel() const
+{
+	return true;
+}
+
+bool SRemoteControlPanel::IsSignaturePanelEnabled() const
+{
+	return bIsSignaturePanelEnabled;
+}
+
 void SRemoteControlPanel::OnRCPanelToggled(ERCPanels InPanelID)
 {
 	if (InPanelID != ActivePanel)
@@ -2557,6 +2656,10 @@ TSharedPtr<class SRCLogicPanelBase> SRemoteControlPanel::GetActiveLogicPanel() c
 	{
 		return ActionPanel;
 	}
+	else if (SignaturePanel->IsListFocused())
+	{
+		return SignaturePanel;
+	}
 
 	return nullptr;
 }
@@ -2730,7 +2833,7 @@ bool SRemoteControlPanel::CanCopyItem() const
 
 	if (const TSharedPtr<SRCLogicPanelBase> ActiveLogicPanel = GetActiveLogicPanel())
 	{
-		return !ActiveLogicPanel->GetSelectedLogicItems().IsEmpty();
+		return ActiveLogicPanel->CanCopyItems();
 	}
 
 	return false;
@@ -2785,7 +2888,7 @@ bool SRemoteControlPanel::CanDuplicateItem() const
 
 	if (const TSharedPtr<SRCLogicPanelBase>& ActiveLogicPanel = GetActiveLogicPanel())
 	{
-		return !ActiveLogicPanel->GetSelectedLogicItems().IsEmpty();
+		return ActiveLogicPanel->CanDuplicateItems();
 	}
 
 	return false;
diff --git a/Engine/Plugins/VirtualProduction/RemoteControl/Source/RemoteControlUI/Private/UI/SRemoteControlPanel.h b/Engine/Plugins/VirtualProduction/RemoteControl/Source/RemoteControlUI/Private/UI/SRemoteControlPanel.h
--- a/Engine/Plugins/VirtualProduction/RemoteControl/Source/RemoteControlUI/Private/UI/SRemoteControlPanel.h
+++ b/Engine/Plugins/VirtualProduction/RemoteControl/Source/RemoteControlUI/Private/UI/SRemoteControlPanel.h
@@ -318,6 +318,10 @@ private:
 	bool CanToggleLogicPanel() const;
 	bool IsLogicPanelEnabled() const;
 
+	void ToggleSignatureEditor_Execute();
+	bool CanToggleSignaturePanel() const;
+	bool IsSignaturePanelEnabled() const;
+
 	void OnRCPanelToggled(ERCPanels InPanelID);
 
 	/** Called when user attempts to delete a group/exposed entity. */
@@ -401,6 +405,8 @@ private:
 	bool bIsInLiveMode = false;
 	/** Whether the logic panel is enabled or not. */
 	bool bIsLogicPanelEnabled = false;
+	/** Whether the signature panel is enabled or not. */
+	bool bIsSignaturePanelEnabled = false;
 	/** Delegate called when the live mode changes. */
 	FOnLiveModeChange OnLiveModeChange;
 	/** Holds the blueprint library picker */
@@ -467,6 +473,8 @@ private:
 	TSharedPtr<class SRCBehaviourPanel> BehaviourPanel;
 	/** Action panel UI widget for Remote Control Logic*/
 	TSharedPtr<class SRCActionPanel> ActionPanel;
+	/** Signature panel UI widget */
+	TSharedPtr<class SRCSignaturePanel> SignaturePanel;
 
 	/** LogicClipboardItems - Holds the items copied from a Logic panel
 	*
diff --git a/Engine/Plugins/VirtualProduction/RemoteControl/Source/RemoteControlUI/Private/UI/Signature/Columns/IRCSignatureColumn.h b/Engine/Plugins/VirtualProduction/RemoteControl/Source/RemoteControlUI/Private/UI/Signature/Columns/IRCSignatureColumn.h
new file mode 100644
--- /dev/null
+++ b/Engine/Plugins/VirtualProduction/RemoteControl/Source/RemoteControlUI/Private/UI/Signature/Columns/IRCSignatureColumn.h
@@ -0,0 +1,33 @@
+// Copyright Epic Games, Inc. All Rights Reserved.
+
+#pragma once
+
+#include "Templates/SharedPointer.h"
+#include "Widgets/Views/SHeaderRow.h"
+
+class FRCSignatureTreeItemBase;
+class SRCSignatureRow;
+class SRCSignatureTree;
+
+class IRCSignatureColumn : public TSharedFromThis<IRCSignatureColumn>
+{
+public:
+	virtual ~IRCSignatureColumn() = default;
+
+	virtual FName GetColumnId() const = 0;
+
+	/**
+	 * Determines whether the Column should be Showing by Default while still be able to toggle it on/off.
+	 * Used when calling SHeaderRow::SetShowGeneratedColumn (requires ShouldGenerateWidget to not be set).
+	 */
+	virtual bool ShouldShowColumnByDefault() const
+	{
+		return false;
+	}
+
+	virtual SHeaderRow::FColumn::FArguments ConstructHeaderRowColumn() = 0;
+
+	virtual TSharedRef<SWidget> ConstructRowWidget(TSharedPtr<FRCSignatureTreeItemBase> InItem
+		, const TSharedRef<SRCSignatureTree>& InList
+		, const TSharedRef<SRCSignatureRow>& InRow) = 0;
+};
diff --git a/Engine/Plugins/VirtualProduction/RemoteControl/Source/RemoteControlUI/Private/UI/Signature/Columns/RCSignatureDescriptionColumn.cpp b/Engine/Plugins/VirtualProduction/RemoteControl/Source/RemoteControlUI/Private/UI/Signature/Columns/RCSignatureDescriptionColumn.cpp
new file mode 100644
--- /dev/null
+++ b/Engine/Plugins/VirtualProduction/RemoteControl/Source/RemoteControlUI/Private/UI/Signature/Columns/RCSignatureDescriptionColumn.cpp
@@ -0,0 +1,32 @@
+// Copyright Epic Games, Inc. All Rights Reserved.
+
+#include "RCSignatureDescriptionColumn.h"
+#include "SRCSignatureDescription.h"
+
+#define LOCTEXT_NAMESPACE "RCSignatureDescriptionColumn"
+
+FName FRCSignatureDescriptionColumn::GetColumnId() const
+{
+	return TEXT("FRCSignatureDescriptionColumn");
+}
+
+bool FRCSignatureDescriptionColumn::ShouldShowColumnByDefault() const
+{
+	return false;
+}
+
+SHeaderRow::FColumn::FArguments FRCSignatureDescriptionColumn::ConstructHeaderRowColumn()
+{
+	return SHeaderRow::Column(GetColumnId())
+		.FillWidth(0.5f)
+		.DefaultLabel(LOCTEXT("DisplayName", "Description"));
+}
+
+TSharedRef<SWidget> FRCSignatureDescriptionColumn::ConstructRowWidget(TSharedPtr<FRCSignatureTreeItemBase> InItem
+	, const TSharedRef<SRCSignatureTree>& InList
+	, const TSharedRef<SRCSignatureRow>& InRow)
+{
+	return SNew(SRCSignatureDescription, InItem.ToSharedRef());
+}
+
+#undef LOCTEXT_NAMESPACE
diff --git a/Engine/Plugins/VirtualProduction/RemoteControl/Source/RemoteControlUI/Private/UI/Signature/Columns/RCSignatureDescriptionColumn.h b/Engine/Plugins/VirtualProduction/RemoteControl/Source/RemoteControlUI/Private/UI/Signature/Columns/RCSignatureDescriptionColumn.h
new file mode 100644
--- /dev/null
+++ b/Engine/Plugins/VirtualProduction/RemoteControl/Source/RemoteControlUI/Private/UI/Signature/Columns/RCSignatureDescriptionColumn.h
@@ -0,0 +1,15 @@
+// Copyright Epic Games, Inc. All Rights Reserved.
+
+#pragma once
+
+#include "UI/Signature/Columns/IRCSignatureColumn.h"
+
+class FRCSignatureDescriptionColumn : public IRCSignatureColumn
+{
+	//~ Begin IRCSignatureColumn
+	virtual FName GetColumnId() const override;
+	virtual bool ShouldShowColumnByDefault() const override;
+	virtual SHeaderRow::FColumn::FArguments ConstructHeaderRowColumn() override;
+	virtual TSharedRef<SWidget> ConstructRowWidget(TSharedPtr<FRCSignatureTreeItemBase> InItem, const TSharedRef<SRCSignatureTree>& InList, const TSharedRef<SRCSignatureRow>& InRow) override;
+	//~ End IRCSignatureColumn
+};
diff --git a/Engine/Plugins/VirtualProduction/RemoteControl/Source/RemoteControlUI/Private/UI/Signature/Columns/RCSignatureLabelColumn.cpp b/Engine/Plugins/VirtualProduction/RemoteControl/Source/RemoteControlUI/Private/UI/Signature/Columns/RCSignatureLabelColumn.cpp
new file mode 100644
--- /dev/null
+++ b/Engine/Plugins/VirtualProduction/RemoteControl/Source/RemoteControlUI/Private/UI/Signature/Columns/RCSignatureLabelColumn.cpp
@@ -0,0 +1,28 @@
+// Copyright Epic Games, Inc. All Rights Reserved.
+
+#include "RCSignatureLabelColumn.h"
+#include "SRCSignatureLabel.h"
+
+#define LOCTEXT_NAMESPACE "RCSignatureLabelColumn"
+
+FName FRCSignatureLabelColumn::GetColumnId() const
+{
+	return TEXT("FRCSignatureLabelColumn");
+}
+
+SHeaderRow::FColumn::FArguments FRCSignatureLabelColumn::ConstructHeaderRowColumn()
+{
+	return SHeaderRow::Column(GetColumnId())
+		.FillWidth(0.5f)
+		.DefaultLabel(LOCTEXT("DisplayName", "Label"))
+		.ShouldGenerateWidget(true);
+}
+
+TSharedRef<SWidget> FRCSignatureLabelColumn::ConstructRowWidget(TSharedPtr<FRCSignatureTreeItemBase> InItem
+	, const TSharedRef<SRCSignatureTree>& InList
+	, const TSharedRef<SRCSignatureRow>& InRow)
+{
+	return SNew(SRCSignatureLabel, InItem.ToSharedRef(), InRow);
+}
+
+#undef LOCTEXT_NAMESPACE
diff --git a/Engine/Plugins/VirtualProduction/RemoteControl/Source/RemoteControlUI/Private/UI/Signature/Columns/RCSignatureLabelColumn.h b/Engine/Plugins/VirtualProduction/RemoteControl/Source/RemoteControlUI/Private/UI/Signature/Columns/RCSignatureLabelColumn.h
new file mode 100644
--- /dev/null
+++ b/Engine/Plugins/VirtualProduction/RemoteControl/Source/RemoteControlUI/Private/UI/Signature/Columns/RCSignatureLabelColumn.h
@@ -0,0 +1,14 @@
+// Copyright Epic Games, Inc. All Rights Reserved.
+
+#pragma once
+
+#include "UI/Signature/Columns/IRCSignatureColumn.h"
+
+class FRCSignatureLabelColumn : public IRCSignatureColumn
+{
+	//~ Begin IRCSignatureColumn
+	virtual FName GetColumnId() const override;
+	virtual SHeaderRow::FColumn::FArguments ConstructHeaderRowColumn() override;
+	virtual TSharedRef<SWidget> ConstructRowWidget(TSharedPtr<FRCSignatureTreeItemBase> InItem, const TSharedRef<SRCSignatureTree>& InList, const TSharedRef<SRCSignatureRow>& InRow) override;
+	//~ End IRCSignatureColumn
+};
diff --git a/Engine/Plugins/VirtualProduction/RemoteControl/Source/RemoteControlUI/Private/UI/Signature/Columns/SRCSignatureDescription.cpp b/Engine/Plugins/VirtualProduction/RemoteControl/Source/RemoteControlUI/Private/UI/Signature/Columns/SRCSignatureDescription.cpp
new file mode 100644
--- /dev/null
+++ b/Engine/Plugins/VirtualProduction/RemoteControl/Source/RemoteControlUI/Private/UI/Signature/Columns/SRCSignatureDescription.cpp
@@ -0,0 +1,29 @@
+// Copyright Epic Games, Inc. All Rights Reserved.
+
+#include "SRCSignatureDescription.h"
+#include "UI/Signature/Items/RCSignatureTreeItemBase.h"
+#include "Widgets/Text/STextBlock.h"
+
+void SRCSignatureDescription::Construct(const FArguments& InArgs, const TSharedRef<FRCSignatureTreeItemBase>& InItem)
+{
+	ItemWeak = InItem;
+
+	ChildSlot
+	.VAlign(VAlign_Center)
+	.Padding(5.f, 0.f, 0.f, 0.f)
+	[
+		SNew(STextBlock)
+		.Text(this, &SRCSignatureDescription::GetDescription)
+		.Justification(ETextJustify::Left)
+	];
+}
+
+FText SRCSignatureDescription::GetDescription() const
+{
+	if (!CachedDescription.IsSet())
+	{
+		TSharedPtr<FRCSignatureTreeItemBase> Item = ItemWeak.Pin();
+		CachedDescription = Item.IsValid() ? Item->GetDescription() : FText::GetEmpty();
+	}
+	return *CachedDescription;
+}
diff --git a/Engine/Plugins/VirtualProduction/RemoteControl/Source/RemoteControlUI/Private/UI/Signature/Columns/SRCSignatureDescription.h b/Engine/Plugins/VirtualProduction/RemoteControl/Source/RemoteControlUI/Private/UI/Signature/Columns/SRCSignatureDescription.h
new file mode 100644
--- /dev/null
+++ b/Engine/Plugins/VirtualProduction/RemoteControl/Source/RemoteControlUI/Private/UI/Signature/Columns/SRCSignatureDescription.h
@@ -0,0 +1,23 @@
+// Copyright Epic Games, Inc. All Rights Reserved.
+
+#pragma once
+
+#include "Widgets/SCompoundWidget.h"
+
+class FRCSignatureTreeItemBase;
+
+class SRCSignatureDescription : public SCompoundWidget
+{
+public:
+	SLATE_BEGIN_ARGS(SRCSignatureDescription) {}
+	SLATE_END_ARGS()
+
+	void Construct(const FArguments& InArgs, const TSharedRef<FRCSignatureTreeItemBase>& InItem);
+
+private:
+	FText GetDescription() const;
+
+	TWeakPtr<FRCSignatureTreeItemBase> ItemWeak;
+
+	mutable TOptional<FText> CachedDescription;
+};
diff --git a/Engine/Plugins/VirtualProduction/RemoteControl/Source/RemoteControlUI/Private/UI/Signature/Columns/SRCSignatureLabel.cpp b/Engine/Plugins/VirtualProduction/RemoteControl/Source/RemoteControlUI/Private/UI/Signature/Columns/SRCSignatureLabel.cpp
new file mode 100644
--- /dev/null
+++ b/Engine/Plugins/VirtualProduction/RemoteControl/Source/RemoteControlUI/Private/UI/Signature/Columns/SRCSignatureLabel.cpp
@@ -0,0 +1,101 @@
+// Copyright Epic Games, Inc. All Rights Reserved.
+
+#include "SRCSignatureLabel.h"
+#include "Styling/SlateTypes.h"
+#include "UI/Signature/Items/RCSignatureTreeItemBase.h"
+#include "UI/Signature/SRCSignatureRow.h"
+#include "Widgets/Input/SCheckBox.h"
+#include "Widgets/Text/SInlineEditableTextBlock.h"
+#include "Widgets/Views/SExpanderArrow.h"
+
+#define LOCTEXT_NAMESPACE "SRCSignatureLabel"
+
+void SRCSignatureLabel::Construct(const FArguments& InArgs
+	, const TSharedRef<FRCSignatureTreeItemBase>& InItem
+	, const TSharedRef<SRCSignatureRow>& InRow)
+{
+	ItemWeak = InItem;
+
+	ChildSlot
+	[
+		SNew(SBox)
+		.MinDesiredHeight(25)
+		[
+			SNew(SHorizontalBox)
+			+ SHorizontalBox::Slot()
+			.AutoWidth()
+			.Padding(6.f, 0.f, 0.f, 0.f)
+			[
+				SNew(SExpanderArrow, InRow)
+				.IndentAmount(12)
+				.ShouldDrawWires(true)
+			]
+			+ SHorizontalBox::Slot()
+			.AutoWidth()
+			[
+				SNew(SCheckBox)
+				.IsChecked(this, &SRCSignatureLabel::GetItemEnabledState)
+				.OnCheckStateChanged(this, &SRCSignatureLabel::SetItemEnabledState)
+				.Padding(FMargin(4.0f, 0.0f))
+				.ToolTipText(LOCTEXT("ItemEnableCheckBoxTooltip", "Determines whether the entry is enabled or not"))
+			]
+			+ SHorizontalBox::Slot()
+			.FillWidth(1.0f)
+			[
+				SNew(SInlineEditableTextBlock)
+				.Text(this, &SRCSignatureLabel::GetSignatureDisplayName)
+				.OnTextCommitted(this, &SRCSignatureLabel::OnSignatureDisplayNameCommitted)
+				.IsSelected(InRow, &SRCSignatureRow::IsSelected)
+			]
+		]
+	];
+}
+
+ECheckBoxState SRCSignatureLabel::GetItemEnabledState() const
+{
+	if (!CachedCheckBoxState.IsSet())
+	{
+		TSharedPtr<FRCSignatureTreeItemBase> Item = ItemWeak.Pin();
+		TOptional<bool> ItemEnabled = Item.IsValid() ? Item->IsEnabled() : TOptional<bool>();
+
+		if (ItemEnabled.IsSet())
+		{
+			CachedCheckBoxState = *ItemEnabled ? ECheckBoxState::Checked : ECheckBoxState::Unchecked;
+		}
+		else
+		{
+			CachedCheckBoxState = ECheckBoxState::Undetermined;
+		}
+	}
+	return *CachedCheckBoxState;
+}
+
+void SRCSignatureLabel::SetItemEnabledState(ECheckBoxState InState)
+{
+	if (TSharedPtr<FRCSignatureTreeItemBase> Item = ItemWeak.Pin())
+	{
+		Item->SetEnabled(InState == ECheckBoxState::Checked);
+		CachedCheckBoxState.Reset();
+	}
+}
+
+FText SRCSignatureLabel::GetSignatureDisplayName() const
+{
+	if (!CachedDisplayName.IsSet())
+	{
+		TSharedPtr<FRCSignatureTreeItemBase> Item = ItemWeak.Pin();
+		CachedDisplayName = Item.IsValid() ? Item->GetDisplayNameText() : FText::GetEmpty();
+	}
+	return *CachedDisplayName;
+}
+
+void SRCSignatureLabel::OnSignatureDisplayNameCommitted(const FText& InText, ETextCommit::Type InCommitType)
+{
+	if (TSharedPtr<FRCSignatureTreeItemBase> Item = ItemWeak.Pin())
+	{
+		Item->SetDisplayNameText(InText);
+		CachedDisplayName.Reset();
+	}
+}
+
+#undef LOCTEXT_NAMESPACE
diff --git a/Engine/Plugins/VirtualProduction/RemoteControl/Source/RemoteControlUI/Private/UI/Signature/Columns/SRCSignatureLabel.h b/Engine/Plugins/VirtualProduction/RemoteControl/Source/RemoteControlUI/Private/UI/Signature/Columns/SRCSignatureLabel.h
new file mode 100644
--- /dev/null
+++ b/Engine/Plugins/VirtualProduction/RemoteControl/Source/RemoteControlUI/Private/UI/Signature/Columns/SRCSignatureLabel.h
@@ -0,0 +1,35 @@
+// Copyright Epic Games, Inc. All Rights Reserved.
+
+#pragma once
+
+#include "Widgets/SCompoundWidget.h"
+
+class FRCSignatureTreeItemBase;
+class SRCSignatureRow;
+enum class ECheckBoxState : uint8;
+
+class SRCSignatureLabel : public SCompoundWidget
+{
+public:
+	SLATE_BEGIN_ARGS(SRCSignatureLabel) {}
+	SLATE_END_ARGS()
+
+	void Construct(const FArguments& InArgs
+		, const TSharedRef<FRCSignatureTreeItemBase>& InItem
+		, const TSharedRef<SRCSignatureRow>& InRow);
+
+private:
+	ECheckBoxState GetItemEnabledState() const;
+
+	void SetItemEnabledState(ECheckBoxState InState);
+
+	FText GetSignatureDisplayName() const;
+
+	void OnSignatureDisplayNameCommitted(const FText& InText, ETextCommit::Type InCommitType);
+
+	TWeakPtr<FRCSignatureTreeItemBase> ItemWeak;
+
+	mutable TOptional<ECheckBoxState> CachedCheckBoxState;
+
+	mutable TOptional<FText> CachedDisplayName;
+};
diff --git a/Engine/Plugins/VirtualProduction/RemoteControl/Source/RemoteControlUI/Private/UI/Signature/Items/RCSignatureTreeFieldItem.cpp b/Engine/Plugins/VirtualProduction/RemoteControl/Source/RemoteControlUI/Private/UI/Signature/Items/RCSignatureTreeFieldItem.cpp
new file mode 100644
--- /dev/null
+++ b/Engine/Plugins/VirtualProduction/RemoteControl/Source/RemoteControlUI/Private/UI/Signature/Items/RCSignatureTreeFieldItem.cpp
@@ -0,0 +1,147 @@
+// Copyright Epic Games, Inc. All Rights Reserved.
+
+#include "RCSignatureTreeFieldItem.h"
+#include "RCSignatureTreeSignatureItem.h"
+#include "RemoteControlSignature.h"
+#include "RemoteControlSignatureRegistry.h"
+#include "ScopedTransaction.h"
+
+#define LOCTEXT_NAMESPACE "RCSignatureTreeFieldItem"
+
+FRCSignatureTreeFieldItem::FRCSignatureTreeFieldItem(int32 InFieldIndex, const TSharedPtr<SRCSignatureTree>& InSignatureTree)
+	: FRCSignatureTreeItemBase(InSignatureTree)
+	, FieldIndex(InFieldIndex)
+{
+}
+
+TOptional<bool> FRCSignatureTreeFieldItem::IsEnabled() const
+{
+	if (const FRCSignatureField* Field = FindField())
+	{
+		return Field->bEnabled;
+	}
+	return TOptional<bool>();
+}
+
+void FRCSignatureTreeFieldItem::SetEnabled(bool bInEnabled)
+{
+	URemoteControlSignatureRegistry* Registry;
+
+	FRCSignatureField* Field = FindFieldMutable(&Registry);
+	if (!Field || Field->bEnabled == bInEnabled)
+	{
+		return;
+	}
+
+	check(Registry);
+
+	FScopedTransaction Transaction(bInEnabled
+		? LOCTEXT("EnableField", "Enable Field")
+		: LOCTEXT("DisableField", "Disable Field"));
+
+	Registry->Modify();
+	Field->bEnabled = bInEnabled;
+}
+
+FText FRCSignatureTreeFieldItem::GetDisplayNameText() const
+{
+	if (const FRCSignatureField* Field = FindField())
+	{
+		return FText::FromName(Field->FieldPath.GetFieldName());
+	}
+	return FText::GetEmpty();
+}
+
+FText FRCSignatureTreeFieldItem::GetDescription() const
+{
+	const FRCSignatureField* Field = FindField();
+	if (!Field)
+	{
+		return FText::GetEmpty();
+	}
+
+	FFormatNamedArguments Arguments;
+	Arguments.Add(TEXT("ObjectPath"), FText::FromString(Field->ObjectRelativePath));
+	Arguments.Add(TEXT("ClassPath"), FText::FromString(Field->SupportedClass.ToString()));
+
+	return FText::Format(LOCTEXT("DescriptionFormat", "{ObjectPath} ({ClassPath})"), MoveTemp(Arguments));
+}
+
+int32 FRCSignatureTreeFieldItem::RemoveFromRegistry()
+{
+	URemoteControlSignatureRegistry* Registry;
+
+	FRCSignature* Signature = FindParentSignature(&Registry);
+	if (!Signature || !Signature->Fields.IsValidIndex(FieldIndex))
+	{
+		return 0;
+	}
+
+	check(Registry);
+
+	FScopedTransaction Transaction(LOCTEXT("RemoveField", "Remove Field"));
+	Registry->Modify();
+	Signature->Fields.RemoveAt(FieldIndex);
+	return 1;
+}
+
+FRCSignatureTreeSignatureItem* FRCSignatureTreeFieldItem::GetParentSignatureItem() const
+{
+	if (TSharedPtr<FRCSignatureTreeItemBase> Parent = GetParent())
+	{
+		return Parent->AsSignatureItem();
+	}
+	return nullptr;
+}
+
+const FRCSignatureField* FRCSignatureTreeFieldItem::FindField() const
+{
+	FRCSignatureTreeSignatureItem* ParentSignatureItem = GetParentSignatureItem();
+	if (!ParentSignatureItem)
+	{
+		return nullptr;
+	}
+
+	const FRCSignature* Signature = ParentSignatureItem->FindSignature();
+	if (!Signature || !Signature->Fields.IsValidIndex(FieldIndex))
+	{
+		return nullptr;
+	}
+
+	return &Signature->Fields[FieldIndex];
+}
+
+FRCSignatureField* FRCSignatureTreeFieldItem::FindFieldMutable(URemoteControlSignatureRegistry** OutRegistry)
+{
+	FRCSignature* Signature = FindParentSignature(OutRegistry);
+	if (!Signature || !Signature->Fields.IsValidIndex(FieldIndex))
+	{
+		return nullptr;
+	}
+	return &Signature->Fields[FieldIndex];
+}
+
+FRCSignature* FRCSignatureTreeFieldItem::FindParentSignature(URemoteControlSignatureRegistry** OutRegistry)
+{
+	FRCSignatureTreeSignatureItem* ParentSignatureItem = GetParentSignatureItem();
+	if (!ParentSignatureItem)
+	{
+		return nullptr;
+	}
+
+	URemoteControlSignatureRegistry* Registry = ParentSignatureItem->GetRegistry();
+
+	FRCSignature* Signature = ParentSignatureItem->FindSignatureMutable(Registry);
+	if (!Signature)
+	{
+		return nullptr;
+	}
+
+	if (OutRegistry)
+	{
+		*OutRegistry = Registry;	
+	}
+	return Signature;
+}
+
+#undef LOCTEXT_NAMESPACE
diff --git a/Engine/Plugins/VirtualProduction/RemoteControl/Source/RemoteControlUI/Private/UI/Signature/Items/RCSignatureTreeFieldItem.h b/Engine/Plugins/VirtualProduction/RemoteControl/Source/RemoteControlUI/Private/UI/Signature/Items/RCSignatureTreeFieldItem.h
new file mode 100644
--- /dev/null
+++ b/Engine/Plugins/VirtualProduction/RemoteControl/Source/RemoteControlUI/Private/UI/Signature/Items/RCSignatureTreeFieldItem.h
@@ -0,0 +1,35 @@
+// Copyright Epic Games, Inc. All Rights Reserved.
+
+#pragma once
+
+#include "RCSignatureTreeItemBase.h"
+
+class URemoteControlSignatureRegistry;
+struct FRCSignature;
+struct FRCSignatureField;
+
+/** Item class representing a Field owned by a Signature */
+class FRCSignatureTreeFieldItem : public FRCSignatureTreeItemBase
+{
+public:
+	explicit FRCSignatureTreeFieldItem(int32 InFieldIndex, const TSharedPtr<SRCSignatureTree>& InSignatureTree);
+
+	//~ Begin FRCSignatureTreeItemBase
+	virtual TOptional<bool> IsEnabled() const override;
+	virtual void SetEnabled(bool bInEnabled) override;
+	virtual FText GetDisplayNameText() const override;
+	virtual FText GetDescription() const override;
+	virtual int32 RemoveFromRegistry() override;
+	//~ End FRCSignatureTreeItemBase
+
+private:
+	FRCSignatureTreeSignatureItem* GetParentSignatureItem() const;
+
+	const FRCSignatureField* FindField() const;
+
+	FRCSignatureField* FindFieldMutable(URemoteControlSignatureRegistry** OutRegistry);
+
+	FRCSignature* FindParentSignature(URemoteControlSignatureRegistry** OutRegistry);
+
+	int32 FieldIndex;
+};
diff --git a/Engine/Plugins/VirtualProduction/RemoteControl/Source/RemoteControlUI/Private/UI/Signature/Items/RCSignatureTreeItemBase.cpp b/Engine/Plugins/VirtualProduction/RemoteControl/Source/RemoteControlUI/Private/UI/Signature/Items/RCSignatureTreeItemBase.cpp
new file mode 100644
--- /dev/null
+++ b/Engine/Plugins/VirtualProduction/RemoteControl/Source/RemoteControlUI/Private/UI/Signature/Items/RCSignatureTreeItemBase.cpp
@@ -0,0 +1,23 @@
+// Copyright Epic Games, Inc. All Rights Reserved.
+
+#include "RCSignatureTreeItemBase.h"
+#include "UI/Signature/SRCSignatureTree.h"
+
+FRCSignatureTreeItemBase::FRCSignatureTreeItemBase(const TSharedPtr<SRCSignatureTree>& InSignatureTree)
+	: FRCLogicModeBase(InSignatureTree ? InSignatureTree->GetRemoteControlPanel() : nullptr)
+	, SignatureTreeWeak(InSignatureTree)
+{
+}
+
+void FRCSignatureTreeItemBase::RebuildChildren()
+{
+	Children.Reset();
+	GenerateChildren(Children);
+
+	TSharedRef<FRCSignatureTreeItemBase> This = SharedThis(this);
+	for (const TSharedPtr<FRCSignatureTreeItemBase>& Child : Children)
+	{
+		Child->ParentWeak = This;
+		Child->RebuildChildren();
+	}
+}
diff --git a/Engine/Plugins/VirtualProduction/RemoteControl/Source/RemoteControlUI/Private/UI/Signature/Items/RCSignatureTreeItemBase.h b/Engine/Plugins/VirtualProduction/RemoteControl/Source/RemoteControlUI/Private/UI/Signature/Items/RCSignatureTreeItemBase.h
new file mode 100644
--- /dev/null
+++ b/Engine/Plugins/VirtualProduction/RemoteControl/Source/RemoteControlUI/Private/UI/Signature/Items/RCSignatureTreeItemBase.h
@@ -0,0 +1,72 @@
+// Copyright Epic Games, Inc. All Rights Reserved.
+
+#pragma once
+
+#include "UI/BaseLogicUI/RCLogicModeBase.h"
+
+class SRCSignatureTree;
+
+/** Base class for any Item represented in the Signature Tree */
+class FRCSignatureTreeItemBase : public FRCLogicModeBase
+{
+public:
+	explicit FRCSignatureTreeItemBase(const TSharedPtr<SRCSignatureTree>& InSignatureTree);
+
+	virtual TOptional<bool> IsEnabled() const
+	{
+		return TOptional<bool>();
+	}
+
+	virtual void SetEnabled(bool bInEnabled)
+	{
+	}
+
+	virtual FText GetDisplayNameText() const = 0;
+
+	virtual bool CanEditDisplayNameText() const
+	{
+		return false;
+	}
+
+	virtual void SetDisplayNameText(const FText& InText)
+	{
+	}
+
+	virtual FText GetDescription() const = 0;
+
+	virtual int32 RemoveFromRegistry() = 0;
+
+	virtual class FRCSignatureTreeSignatureItem* AsSignatureItem()
+	{
+		return nullptr;
+	}
+
+	TConstArrayView<TSharedPtr<FRCSignatureTreeItemBase>> GetChildren() const
+	{
+		return Children;
+	}
+
+	TSharedPtr<FRCSignatureTreeItemBase> GetParent() const
+	{
+		return ParentWeak.Pin();
+	}
+
+	TSharedPtr<SRCSignatureTree> GetSignatureTree() const
+	{
+		return SignatureTreeWeak.Pin();
+	}
+
+	void RebuildChildren();
+
+protected:
+	virtual void GenerateChildren(TArray<TSharedPtr<FRCSignatureTreeItemBase>>& OutChildren) const
+	{
+	}
+
+private:
+	TArray<TSharedPtr<FRCSignatureTreeItemBase>> Children;
+
+	TWeakPtr<FRCSignatureTreeItemBase> ParentWeak;
+
+	TWeakPtr<SRCSignatureTree> SignatureTreeWeak;
+};
diff --git a/Engine/Plugins/VirtualProduction/RemoteControl/Source/RemoteControlUI/Private/UI/Signature/Items/RCSignatureTreeSignatureItem.cpp b/Engine/Plugins/VirtualProduction/RemoteControl/Source/RemoteControlUI/Private/UI/Signature/Items/RCSignatureTreeSignatureItem.cpp
new file mode 100644
--- /dev/null
+++ b/Engine/Plugins/VirtualProduction/RemoteControl/Source/RemoteControlUI/Private/UI/Signature/Items/RCSignatureTreeSignatureItem.cpp
@@ -0,0 +1,214 @@
+// Copyright Epic Games, Inc. All Rights Reserved.
+
+#include "RCSignatureTreeSignatureItem.h"
+#include "RCSignatureTreeFieldItem.h"
+#include "RemoteControlPreset.h"
+#include "RemoteControlSignatureRegistry.h"
+#include "ScopedTransaction.h"
+#include "UI/Signature/SRCSignatureTree.h"
+
+#define LOCTEXT_NAMESPACE "RCSignatureTreeSignatureItem"
+
+FRCSignatureTreeSignatureItem::FRCSignatureTreeSignatureItem(const FRCSignature& InSignature, const TSharedPtr<SRCSignatureTree>& InSignatureTree)
+	: FRCSignatureTreeItemBase(InSignatureTree)
+	, SignatureId(InSignature.Id)
+{
+	if (URemoteControlPreset* Preset = GetPreset())
+	{
+		RegistryWeak = Preset->GetSignatureRegistry();
+	}
+}
+
+const FGuid& FRCSignatureTreeSignatureItem::GetSignatureId() const
+{
+	return SignatureId;
+}
+
+URemoteControlSignatureRegistry* FRCSignatureTreeSignatureItem::GetRegistry() const
+{
+	return RegistryWeak.Get();
+}
+
+const FRCSignature* FRCSignatureTreeSignatureItem::FindSignature() const
+{
+	if (URemoteControlSignatureRegistry* SignatureRegistry = GetRegistry())
+	{
+		return SignatureRegistry->FindSignature(SignatureId);
+	}
+	return nullptr;
+}
+
+FRCSignature* FRCSignatureTreeSignatureItem::FindSignatureMutable(URemoteControlSignatureRegistry* InRegistry)
+{
+	if (InRegistry)
+	{
+		return InRegistry->FindSignatureMutable(SignatureId);
+	}
+	return nullptr;
+}
+
+void FRCSignatureTreeSignatureItem::AddFieldEntities(TConstArrayView<FGuid> InFieldEntityIds)
+{
+	if (InFieldEntityIds.IsEmpty())
+	{
+		return;
+	}
+
+	URemoteControlPreset* Preset = GetPreset();
+	if (!Preset)
+	{
+		return;
+	}
+
+	URemoteControlSignatureRegistry* Registry = GetRegistry();
+
+	FRCSignature* Signature = FindSignatureMutable(Registry);
+	if (!Signature)
+	{
+		return;
+	}
+
+	FScopedTransaction Transaction(LOCTEXT("AddProperties", "Add Properties"));
+	Registry->Modify();
+
+	const int32 AddedFieldCount = Signature->AddFieldsFromEntities(Preset, InFieldEntityIds);
+
+	if (AddedFieldCount == 0)
+	{
+		// Nothing happened, cancel transaction
+		Transaction.Cancel();
+	}
+	else if (TSharedPtr<SRCSignatureTree> SignatureTree = GetSignatureTree())
+	{
+		SignatureTree->Refresh();
+	}
+}
+
+void FRCSignatureTreeSignatureItem::ApplySignature(TConstArrayView<TWeakObjectPtr<AActor>> InActors)
+{
+	if (InActors.IsEmpty())
+	{
+		return;
+	}
+
+	URemoteControlPreset* Preset = GetPreset();
+	if (!Preset)
+	{
+		return;
+	}
+
+	const FRCSignature* Signature = FindSignature();
+	if (!Signature || !Signature->bEnabled)
+	{
+		return;
+	}
+
+	FScopedTransaction Transaction(LOCTEXT("ApplySignatureActors", "Apply Signature to Actors"));
+
+	int32 ExposedCount = Signature->ApplySignature(Preset, InActors);
+
+	if (ExposedCount == 0)
+	{
+		// Nothing was exposed, cancel transaction
+		Transaction.Cancel();
+	}
+}
+
+TOptional<bool> FRCSignatureTreeSignatureItem::IsEnabled() const
+{
+	if (const FRCSignature* Signature = FindSignature())
+	{
+		return Signature->bEnabled;
+	}
+	return TOptional<bool>();
+}
+
+void FRCSignatureTreeSignatureItem::SetEnabled(bool bInEnabled)
+{
+	URemoteControlSignatureRegistry* SignatureRegistry = GetRegistry();
+
+	FRCSignature* Signature = FindSignatureMutable(SignatureRegistry);
+	if (!Signature || Signature->bEnabled == bInEnabled)
+	{
+		return;
+	}
+
+	FScopedTransaction Transaction(bInEnabled
+		? LOCTEXT("EnableSignature", "Enable Signature")
+		: LOCTEXT("DisableSignature", "Disable Signature"));
+
+	SignatureRegistry->Modify();
+	Signature->bEnabled = bInEnabled;
+}
+
+FText FRCSignatureTreeSignatureItem::GetDisplayNameText() const
+{
+	if (const FRCSignature* Signature = FindSignature())
+	{
+		return Signature->DisplayName;
+	}
+	return FText::GetEmpty();
+}
+
+bool FRCSignatureTreeSignatureItem::CanEditDisplayNameText() const
+{
+	return true;
+}
+
+void FRCSignatureTreeSignatureItem::SetDisplayNameText(const FText& InText)
+{
+	URemoteControlSignatureRegistry* SignatureRegistry = GetRegistry();
+
+	FRCSignature* Signature = FindSignatureMutable(SignatureRegistry);
+	if (!Signature || Signature->DisplayName.EqualTo(InText))
+	{
+		return;
+	}
+
+	FScopedTransaction Transaction(LOCTEXT("SetSignatureName", "Set Signature Name"));
+	SignatureRegistry->Modify();
+	Signature->DisplayName = InText;
+}
+
+FText FRCSignatureTreeSignatureItem::GetDescription() const
+{
+	return FText::GetEmpty();
+}
+
+int32 FRCSignatureTreeSignatureItem::RemoveFromRegistry()
+{
+	URemoteControlSignatureRegistry* SignatureRegistry = GetRegistry();
+	if (!SignatureRegistry)
+	{
+		return 0;
+	}
+
+	// Remove Signatures from Preset
+	FScopedTransaction Transaction(LOCTEXT("RemoveSignature", "Remove Signature"));
+	SignatureRegistry->Modify();
+	return SignatureRegistry->RemoveSignature(SignatureId);
+}
+
+FRCSignatureTreeSignatureItem* FRCSignatureTreeSignatureItem::AsSignatureItem()
+{
+	return this;
+}
+
+void FRCSignatureTreeSignatureItem::GenerateChildren(TArray<TSharedPtr<FRCSignatureTreeItemBase>>& OutChildren) const
+{
+	const FRCSignature* Signature = FindSignature();
+	if (!Signature)
+	{
+		return;
+	}
+
+	const TSharedPtr<SRCSignatureTree> SignatureTree = GetSignatureTree();
+
+	OutChildren.Reserve(OutChildren.Num() + Signature->Fields.Num());
+	for (int32 FieldIndex = 0; FieldIndex < Signature->Fields.Num(); ++FieldIndex)
+	{
+		OutChildren.Add(MakeShared<FRCSignatureTreeFieldItem>(FieldIndex, SignatureTree));
+	}
+}
+
+#undef LOCTEXT_NAMESPACE
diff --git a/Engine/Plugins/VirtualProduction/RemoteControl/Source/RemoteControlUI/Private/UI/Signature/Items/RCSignatureTreeSignatureItem.h b/Engine/Plugins/VirtualProduction/RemoteControl/Source/RemoteControlUI/Private/UI/Signature/Items/RCSignatureTreeSignatureItem.h
new file mode 100644
--- /dev/null
+++ b/Engine/Plugins/VirtualProduction/RemoteControl/Source/RemoteControlUI/Private/UI/Signature/Items/RCSignatureTreeSignatureItem.h
@@ -0,0 +1,50 @@
+// Copyright Epic Games, Inc. All Rights Reserved.
+
+#pragma once
+
+#include "Misc/Guid.h"
+#include "RCSignatureTreeItemBase.h"
+#include "UI/Signature/IRCSignatureItem.h"
+#include "UObject/WeakObjectPtr.h"
+#include "UObject/WeakObjectPtrTemplates.h"
+
+class URemoteControlSignatureRegistry;
+struct FRCSignature;
+
+/** Item class representing an RC Signature */
+class FRCSignatureTreeSignatureItem : public FRCSignatureTreeItemBase, public IRCSignatureItem
+{
+public:
+	explicit FRCSignatureTreeSignatureItem(const FRCSignature& InSignature, const TSharedPtr<SRCSignatureTree>& InSignatureTree);
+
+	const FGuid& GetSignatureId() const;
+
+	URemoteControlSignatureRegistry* GetRegistry() const;
+
+	const FRCSignature* FindSignature() const;
+
+	FRCSignature* FindSignatureMutable(URemoteControlSignatureRegistry* InRegistry);
+
+	//~ Begin IRCSignatureItem
+	virtual void AddFieldEntities(TConstArrayView<FGuid> InFieldEntityIds) override;
+	virtual void ApplySignature(TConstArrayView<TWeakObjectPtr<AActor>> InActors) override;
+	//~ End IRCSignatureItem
+
+protected:
+	//~ Begin FRCSignatureTreeItem
+	virtual TOptional<bool> IsEnabled() const override;
+	virtual void SetEnabled(bool bInEnabled) override;
+	virtual FText GetDisplayNameText() const override;
+	virtual bool CanEditDisplayNameText() const override;
+	virtual void SetDisplayNameText(const FText& InText) override;
+	virtual FText GetDescription() const override;
+	virtual int32 RemoveFromRegistry() override;
+	virtual FRCSignatureTreeSignatureItem* AsSignatureItem() override;
+	virtual void GenerateChildren(TArray<TSharedPtr<FRCSignatureTreeItemBase>>& OutChildren) const override;
+	//~ End FRCSignatureTreeItem
+
+private:
+	TWeakObjectPtr<URemoteControlSignatureRegistry> RegistryWeak;
+
+	FGuid SignatureId;
+};
diff --git a/Engine/Plugins/VirtualProduction/RemoteControl/Source/RemoteControlUI/Private/UI/Signature/SRCSignaturePanel.cpp b/Engine/Plugins/VirtualProduction/RemoteControl/Source/RemoteControlUI/Private/UI/Signature/SRCSignaturePanel.cpp
new file mode 100644
--- /dev/null
+++ b/Engine/Plugins/VirtualProduction/RemoteControl/Source/RemoteControlUI/Private/UI/Signature/SRCSignaturePanel.cpp
@@ -0,0 +1,196 @@
+// Copyright Epic Games, Inc. All Rights Reserved.
+
+#include "SRCSignaturePanel.h"
+#include "Columns/RCSignatureDescriptionColumn.h"
+#include "Columns/RCSignatureLabelColumn.h"
+#include "Misc/MessageDialog.h"
+#include "RemoteControlPreset.h"
+#include "RemoteControlSignatureRegistry.h"
+#include "SRCSignatureTree.h"
+#include "ScopedTransaction.h"
+#include "Styling/RemoteControlStyles.h"
+#include "UI/Panels/SRCDockPanel.h"
+#include "UI/RemoteControlPanelStyle.h"
+#include "UI/SRemoteControlPanel.h"
+
+#define LOCTEXT_NAMESPACE "SRCSignaturePanel"
+
+void SRCSignaturePanel::Construct(const FArguments& InArgs, const TSharedRef<SRemoteControlPanel>& InPanel)
+{
+	SRCLogicPanelBase::Construct(SRCLogicPanelBase::FArguments(), InPanel);
+
+	const FRCPanelStyle* RCPanelStyle = &FRemoteControlPanelStyle::Get()->GetWidgetStyle<FRCPanelStyle>("RemoteControlPanel.MinorPanel");
+
+	// Signature Dock Panel
+	TSharedRef<SRCMinorPanel> SignatureDockPanel = SNew(SRCMinorPanel)
+		.HeaderLabel(LOCTEXT("SignaturesLabel", "Signatures"))
+		.EnableFooter(false)
+		[
+			SAssignNew(SignatureTreeView, SRCSignatureTree, SharedThis(this), InPanel)
+			.Columns(
+				{
+					MakeShared<FRCSignatureLabelColumn>(),
+					MakeShared<FRCSignatureDescriptionColumn>(),
+				})
+		];
+
+	// Add New Signature Button
+	TSharedRef<SButton> AddSignatureButton = SNew(SButton)
+		.AddMetaData<FTagMetaData>(FTagMetaData(TEXT("Add Signature")))
+		.IsEnabled_Lambda([LiveMode = InArgs._LiveMode]() { return !LiveMode.Get(); })
+		.OnClicked(this, &SRCSignaturePanel::OnAddButtonClicked)
+		.ForegroundColor(FSlateColor::UseForeground())
+		.ButtonStyle(&RCPanelStyle->FlatButtonStyle)
+		.ContentPadding(FMargin(4.f, 2.f))
+		.HAlign(HAlign_Center)
+		.VAlign(VAlign_Center)
+		.Content()
+		[
+			SNew(SBox)
+			.WidthOverride(RCPanelStyle->IconSize.X)
+			.HeightOverride(RCPanelStyle->IconSize.Y)
+			[
+				SNew(SImage)
+				.ColorAndOpacity(FSlateColor::UseForeground())
+				.Image(FAppStyle::GetBrush("Icons.PlusCircle"))
+			]
+		];
+
+	SignatureDockPanel->AddHeaderToolbarItem(EToolbar::Left, AddSignatureButton);
+
+	ChildSlot
+	.Padding(RCPanelStyle->PanelPadding)
+	[
+		SignatureDockPanel
+	];
+}
+
+URemoteControlSignatureRegistry* SRCSignaturePanel::GetSignatureRegistry() const
+{
+	if (URemoteControlPreset* Preset = GetPreset())
+	{
+		return Preset->GetSignatureRegistry();
+	}
+	return nullptr;
+}
+
+bool SRCSignaturePanel::IsListFocused() const
+{
+	return SignatureTreeView.IsValid() && SignatureTreeView->IsListFocused();
+}
+
+TArray<TSharedPtr<FRCLogicModeBase>> SRCSignaturePanel::GetSelectedLogicItems() const
+{
+	if (SignatureTreeView.IsValid())
+	{
+		return SignatureTreeView->GetSelectedLogicItems();
+	}
+	return {};
+}
+
+FReply SRCSignaturePanel::RequestDeleteSelectedItem()
+{
+	if (!SignatureTreeView.IsValid())
+	{
+		return FReply::Unhandled();
+	}
+
+	const EAppReturnType::Type UserResponse = FMessageDialog::Open(EAppMsgType::YesNo
+		, LOCTEXT("DeleteSelectedWarning", "Delete the Selected Signatures?"));
+
+	if (UserResponse == EAppReturnType::Yes)
+	{
+		DeleteSelectedPanelItems();
+	}
+
+	return FReply::Handled();
+}
+
+FReply SRCSignaturePanel::RequestDeleteAllItems()
+{
+	if (!SignatureTreeView.IsValid())
+	{
+		return FReply::Unhandled();
+	}
+
+	const EAppReturnType::Type UserResponse = FMessageDialog::Open(EAppMsgType::YesNo
+		, FText::Format(LOCTEXT("DeleteAllWarning", "You are about to delete {0} signatures. Are you sure you want to proceed?"), SignatureTreeView->Num()));
+
+	if (UserResponse == EAppReturnType::Yes)
+	{
+		return DeleteAllItems();
+	}
+
+	return FReply::Handled();
+}
+
+bool SRCSignaturePanel::CanCopyItems() const
+{
+	// todo: Copy/Paste Signatures
+	return false;
+}
+
+bool SRCSignaturePanel::CanDuplicateItems() const
+{
+	// todo: Duplicate Signatures
+	return false;
+}
+
+void SRCSignaturePanel::DeleteSelectedPanelItems()
+{
+	if (SignatureTreeView.IsValid())
+	{
+		SignatureTreeView->DeleteSelectedPanelItems();
+	}
+}
+
+FReply SRCSignaturePanel::OnAddButtonClicked()
+{
+	URemoteControlSignatureRegistry* SignatureRegistry = GetSignatureRegistry();
+	if (!SignatureRegistry)
+	{
+		return FReply::Unhandled();
+	}
+
+	// Add new signature
+	{
+		FScopedTransaction Transaction(LOCTEXT("NewSignature", "New Signature"));
+		SignatureRegistry->Modify();
+		SignatureRegistry->AddSignature();
+	}
+
+	if (SignatureTreeView.IsValid())
+	{
+		SignatureTreeView->Refresh();
+	}
+
+	return FReply::Handled();
+}
+
+FReply SRCSignaturePanel::DeleteAllItems()
+{
+	URemoteControlSignatureRegistry* SignatureRegistry = GetSignatureRegistry();
+	if (!SignatureRegistry)
+	{
+		return FReply::Unhandled();
+	}
+
+	// Don't generate a transaction + modify if the signature container is already empty
+	if (SignatureRegistry->GetSignatures().IsEmpty())
+	{
+		return FReply::Handled();
+	}
+
+	FScopedTransaction Transaction(LOCTEXT("EmptySignatures", "Empty Signatures"));
+	SignatureRegistry->Modify();
+	SignatureRegistry->EmptySignatures();
+
+	if (SignatureTreeView.IsValid())
+	{
+		SignatureTreeView->Refresh();
+	}
+
+	return FReply::Handled();
+}
+
+#undef LOCTEXT_NAMESPACE 
diff --git a/Engine/Plugins/VirtualProduction/RemoteControl/Source/RemoteControlUI/Private/UI/Signature/SRCSignaturePanel.h b/Engine/Plugins/VirtualProduction/RemoteControl/Source/RemoteControlUI/Private/UI/Signature/SRCSignaturePanel.h
new file mode 100644
--- /dev/null
+++ b/Engine/Plugins/VirtualProduction/RemoteControl/Source/RemoteControlUI/Private/UI/Signature/SRCSignaturePanel.h
@@ -0,0 +1,41 @@
+// Copyright Epic Games, Inc. All Rights Reserved.
+
+#pragma once
+
+#include "UI/BaseLogicUI/SRCLogicPanelBase.h"
+
+class IRCSignatureColumn;
+class URemoteControlSignatureRegistry;
+class SRCSignatureTree;
+
+class SRCSignaturePanel : public SRCLogicPanelBase
+{
+public:
+	SLATE_BEGIN_ARGS(SRCSignaturePanel) {}
+		SLATE_ATTRIBUTE(bool, LiveMode)
+	SLATE_END_ARGS()
+
+	void Construct(const FArguments& InArgs, const TSharedRef<SRemoteControlPanel>& InPanel);
+
+	/** Gets the Signature Registry from the RC Preset */
+	URemoteControlSignatureRegistry* GetSignatureRegistry() const;
+
+	/** Whether this widget currently has focus */
+	bool IsListFocused() const;
+
+	//~ Begin SRCLogicPanelBase
+	virtual TArray<TSharedPtr<FRCLogicModeBase>> GetSelectedLogicItems() const override;
+	virtual FReply RequestDeleteSelectedItem() override;
+	virtual FReply RequestDeleteAllItems() override;
+	virtual bool CanCopyItems() const override;
+	virtual bool CanDuplicateItems() const override;
+	virtual void DeleteSelectedPanelItems() override;
+	//~ End SRCLogicPanelBase
+
+private:
+	FReply OnAddButtonClicked();
+
+	FReply DeleteAllItems();
+
+	TSharedPtr<SRCSignatureTree> SignatureTreeView;
+};
diff --git a/Engine/Plugins/VirtualProduction/RemoteControl/Source/RemoteControlUI/Private/UI/Signature/SRCSignatureRow.cpp b/Engine/Plugins/VirtualProduction/RemoteControl/Source/RemoteControlUI/Private/UI/Signature/SRCSignatureRow.cpp
new file mode 100644
--- /dev/null
+++ b/Engine/Plugins/VirtualProduction/RemoteControl/Source/RemoteControlUI/Private/UI/Signature/SRCSignatureRow.cpp
@@ -0,0 +1,181 @@
+// Copyright Epic Games, Inc. All Rights Reserved.
+
+#include "SRCSignatureRow.h"
+#include "Columns/IRCSignatureColumn.h"
+#include "DragAndDrop/ActorDragDropOp.h"
+#include "DragAndDrop/CompositeDragDropOp.h"
+#include "DragAndDrop/FolderDragDropOp.h"
+#include "EditorActorFolders.h"
+#include "Items/RCSignatureTreeItemBase.h"
+#include "Items/RCSignatureTreeSignatureItem.h"
+#include "RemoteControlUIModule.h"
+#include "SRCSignatureTree.h"
+#include "UI/SRCPanelExposedEntity.h"
+#include "UI/Signature/IRCSignatureCustomization.h"
+#include "Widgets/SNullWidget.h"
+
+void SRCSignatureRow::Construct(const FArguments& InArgs
+	, const TSharedPtr<FRCSignatureTreeItemBase>& InItem
+	, const TSharedRef<SRCSignatureTree>& InSignatureTree
+	, const TSharedRef<STableViewBase>& InTableView)
+{
+	ItemWeak = InItem;
+	SignatureTreeWeak = InSignatureTree;
+
+	SMultiColumnTableRow::Construct(FSuperRowType::FArguments()
+		.ShowWires(true)
+		.OnCanAcceptDrop(this, &SRCSignatureRow::OnRowCanAcceptDrop)
+		.OnAcceptDrop(this, &SRCSignatureRow::OnRowAcceptDrop)
+		, InTableView);
+}
+
+TSharedRef<SWidget> SRCSignatureRow::GenerateWidgetForColumn(const FName& InColumnName)
+{
+	TSharedPtr<SRCSignatureTree> SignatureTree = SignatureTreeWeak.Pin();
+	TSharedPtr<FRCSignatureTreeItemBase> Item = ItemWeak.Pin();
+
+	if (!SignatureTree.IsValid() || !Item.IsValid())
+	{
+		return SNullWidget::NullWidget;
+	}
+
+	const TSharedPtr<IRCSignatureColumn> Column = SignatureTree->FindColumn(InColumnName);
+	if (ensureMsgf(Column.IsValid(), TEXT("Column %s was unexpectedly not found."), *InColumnName.ToString()))
+	{
+		return Column->ConstructRowWidget(Item.ToSharedRef(), SignatureTree.ToSharedRef(), SharedThis(this));
+	}
+
+	return SNullWidget::NullWidget;
+}
+
+TOptional<EItemDropZone> SRCSignatureRow::OnRowCanAcceptDrop(const FDragDropEvent& InDragDropEvent
+	, EItemDropZone InDropZone
+	, TSharedPtr<FRCSignatureTreeItemBase> InItem)
+{
+	if (!InItem.IsValid())
+	{
+		return TOptional<EItemDropZone>();
+	}
+
+	if (FRCSignatureTreeSignatureItem* SignatureItem = InItem->AsSignatureItem())
+	{
+		// External Customizations to handle drag drops first
+		for (TSharedRef<IRCSignatureCustomization> Customization : FRemoteControlUIModule::Get().GetSignatureCustomizations())
+		{
+			if (Customization->CanAcceptDrop(InDragDropEvent, SignatureItem))
+			{
+				return EItemDropZone::OntoItem;;
+			}
+		}
+
+		// Entity Drag Drop
+		if (InDragDropEvent.GetOperationAs<FExposedEntityDragDrop>())
+		{
+			return EItemDropZone::OntoItem;
+		}
+
+		// Actor Drag Drop (only Actors)
+		if (InDragDropEvent.GetOperationAs<FActorDragDropOp>())
+		{
+			return EItemDropZone::OntoItem;
+		}
+
+		// Folder Drag Drop (only folders)
+		if (InDragDropEvent.GetOperationAs<FFolderDragDropOp>())
+		{
+			return EItemDropZone::OntoItem;
+		}
+
+		// Composite Drag and Drop (mix of Actors and Folders)
+		if (TSharedPtr<FCompositeDragDropOp> CompositeDragDropOp = InDragDropEvent.GetOperationAs<FCompositeDragDropOp>())
+		{
+			if (CompositeDragDropOp->GetSubOp<FActorDragDropOp>())
+			{
+				return EItemDropZone::OntoItem;
+			}
+
+			if (CompositeDragDropOp->GetSubOp<FFolderDragDropOp>())
+			{
+				return EItemDropZone::OntoItem;
+			}
+		}
+	}
+
+	return TOptional<EItemDropZone>();
+}
+
+FReply SRCSignatureRow::OnRowAcceptDrop(const FDragDropEvent& InDragDropEvent
+	, EItemDropZone InDropZone
+	, TSharedPtr<FRCSignatureTreeItemBase> InItem)
+{
+	if (!InItem.IsValid())
+	{
+		return FReply::Unhandled();
+	}
+
+	if (FRCSignatureTreeSignatureItem* SignatureItem = InItem->AsSignatureItem())
+	{
+		// External Customizations to handle drops first
+		for (TSharedRef<IRCSignatureCustomization> Customization : FRemoteControlUIModule::Get().GetSignatureCustomizations())
+		{
+			if (Customization->AcceptDrop(InDragDropEvent, SignatureItem).IsEventHandled())
+			{
+				return FReply::Handled();
+			}
+		}
+
+		// Entity Drag Drop
+		if (TSharedPtr<FExposedEntityDragDrop> ExposedEntityDragDropOp = InDragDropEvent.GetOperationAs<FExposedEntityDragDrop>())
+		{
+			SignatureItem->AddFieldEntities(ExposedEntityDragDropOp->GetSelectedFieldsId());
+			return FReply::Handled();
+		}
+
+		// Actor Drag Drop (Only Actors dragged)
+		if (TSharedPtr<FActorDragDropOp> ActorDragDropOp = InDragDropEvent.GetOperationAs<FActorDragDropOp>())
+		{
+			HandleActorDragDrop(SignatureItem, *ActorDragDropOp);
+			return FReply::Handled();
+		}
+
+		// Folder Drag Drop (Only Folders dragged)
+		if (TSharedPtr<FFolderDragDropOp> FolderDragDropOp = InDragDropEvent.GetOperationAs<FFolderDragDropOp>())
+		{
+			HandleFolderDragDrop(SignatureItem, *FolderDragDropOp);
+			return FReply::Handled();
+		}
+
+		// Composite Drag and Drop (mix of Actors and Folders)
+		if (TSharedPtr<FCompositeDragDropOp> CompositeDragDropOp = InDragDropEvent.GetOperationAs<FCompositeDragDropOp>())
+		{
+			if (TSharedPtr<FActorDragDropOp> ActorDragDropOp = CompositeDragDropOp->GetSubOp<FActorDragDropOp>())
+			{
+				HandleActorDragDrop(SignatureItem, *ActorDragDropOp);
+				return FReply::Handled();
+			}
+
+			if (TSharedPtr<FFolderDragDropOp> FolderDragDropOp = CompositeDragDropOp->GetSubOp<FFolderDragDropOp>())
+			{
+				HandleFolderDragDrop(SignatureItem, *FolderDragDropOp);
+				return FReply::Handled();
+			}
+		}
+	}
+
+	return FReply::Unhandled();
+}
+
+void SRCSignatureRow::HandleActorDragDrop(IRCSignatureItem* InSignatureItem, const FActorDragDropOp& InActorDragDropOp)
+{
+	InSignatureItem->ApplySignature(InActorDragDropOp.Actors);
+}
+
+void SRCSignatureRow::HandleFolderDragDrop(IRCSignatureItem* InSignatureItem, const FFolderDragDropOp& InFolderDragDropOp)
+{
+	if (UWorld* World = InFolderDragDropOp.World.Get())
+	{
+		TArray<TWeakObjectPtr<AActor>> Actors;
+		FActorFolders::GetWeakActorsFromFolders(*World, InFolderDragDropOp.Folders, Actors);
+		InSignatureItem->ApplySignature(Actors);
+	}
+}
diff --git a/Engine/Plugins/VirtualProduction/RemoteControl/Source/RemoteControlUI/Private/UI/Signature/SRCSignatureRow.h b/Engine/Plugins/VirtualProduction/RemoteControl/Source/RemoteControlUI/Private/UI/Signature/SRCSignatureRow.h
new file mode 100644
--- /dev/null
+++ b/Engine/Plugins/VirtualProduction/RemoteControl/Source/RemoteControlUI/Private/UI/Signature/SRCSignatureRow.h
@@ -0,0 +1,41 @@
+// Copyright Epic Games, Inc. All Rights Reserved.
+
+#pragma once
+
+#include "Widgets/Views/SListView.h"
+#include "Widgets/Views/STableRow.h"
+
+class FActorDragDropOp;
+class FFolderDragDropOp;
+class FRCSignatureTreeItemBase;
+class IRCSignatureItem;
+class SRCSignatureTree;
+
+class SRCSignatureRow : public SMultiColumnTableRow<TSharedPtr<FRCSignatureTreeItemBase>>
+{
+public:
+	SLATE_BEGIN_ARGS(SRCSignatureRow) {}
+	SLATE_END_ARGS()
+
+	void Construct(const FArguments& InArgs
+		, const TSharedPtr<FRCSignatureTreeItemBase>& InItem
+		, const TSharedRef<SRCSignatureTree>& InSignatureTree
+		, const TSharedRef<STableViewBase>& InTableView);
+
+	//~ Begin SMultiColumnTableRow
+	virtual TSharedRef<SWidget> GenerateWidgetForColumn(const FName& InColumnName) override;
+	//~ End SMultiColumnTableRow
+
+private:
+	TOptional<EItemDropZone> OnRowCanAcceptDrop(const FDragDropEvent& InDragDropEvent, EItemDropZone InDropZone, TSharedPtr<FRCSignatureTreeItemBase> InItem);
+
+	FReply OnRowAcceptDrop(const FDragDropEvent& InDragDropEvent, EItemDropZone InDropZone, TSharedPtr<FRCSignatureTreeItemBase> InItem);
+
+	void HandleActorDragDrop(IRCSignatureItem* InSignatureItem, const FActorDragDropOp& InActorDragDropOp);
+
+	void HandleFolderDragDrop(IRCSignatureItem* InSignatureItem, const FFolderDragDropOp& InFolderDragDropOp);
+
+	TWeakPtr<FRCSignatureTreeItemBase> ItemWeak;
+
+	TWeakPtr<SRCSignatureTree> SignatureTreeWeak;
+};
diff --git a/Engine/Plugins/VirtualProduction/RemoteControl/Source/RemoteControlUI/Private/UI/Signature/SRCSignatureTree.cpp b/Engine/Plugins/VirtualProduction/RemoteControl/Source/RemoteControlUI/Private/UI/Signature/SRCSignatureTree.cpp
new file mode 100644
--- /dev/null
+++ b/Engine/Plugins/VirtualProduction/RemoteControl/Source/RemoteControlUI/Private/UI/Signature/SRCSignatureTree.cpp
@@ -0,0 +1,187 @@
+// Copyright Epic Games, Inc. All Rights Reserved.
+
+#include "SRCSignatureTree.h"
+#include "Columns/IRCSignatureColumn.h"
+#include "Items/RCSignatureTreeItemBase.h"
+#include "Items/RCSignatureTreeSignatureItem.h"
+#include "RemoteControlPreset.h"
+#include "RemoteControlSignatureRegistry.h"
+#include "SRCSignaturePanel.h"
+#include "SRCSignatureRow.h"
+#include "ScopedTransaction.h"
+#include "Styling/RemoteControlStyles.h"
+#include "UI/RemoteControlPanelStyle.h"
+#include "Widgets/Views/SHeaderRow.h"
+#include "Widgets/Views/STreeView.h"
+
+#define LOCTEXT_NAMESPACE "SRCSignatureTreeView"
+
+void SRCSignatureTree::Construct(const FArguments& InArgs, const TSharedRef<SRCSignaturePanel>& InSignaturePanel, const TSharedRef<SRemoteControlPanel>& InRCPanel)
+{
+	SRCLogicPanelListBase::Construct(SRCLogicPanelListBase::FArguments(), InSignaturePanel, InRCPanel);
+
+	SignaturePanelWeak = InSignaturePanel;
+
+	const FRCPanelStyle* RCPanelStyle = &FRemoteControlPanelStyle::Get()->GetWidgetStyle<FRCPanelStyle>("RemoteControlPanel.LogicControllersPanel");
+
+	HeaderRow = SNew(SHeaderRow)
+		.Style(&RCPanelStyle->HeaderRowStyle)
+		.Visibility(EVisibility::Visible)
+		.CanSelectGeneratedColumn(true);
+
+	ConstructColumns(InArgs._Columns);
+
+	ChildSlot
+	[
+		SAssignNew(SignatureTreeView, STreeView<TSharedPtr<FRCSignatureTreeItemBase>>)
+		.TreeItemsSource(&SignatureItems)
+		.HeaderRow(HeaderRow)
+		.OnGetChildren(this, &SRCSignatureTree::OnGetChildren)
+		.OnGenerateRow(this, &SRCSignatureTree::OnGenerateRow)
+		.OnContextMenuOpening(this, &SRCLogicPanelListBase::GetContextMenuWidget)
+		.SelectionMode(ESelectionMode::Multi)
+		.HighlightParentNodesForSelection(true)
+	];
+
+	Refresh();
+}
+
+URemoteControlSignatureRegistry* SRCSignatureTree::GetSignatureRegistry() const
+{
+	TSharedPtr<SRCSignaturePanel> SignaturePanel = SignaturePanelWeak.Pin();
+	if (!SignaturePanel.IsValid())
+	{
+		return nullptr;
+	}
+
+	if (URemoteControlPreset* Preset = SignaturePanel->GetPreset())
+	{
+		return Preset->GetSignatureRegistry();
+	}
+	return nullptr;
+}
+
+TSharedPtr<IRCSignatureColumn> SRCSignatureTree::FindColumn(FName InColumnName) const
+{
+	if (const TSharedRef<IRCSignatureColumn>* Column = Columns.Find(InColumnName))
+	{
+		return *Column;
+	}
+	return nullptr;
+}
+
+void SRCSignatureTree::Refresh()
+{
+	if (URemoteControlSignatureRegistry* SignatureRegistry = GetSignatureRegistry())
+	{
+		TConstArrayView<FRCSignature> Signatures = SignatureRegistry->GetSignatures();
+
+		SignatureItems.Reset(Signatures.Num());
+
+		TSharedPtr<SRCSignatureTree> SignatureTree = SharedThis(this);
+
+		for (const FRCSignature& Signature : Signatures)
+		{
+			TSharedRef<FRCSignatureTreeSignatureItem> SignatureItem = MakeShared<FRCSignatureTreeSignatureItem>(Signature, SignatureTree);
+			SignatureItem->RebuildChildren();
+			SignatureItems.Emplace(MoveTemp(SignatureItem));
+		}
+	}
+	else
+	{
+		SignatureItems.Empty();
+	}
+
+	SignatureTreeView->RequestTreeRefresh();
+}
+
+URemoteControlPreset* SRCSignatureTree::GetPreset()
+{
+	if (TSharedPtr<SRCSignaturePanel> SignaturePanel = SignaturePanelWeak.Pin())
+	{
+		return SignaturePanel->GetPreset();
+	}
+	return nullptr;
+}
+
+TArray<TSharedPtr<FRCLogicModeBase>> SRCSignatureTree::GetSelectedLogicItems()
+{
+	TArray<TSharedPtr<FRCLogicModeBase>> SelectedItems;
+	SelectedItems.Append(SignatureTreeView->GetSelectedItems());
+	return SelectedItems;
+}
+
+bool SRCSignatureTree::IsEmpty() const
+{
+	return SignatureItems.IsEmpty();
+}
+
+bool SRCSignatureTree::IsListFocused() const
+{
+	return SignatureTreeView->HasAnyUserFocus().IsSet() || ContextMenuWidgetCached.IsValid();
+}
+
+int32 SRCSignatureTree::Num() const
+{
+	return SignatureItems.Num();
+}
+
+int32 SRCSignatureTree::NumSelectedLogicItems() const
+{
+	return SignatureTreeView->GetNumItemsSelected();
+}
+
+int32 SRCSignatureTree::RemoveModel(const TSharedPtr<FRCLogicModeBase> InItem)
+{
+	if (InItem.IsValid())
+	{
+		return StaticCastSharedPtr<FRCSignatureTreeItemBase>(InItem)->RemoveFromRegistry();
+	}
+	return 0;
+}
+
+void SRCSignatureTree::DeleteSelectedPanelItems()
+{
+	const TArray<TSharedPtr<FRCSignatureTreeItemBase>> SelectedSignatures = SignatureTreeView->GetSelectedItems();
+	if (SelectedSignatures.IsEmpty())
+	{
+		return;
+	}
+
+	FScopedTransaction Transaction(LOCTEXT("RemoveSelectedSignatures", "Remove Selected Signatures"));
+	DeleteItemsFromLogicPanel<FRCSignatureTreeItemBase>(SignatureItems, SelectedSignatures);
+}
+
+void SRCSignatureTree::Reset()
+{
+	Refresh();
+}
+
+void SRCSignatureTree::ConstructColumns(TConstArrayView<TSharedRef<IRCSignatureColumn>> InColumns)
+{
+	HeaderRow->ClearColumns();
+	Columns.Empty(InColumns.Num());
+
+	for (const TSharedRef<IRCSignatureColumn>& Column : InColumns)
+	{
+		Columns.Add(Column->GetColumnId(), Column);
+		HeaderRow->AddColumn(Column->ConstructHeaderRowColumn());
+		HeaderRow->SetShowGeneratedColumn(Column->GetColumnId(), Column->ShouldShowColumnByDefault());
+	}
+}
+
+TSharedRef<ITableRow> SRCSignatureTree::OnGenerateRow(TSharedPtr<FRCSignatureTreeItemBase> InItem, const TSharedRef<STableViewBase>& InTableView)
+{
+	check(InItem.IsValid());
+	return SNew(SRCSignatureRow, InItem, SharedThis(this), InTableView);
+}
+
+void SRCSignatureTree::OnGetChildren(TSharedPtr<FRCSignatureTreeItemBase> InItem, TArray<TSharedPtr<FRCSignatureTreeItemBase>>& OutChildren) const
+{
+	if (InItem.IsValid())
+	{
+		OutChildren.Append(InItem->GetChildren());
+	}
+}
+
+#undef LOCTEXT_NAMESPACE
diff --git a/Engine/Plugins/VirtualProduction/RemoteControl/Source/RemoteControlUI/Private/UI/Signature/SRCSignatureTree.h b/Engine/Plugins/VirtualProduction/RemoteControl/Source/RemoteControlUI/Private/UI/Signature/SRCSignatureTree.h
new file mode 100644
--- /dev/null
+++ b/Engine/Plugins/VirtualProduction/RemoteControl/Source/RemoteControlUI/Private/UI/Signature/SRCSignatureTree.h
@@ -0,0 +1,59 @@
+// Copyright Epic Games, Inc. All Rights Reserved.
+
+#pragma once
+
+#include "UI/BaseLogicUI/SRCLogicPanelListBase.h"
+
+class FRCSignatureTreeItemBase;
+class IRCSignatureColumn;
+class SHeaderRow;
+class SRCSignaturePanel;
+class URemoteControlSignatureRegistry;
+template<typename ItemType> class STreeView;
+
+class SRCSignatureTree : public SRCLogicPanelListBase
+{
+public:
+	SLATE_BEGIN_ARGS(SRCSignatureTree) {}
+		SLATE_ARGUMENT(TArray<TSharedRef<IRCSignatureColumn>>, Columns)
+	SLATE_END_ARGS()
+
+	void Construct(const FArguments& InArgs, const TSharedRef<SRCSignaturePanel>& InSignaturePanel, const TSharedRef<SRemoteControlPanel>& InRCPanel);
+
+	URemoteControlSignatureRegistry* GetSignatureRegistry() const;
+
+	TSharedPtr<IRCSignatureColumn> FindColumn(FName InColumnName) const;
+
+	void Refresh();
+
+	//~ Begin SRCLogicPanelListBase
+	virtual URemoteControlPreset* GetPreset() override;
+	virtual TArray<TSharedPtr<FRCLogicModeBase>> GetSelectedLogicItems() override;
+	virtual bool IsEmpty() const override;
+	virtual bool IsListFocused() const override;
+	virtual int32 Num() const override;
+	virtual int32 NumSelectedLogicItems() const override;
+	virtual int32 RemoveModel(const TSharedPtr<FRCLogicModeBase> InItem) override;
+	virtual void DeleteSelectedPanelItems() override;
+	virtual void BroadcastOnItemRemoved() override {}
+	virtual void Reset() override;
+	//~ End SRCLogicPanelListBase
+
+private:
+	void ConstructColumns(TConstArrayView<TSharedRef<IRCSignatureColumn>> InColumns);
+
+	TSharedRef<ITableRow> OnGenerateRow(TSharedPtr<FRCSignatureTreeItemBase> InItem, const TSharedRef<STableViewBase>& InTableView);
+
+	void OnGetChildren(TSharedPtr<FRCSignatureTreeItemBase> InItem, TArray<TSharedPtr<FRCSignatureTreeItemBase>>& OutChildren) const;
+
+	/** All Known Signatures in this Tree View */
+	TArray<TSharedPtr<FRCSignatureTreeItemBase>> SignatureItems;
+
+	TMap<FName, TSharedRef<IRCSignatureColumn>> Columns;
+
+	TSharedPtr<STreeView<TSharedPtr<FRCSignatureTreeItemBase>>> SignatureTreeView;
+
+	TSharedPtr<SHeaderRow> HeaderRow;
+
+	TWeakPtr<SRCSignaturePanel> SignaturePanelWeak;
+};
diff --git a/Engine/Plugins/VirtualProduction/RemoteControl/Source/RemoteControlUI/Public/IRemoteControlUIModule.h b/Engine/Plugins/VirtualProduction/RemoteControl/Source/RemoteControlUI/Public/IRemoteControlUIModule.h
--- a/Engine/Plugins/VirtualProduction/RemoteControl/Source/RemoteControlUI/Public/IRemoteControlUIModule.h
+++ b/Engine/Plugins/VirtualProduction/RemoteControl/Source/RemoteControlUI/Public/IRemoteControlUIModule.h
@@ -16,6 +16,7 @@ class FRCPanelWidgetRegistry;
 class IDetailCategoryBuilder;
 class IDetailLayoutBuilder;
 class IPropertyHandle;
+class IRCSignatureCustomization;
 class URCVirtualPropertyBase;
 class URemoteControlPreset;
 struct FOnGenerateGlobalRowExtensionArgs;
@@ -234,6 +235,18 @@ public:
 	 */
 	virtual uint32 GetRemoteControlAssetCategory() const = 0;
 
+	/**
+	 * Register the given signature customization to extend the Signature Panel capabilities like Drag/Drop handling
+	 * @param InCustomization the signature customization instance to register
+	 */
+	virtual void RegisterSignatureCustomization(const TSharedPtr<IRCSignatureCustomization>& InCustomization) = 0;
+
+	/**
+	 * Unregisters the given signature customization
+	 * @param InCustomization the signature customization instance to unregister
+	 */
+	virtual void UnregisterSignatureCustomization(const TSharedPtr<IRCSignatureCustomization>& InCustomization) = 0;
+
 	/**
 	 * Register a widget factory to handle creating a widget in the control panel.
 	 */
diff --git a/Engine/Plugins/VirtualProduction/RemoteControl/Source/RemoteControlUI/Public/UI/Signature/IRCSignatureCustomization.h b/Engine/Plugins/VirtualProduction/RemoteControl/Source/RemoteControlUI/Public/UI/Signature/IRCSignatureCustomization.h
new file mode 100644
--- /dev/null
+++ b/Engine/Plugins/VirtualProduction/RemoteControl/Source/RemoteControlUI/Public/UI/Signature/IRCSignatureCustomization.h
@@ -0,0 +1,16 @@
+// Copyright Epic Games, Inc. All Rights Reserved.
+
+#pragma once
+
+class FDragDropEvent;
+class FReply;
+class IRCSignatureItem;
+
+/** Customization Interface to extend the capabilities of RC Signature UI */
+class IRCSignatureCustomization
+{
+public:
+	virtual bool CanAcceptDrop(const FDragDropEvent& InDragDropEvent, IRCSignatureItem* InSignatureItem) const = 0;
+
+	virtual FReply AcceptDrop(const FDragDropEvent& InDragDropEvent, IRCSignatureItem* InSignatureItem) const = 0;
+};
diff --git a/Engine/Plugins/VirtualProduction/RemoteControl/Source/RemoteControlUI/Public/UI/Signature/IRCSignatureItem.h b/Engine/Plugins/VirtualProduction/RemoteControl/Source/RemoteControlUI/Public/UI/Signature/IRCSignatureItem.h
new file mode 100644
--- /dev/null
+++ b/Engine/Plugins/VirtualProduction/RemoteControl/Source/RemoteControlUI/Public/UI/Signature/IRCSignatureItem.h
@@ -0,0 +1,18 @@
+// Copyright Epic Games, Inc. All Rights Reserved.
+
+#pragma once
+
+#include "Containers/ContainersFwd.h"
+#include "UObject/WeakObjectPtrTemplatesFwd.h"
+
+class AActor;
+struct FGuid;
+
+/** Interface for external implementations to interact with the Signature View Model */
+class IRCSignatureItem
+{
+public:
+	virtual void AddFieldEntities(TConstArrayView<FGuid> InFieldEntityIds) = 0;
+
+	virtual void ApplySignature(TConstArrayView<TWeakObjectPtr<AActor>> InActors) = 0;
+};
