Remote Control: Hide Save button when RC Preset is not an asset.
Without this change the Save button still appears in non-asset presets, but does nothing and logs an error instead.

diff --git a/Engine/Plugins/VirtualProduction/RemoteControl/Source/RemoteControlUI/Private/UI/SRemoteControlPanel.cpp b/Engine/Plugins/VirtualProduction/RemoteControl/Source/RemoteControlUI/Private/UI/SRemoteControlPanel.cpp
--- a/Engine/Plugins/VirtualProduction/RemoteControl/Source/RemoteControlUI/Private/UI/SRemoteControlPanel.cpp
+++ b/Engine/Plugins/VirtualProduction/RemoteControl/Source/RemoteControlUI/Private/UI/SRemoteControlPanel.cpp
@@ -1418,7 +1418,9 @@ void SRemoteControlPanel::BindRemoteControlCommands()
 	ActionList.MapAction(
 		Commands.SavePreset,
 		FExecuteAction::CreateSP(this, &SRemoteControlPanel::SaveAsset_Execute),
-		FCanExecuteAction::CreateSP(this, &SRemoteControlPanel::CanSaveAsset));
+		FCanExecuteAction(),
+		FGetActionCheckState(),
+		FIsActionButtonVisible::CreateSP(this, &SRemoteControlPanel::CanSaveAsset));
 
 	ActionList.MapAction(
 		Commands.FindPresetInContentBrowser,
@@ -2105,9 +2107,6 @@ void SRemoteControlPanel::RegisterDefaultToolBar()
 		UToolMenu* ToolbarBuilder = ToolMenus->RegisterMenu(DefaultRemoteControlPanelToolBarName, NAME_None, EMultiBoxType::SlimHorizontalToolBar);
 		ToolbarBuilder->StyleName = "ContentBrowser.ToolBar";
 
-#if 0
-		ToolbarBuilder->StyleName = "AssetEditorToolbar";
-#endif
 		{
 			FToolMenuSection& AssetSection = ToolbarBuilder->AddSection("Asset");
 			AssetSection.AddEntry(FToolMenuEntry::InitToolBarButton(FRemoteControlCommands::Get().SavePreset, TAttribute<FText>(), TAttribute<FText>(), FSlateIcon(FAppStyle::GetAppStyleSetName(), TEXT("Icons.Save"))));
@@ -2469,7 +2468,7 @@ FText SRemoteControlPanel::HandlePresetName() const
 
 bool SRemoteControlPanel::CanSaveAsset() const
 {
-	return Preset.IsValid();
+	return Preset.IsValid() && Preset->IsAsset();
 }
 
 void SRemoteControlPanel::SaveAsset_Execute() const
