Remote Control: Actor Owner  to be the default group filter

diff --git a/Engine/Plugins/VirtualProduction/RemoteControl/Source/RemoteControlUI/Private/UI/SRCPanelExposedEntitiesGroup.h b/Engine/Plugins/VirtualProduction/RemoteControl/Source/RemoteControlUI/Private/UI/SRCPanelExposedEntitiesGroup.h
--- a/Engine/Plugins/VirtualProduction/RemoteControl/Source/RemoteControlUI/Private/UI/SRCPanelExposedEntitiesGroup.h
+++ b/Engine/Plugins/VirtualProduction/RemoteControl/Source/RemoteControlUI/Private/UI/SRCPanelExposedEntitiesGroup.h
@@ -8,7 +8,7 @@
 #include "Widgets/SCompoundWidget.h"
 
 /** Grouping type available */
-enum class EFieldGroupType
+enum class EFieldGroupType : uint8
 {
 	/** No Grouping. Will be shown as a flat list */
 	None,
diff --git a/Engine/Plugins/VirtualProduction/RemoteControl/Source/RemoteControlUI/Private/UI/SRCPanelExposedEntitiesList.h b/Engine/Plugins/VirtualProduction/RemoteControl/Source/RemoteControlUI/Private/UI/SRCPanelExposedEntitiesList.h
--- a/Engine/Plugins/VirtualProduction/RemoteControl/Source/RemoteControlUI/Private/UI/SRCPanelExposedEntitiesList.h
+++ b/Engine/Plugins/VirtualProduction/RemoteControl/Source/RemoteControlUI/Private/UI/SRCPanelExposedEntitiesList.h
@@ -278,7 +278,7 @@ private:
 	/** Holds all the exposed entities groups. */
 	TArray<TSharedPtr<SRCPanelExposedEntitiesGroup>> ExposedEntitiesGroups;
 	/** Holds the current group type */
-	EFieldGroupType CurrentGroupType = EFieldGroupType::None;
+	EFieldGroupType CurrentGroupType = EFieldGroupType::Owner;
 	/** Holds the current sorting type */
 	ERCGroupOrder CurrentGroupSortType = ERCGroupOrder::None;
 	/** Holds all the entities groups currently in the list */
