[WebRemoteControl] Fixing API inconsistency with PresetGet/SetController routes. It should support either the Preset Name or Id to be consistent with the other routes.

diff --git a/Engine/Plugins/VirtualProduction/RemoteControl/Source/WebRemoteControl/Private/WebRemoteControl.cpp b/Engine/Plugins/VirtualProduction/RemoteControl/Source/WebRemoteControl/Private/WebRemoteControl.cpp
--- a/Engine/Plugins/VirtualProduction/RemoteControl/Source/WebRemoteControl/Private/WebRemoteControl.cpp
+++ b/Engine/Plugins/VirtualProduction/RemoteControl/Source/WebRemoteControl/Private/WebRemoteControl.cpp
@@ -2236,7 +2236,7 @@ bool FWebRemoteControlModule::HandlePresetSetControllerRoute(const FHttpServerRe
 	Args.ControllerName = Request.PathParams.FindChecked(TEXT("controller"));
 
 	// 5. Acquire Remote Control Preset object
-	URemoteControlPreset* Preset = IRemoteControlModule::Get().ResolvePreset(*Args.PresetName);
+	URemoteControlPreset* Preset = WebRemoteControl::GetPreset(*Args.PresetName);
 	if (Preset == nullptr)
 	{
 		Response->Code = EHttpServerResponseCodes::NotFound;
@@ -2296,7 +2296,7 @@ bool FWebRemoteControlModule::HandlePresetGetControllerRoute(const FHttpServerRe
 	Args.ControllerName = Request.PathParams.FindChecked(TEXT("controller"));
 
 	// 3. Acquire Remote Control Preset object
-	URemoteControlPreset* Preset = IRemoteControlModule::Get().ResolvePreset(*Args.PresetName);
+	URemoteControlPreset* Preset = WebRemoteControl::GetPreset(*Args.PresetName);
 	if (Preset == nullptr)
 	{
 		Response->Code = EHttpServerResponseCodes::NotFound;
