[WebRemoteControl] Fix Websocket http message wrapper to propagate the RequestId correctly from the incoming MessageId.

diff --git a/Engine/Plugins/VirtualProduction/RemoteControl/Source/WebRemoteControl/Private/WebRemoteControl.cpp b/Engine/Plugins/VirtualProduction/RemoteControl/Source/WebRemoteControl/Private/WebRemoteControl.cpp
--- a/Engine/Plugins/VirtualProduction/RemoteControl/Source/WebRemoteControl/Private/WebRemoteControl.cpp
+++ b/Engine/Plugins/VirtualProduction/RemoteControl/Source/WebRemoteControl/Private/WebRemoteControl.cpp
@@ -2376,6 +2376,8 @@ void FWebRemoteControlModule::HandleWebSocketHttpMessage(const FRemoteControlWeb
 		Wrapper.Passphrase = WebSocketMessage.Header[WebRemoteControlInternalUtils::PassphraseHeader][0];
 	}
 
+	Wrapper.RequestId = WebSocketMessage.MessageId;
+
 	LogRequestExternally(Wrapper.RequestId, TEXT("UE Received"));
 	
 	FMemoryWriter Writer(UTF8Response);
