[WebRemoteControl] Support for jpeg thumbnails and correctly identify the mime types in header (include svg for slate brushes)

diff --git a/Engine/Plugins/VirtualProduction/RemoteControl/Source/WebRemoteControl/Private/WebRemoteControlEditorRoutes.cpp b/Engine/Plugins/VirtualProduction/RemoteControl/Source/WebRemoteControl/Private/WebRemoteControlEditorRoutes.cpp
--- a/Engine/Plugins/VirtualProduction/RemoteControl/Source/WebRemoteControl/Private/WebRemoteControlEditorRoutes.cpp
+++ b/Engine/Plugins/VirtualProduction/RemoteControl/Source/WebRemoteControl/Private/WebRemoteControlEditorRoutes.cpp
@@ -7,6 +7,7 @@
 #include "WebRemoteControl.h"
 
 // Http
+#include "GenericPlatform/GenericPlatformHttp.h"
 #include "HttpPath.h"
 #include "HttpServerRequest.h"
 #include "HttpServerConstants.h"
@@ -137,16 +138,18 @@ bool FWebRemoteControlEditorRoutes::HandleGetThumbnailRoute(const FHttpServerReq
 		{
 			if (FObjectThumbnail* Thumbnail = ThumbnailMap.Find(ObjectFullName))
 			{
-				WebRemoteControlInternalUtils::AddContentTypeHeaders(Response.Get(), TEXT("image/png"));
 				IImageWrapperModule& ImageWrapperModule = FModuleManager::Get().LoadModuleChecked<IImageWrapperModule>(TEXT("ImageWrapper"));
 				EImageFormat Format = ImageWrapperModule.DetectImageFormat((void*)Thumbnail->AccessCompressedImageData().GetData(), Thumbnail->AccessCompressedImageData().Num());
-				if (Format == EImageFormat::PNG)
+
+				if (Format == EImageFormat::PNG || Format == EImageFormat::JPEG)
 				{
+					// Converting png/jpeg thumbnails to be in the right color format (BGRA) for web browsers.
+					WebRemoteControlInternalUtils::AddContentTypeHeaders(Response.Get(), Format == EImageFormat::PNG ? TEXT("image/png") : TEXT("image/jpeg"));
 					TSharedPtr<IImageWrapper> Wrapper = ImageWrapperModule.CreateImageWrapper(Format);
+
 					Wrapper->SetRaw(Thumbnail->GetUncompressedImageData().GetData(), Thumbnail->GetUncompressedImageData().Num(), Thumbnail->GetImageWidth(), Thumbnail->GetImageHeight(), ERGBFormat::BGRA, 8);
 					if (Wrapper)
 					{
-						ERGBFormat RGBFormat = Wrapper->GetFormat();
 						Response->Body = Wrapper->GetCompressed();
 						Response->Code = EHttpServerResponseCodes::Ok;
 					}
@@ -158,10 +161,12 @@ bool FWebRemoteControlEditorRoutes::HandleGetThumbnailRoute(const FHttpServerReq
 		{
 			if (const FSlateBrush* ThumbnailBrush = FClassIconFinder::FindThumbnailForClass(AssetData.GetClass()))
 			{
-				FName ResourceName = ThumbnailBrush->GetResourceName();
-				if (FFileHelper::LoadFileToArray(Response->Body, *ResourceName.ToString()))
+				const FString ResourceName = ThumbnailBrush->GetResourceName().ToString().ToLower();
+				const FString MimeType = FGenericPlatformHttp::GetMimeType(ResourceName);
+
+				if (MimeType.StartsWith(TEXT("image/")) && FFileHelper::LoadFileToArray(Response->Body, *ResourceName))
 				{
-					WebRemoteControlInternalUtils::AddContentTypeHeaders(Response.Get(), TEXT("image/png"));
+					WebRemoteControlInternalUtils::AddContentTypeHeaders(Response.Get(), MimeType);
 					Response->Code = EHttpServerResponseCodes::Ok;
 				}
 			}
diff --git a/Engine/Plugins/VirtualProduction/RemoteControl/Source/WebRemoteControl/Private/WebRemoteControlInternalUtils.cpp b/Engine/Plugins/VirtualProduction/RemoteControl/Source/WebRemoteControl/Private/WebRemoteControlInternalUtils.cpp
--- a/Engine/Plugins/VirtualProduction/RemoteControl/Source/WebRemoteControl/Private/WebRemoteControlInternalUtils.cpp
+++ b/Engine/Plugins/VirtualProduction/RemoteControl/Source/WebRemoteControl/Private/WebRemoteControlInternalUtils.cpp
@@ -125,9 +125,9 @@ TSharedRef<FHttpServerRequest> UnwrapHttpRequest(const FRCRequestWrapper& Wrappe
 void SerializeWrappedCallResponse(int32 RequestId, TUniquePtr<FHttpServerResponse> Response, FMemoryWriter& Writer)
 {
 	FRCJsonStructSerializerBackend Backend(Writer, FRCJsonStructSerializerBackend::DefaultSerializerFlags);
-	TSharedPtr<TJsonWriter<ANSICHAR>> JsonWriter = TJsonWriter<ANSICHAR>::Create(&Writer);
-	TArray<FString>* ContentTypeHeaders = Response->Headers.Find(TEXT("Content-Type"));
-	const bool bIsBinaryData = ContentTypeHeaders && ContentTypeHeaders->Contains(TEXT("image/png"));
+	const TSharedPtr<TJsonWriter<ANSICHAR>> JsonWriter = TJsonWriter<ANSICHAR>::Create(&Writer);
+	const TArray<FString>* ContentTypeHeaders = Response->Headers.Find(TEXT("Content-Type"));
+	const bool bIsBinaryData = ContentTypeHeaders && ContentTypeHeaders->ContainsByPredicate([](const FString& InHeader) { return InHeader.StartsWith(TEXT("image/"));}); 
 
 	JsonWriter->WriteObjectStart();
 	JsonWriter->WriteValue(TEXT("RequestId"), RequestId);
