Remote Control Preset: Fix crash accessing  empty override material array

diff --git a/Engine/Plugins/VirtualProduction/RemoteControl/Source/RemoteControl/Private/RemoteControlPreset.cpp b/Engine/Plugins/VirtualProduction/RemoteControl/Source/RemoteControl/Private/RemoteControlPreset.cpp
--- a/Engine/Plugins/VirtualProduction/RemoteControl/Source/RemoteControl/Private/RemoteControlPreset.cpp
+++ b/Engine/Plugins/VirtualProduction/RemoteControl/Source/RemoteControl/Private/RemoteControlPreset.cpp
@@ -2160,7 +2160,7 @@ void URemoteControlPreset::OnObjectPropertyChanged(UObject* Object, struct FProp
 			{
 				if (UMeshComponent* MeshComponent = Cast<UMeshComponent>(Object))
 				{
-					if (CacheEntry.bHadValue && MeshComponent->OverrideMaterials[CacheEntry.ArrayIndex] == NULL)
+					if (CacheEntry.bHadValue && MeshComponent->OverrideMaterials.IsValidIndex(CacheEntry.ArrayIndex) && MeshComponent->OverrideMaterials[CacheEntry.ArrayIndex] == NULL)
 					{
 						FRCResetToDefaultArgs Args;
 						Args.Property = Event.Property;
@@ -2273,7 +2273,7 @@ void URemoteControlPreset::OnPreObjectPropertyChanged(UObject* Object, const cla
 								{
 									FPreMaterialModifiedCache& NewEntry = PreMaterialModifiedCache.FindOrAdd(RCProperty->GetId());
 									NewEntry.ArrayIndex = RCProperty->FieldPathInfo.Segments[0].ArrayIndex;
-									NewEntry.bHadValue = MeshComponent->OverrideMaterials[NewEntry.ArrayIndex] != NULL;
+									NewEntry.bHadValue = MeshComponent->OverrideMaterials.IsValidIndex(NewEntry.ArrayIndex) && MeshComponent->OverrideMaterials[NewEntry.ArrayIndex] != NULL;
 								}
 							}
 							
