Remote Control: set the grouping default back to no grouping

diff --git a/Engine/Plugins/VirtualProduction/RemoteControl/Source/RemoteControlUI/Private/UI/SRCPanelExposedEntitiesList.h b/Engine/Plugins/VirtualProduction/RemoteControl/Source/RemoteControlUI/Private/UI/SRCPanelExposedEntitiesList.h
--- a/Engine/Plugins/VirtualProduction/RemoteControl/Source/RemoteControlUI/Private/UI/SRCPanelExposedEntitiesList.h
+++ b/Engine/Plugins/VirtualProduction/RemoteControl/Source/RemoteControlUI/Private/UI/SRCPanelExposedEntitiesList.h
@@ -278,7 +278,7 @@ private:
 	/** Holds all the exposed entities groups. */
 	TArray<TSharedPtr<SRCPanelExposedEntitiesGroup>> ExposedEntitiesGroups;
 	/** Holds the current group type */
-	EFieldGroupType CurrentGroupType = EFieldGroupType::Owner;
+	EFieldGroupType CurrentGroupType = EFieldGroupType::None;
 	/** Holds the current sorting type */
 	ERCGroupOrder CurrentGroupSortType = ERCGroupOrder::None;
 	/** Holds all the entities groups currently in the list */
