Remote Control UI: Removed the Panel Drawers as this was a main cause of confusion of mode switching. Instead the Toolbar Buttons at the top right are combining their existing functionality and the functionality of the panel drawer tabs.

diff --git a/Engine/Plugins/VirtualProduction/RemoteControl/Source/RemoteControlUI/Private/Commands/RemoteControlCommands.cpp b/Engine/Plugins/VirtualProduction/RemoteControl/Source/RemoteControlUI/Private/Commands/RemoteControlCommands.cpp
--- a/Engine/Plugins/VirtualProduction/RemoteControl/Source/RemoteControlUI/Private/Commands/RemoteControlCommands.cpp
+++ b/Engine/Plugins/VirtualProduction/RemoteControl/Source/RemoteControlUI/Private/Commands/RemoteControlCommands.cpp
@@ -20,40 +20,88 @@ FRemoteControlCommands::FRemoteControlCommands()
 void FRemoteControlCommands::RegisterCommands()
 {
 	// Save Preset
-	UI_COMMAND(SavePreset, "Save", "Saves this preset", EUserInterfaceActionType::Button, FInputChord(EModifierKey::Control, EKeys::S));
-
-	// Find Preset
-	UI_COMMAND(FindPresetInContentBrowser, "Browse to Preset", "Browses to the associated preset and selects it in the most recently used Content Browser (summoning one if necessary)", EUserInterfaceActionType::Button, FInputChord(EModifierKey::Control, EKeys::B));
-
-	// Toggle Protocol Mappings
-	UI_COMMAND(ToggleProtocolMappings, "Protocols", "View list of protocols mapped to active selection.", EUserInterfaceActionType::ToggleButton, FInputChord());
-
-	// Toggle Logic Editor
-	UI_COMMAND(ToggleLogicEditor, "Logic", "View the logic applied to active selection.", EUserInterfaceActionType::ToggleButton, FInputChord());
-
-	// Toggle Signature Editor
-	UI_COMMAND(ToggleSignatureEditor, "Signature", "View the current signatures available", EUserInterfaceActionType::ToggleButton, FInputChord());
+	UI_COMMAND(SavePreset, 
+		"Save", "Saves this preset",
+		EUserInterfaceActionType::Button,
+		FInputChord(EModifierKey::Control, EKeys::S));
+
+	// Protocol Mode
+	UI_COMMAND(ActivateProtocolsMode,
+		"Protocols", "View list of protocols mapped to active selection.",
+		EUserInterfaceActionType::RadioButton,
+		FInputChord());
+
+	// Entity Details Mode
+	UI_COMMAND(ActivateDetailsMode,
+		"Details", "Activate the Entity Details Panel Mode",
+		EUserInterfaceActionType::RadioButton,
+		FInputChord());
+
+	// Logic Mode
+	UI_COMMAND(ActivateLogicMode,
+		"Logic", "View the logic applied to active selection.",
+		EUserInterfaceActionType::RadioButton,
+		FInputChord());
+
+	// Signature Mode
+	UI_COMMAND(ActivateSignatureMode,
+		"Signature",
+		"View the current signatures available",
+		EUserInterfaceActionType::RadioButton,
+		FInputChord());
+
+	// Output Log Mode
+	UI_COMMAND(ActivateOutputLogMode,
+		"Output Log",
+		"View the output log",
+		EUserInterfaceActionType::RadioButton,
+		FInputChord());
 
 	// Delete Entity
-	UI_COMMAND(DeleteEntity, "Delete", "Delete the selected group/exposed entity from the list.", EUserInterfaceActionType::Button, FInputChord(EKeys::Delete));
+	UI_COMMAND(DeleteEntity,
+		"Delete",
+		"Delete the selected group/exposed entity from the list.",
+		EUserInterfaceActionType::Button,
+		FInputChord(EKeys::Delete));
 
 	// Rename Entity
-	UI_COMMAND(RenameEntity, "Rename", "Rename the selected group/exposed entity.", EUserInterfaceActionType::Button, FInputChord(EKeys::F2));
+	UI_COMMAND(RenameEntity,
+		"Rename", "Rename the selected group/exposed entity.",
+		EUserInterfaceActionType::Button,
+		FInputChord(EKeys::F2));
 
 	// Modify Entity Prop Id
-	UI_COMMAND(ChangePropId, "ChangePropid", "Change the selected property Ids", EUserInterfaceActionType::Button, FInputChord(EKeys::F3, EModifierKey::Control));
+	UI_COMMAND(ChangePropId,
+		"ChangePropid", "Change the selected property Ids",
+		EUserInterfaceActionType::Button,
+		FInputChord(EKeys::F3, EModifierKey::Control));
 
 	// Copy Item
-	UI_COMMAND(CopyItem, "Copy", "Copy the selected UI item", EUserInterfaceActionType::Button, FInputChord(EKeys::C, EModifierKey::Control));
+	UI_COMMAND(CopyItem,
+		"Copy",
+		"Copy the selected UI item",
+		EUserInterfaceActionType::Button,
+		FInputChord(EKeys::C, EModifierKey::Control));
 
 	// Paste Item
-	UI_COMMAND(PasteItem, "Paste", "Paste the selected UI item", EUserInterfaceActionType::Button, FInputChord(EKeys::V, EModifierKey::Control));
+	UI_COMMAND(PasteItem,
+		"Paste", "Paste the selected UI item",
+		EUserInterfaceActionType::Button,
+		FInputChord(EKeys::V, EModifierKey::Control));
 
 	// Duplicate Item
-	UI_COMMAND(DuplicateItem, "Duplicate", "Duplicate the selected UI item", EUserInterfaceActionType::Button, FInputChord(EKeys::D, EModifierKey::Control));
+	UI_COMMAND(DuplicateItem,
+		"Duplicate",
+		"Duplicate the selected UI item",
+		EUserInterfaceActionType::Button,
+		FInputChord(EKeys::D, EModifierKey::Control));
 
 	// Update Value
-	UI_COMMAND(UpdateValue, "Update Value", "Update the selected UI item value with the one in the fields list", EUserInterfaceActionType::Button, FInputChord(EKeys::U, EModifierKey::Control));
+	UI_COMMAND(UpdateValue,
+		"Update Value",
+		"Update the selected UI item value with the one in the fields list",
+		EUserInterfaceActionType::Button,
+		FInputChord(EKeys::U, EModifierKey::Control));
 }
 
 #undef LOCTEXT_NAMESPACE
diff --git a/Engine/Plugins/VirtualProduction/RemoteControl/Source/RemoteControlUI/Private/Commands/RemoteControlCommands.h b/Engine/Plugins/VirtualProduction/RemoteControl/Source/RemoteControlUI/Private/Commands/RemoteControlCommands.h
--- a/Engine/Plugins/VirtualProduction/RemoteControl/Source/RemoteControlUI/Private/Commands/RemoteControlCommands.h
+++ b/Engine/Plugins/VirtualProduction/RemoteControl/Source/RemoteControlUI/Private/Commands/RemoteControlCommands.h
@@ -21,11 +21,6 @@ public:
 
 	//~ END : TCommands<> Implementation(s)
 
-	/**
-	 * Holds the information about UI Command that finds the actively edited preset in the Content Browser.
-	 */
-	TSharedPtr<FUICommandInfo> FindPresetInContentBrowser;
-
 	/**
 	 * Holds the information about UI Command that saves the actively edited preset.
 	 */
@@ -34,17 +29,27 @@ public:
 	/**
 	 * Holds the information about UI Command that brings up a panel which holds the active protocol mappings.
 	 */
-	TSharedPtr<FUICommandInfo> ToggleProtocolMappings;
+	TSharedPtr<FUICommandInfo> ActivateProtocolsMode;
+
+	/**
+	 * Holds the information about UI Command that brings up the Entity Details UI panel.
+	 */
+	TSharedPtr<FUICommandInfo> ActivateDetailsMode;
 
 	/**
 	 * Holds the information about UI Command that brings up a panel which enables the RC Logical Behaviour.
 	 */
-	TSharedPtr<FUICommandInfo> ToggleLogicEditor;
+	TSharedPtr<FUICommandInfo> ActivateLogicMode;
 
 	/**
 	 * Holds the information about UI Command that brings up the RC Signature UI panel.
 	 */
-	TSharedPtr<FUICommandInfo> ToggleSignatureEditor;
+	TSharedPtr<FUICommandInfo> ActivateSignatureMode;
+
+	/**
+	 * Holds the information about UI Command that brings up the Output panel.
+	 */
+	TSharedPtr<FUICommandInfo> ActivateOutputLogMode;
 
 	/**
 	 * Holds the information about UI Command that deletes currently selected group/exposed entity.
diff --git a/Engine/Plugins/VirtualProduction/RemoteControl/Source/RemoteControlUI/Private/RemoteControlUIModule.cpp b/Engine/Plugins/VirtualProduction/RemoteControl/Source/RemoteControlUI/Private/RemoteControlUIModule.cpp
--- a/Engine/Plugins/VirtualProduction/RemoteControl/Source/RemoteControlUI/Private/RemoteControlUIModule.cpp
+++ b/Engine/Plugins/VirtualProduction/RemoteControl/Source/RemoteControlUI/Private/RemoteControlUIModule.cpp
@@ -35,6 +35,7 @@
 #include "Styling/AppStyle.h"
 #include "Styling/AppStyle.h"
 #include "Textures/SlateIcon.h"
+#include "Toolkits/GlobalEditorCommonCommands.h"
 #include "UI/Action/SRCActionPanel.h"
 #include "UI/Behaviour/Builtin/Conditional/SRCBehaviourConditional.h"
 #include "UI/Behaviour/SRCBehaviourPanel.h"
@@ -313,7 +314,7 @@ TSharedRef<SRemoteControlPanel> FRemoteControlUIModule::CreateRemoteControlPanel
 
 	if (TSharedPtr<SRemoteControlPanel> Panel = WeakActivePanel.Pin())
 	{
-		Panel->SetLiveMode(bIsInLiveMode);
+		Panel->SetActiveMode(ERCPanelMode::Live);
 	}
 
 	TSharedRef<SRemoteControlPanel> PanelRef = SNew(SRemoteControlPanel, Preset, ToolkitHost)
@@ -327,7 +328,7 @@ TSharedRef<SRemoteControlPanel> FRemoteControlUIModule::CreateRemoteControlPanel
 					{
 						if (ActivePanel != InPanel)
 						{
-							ActivePanel->SetLiveMode(true);
+							ActivePanel->SetActiveMode(ERCPanelMode::Live);
 						}
 					}
 					WeakActivePanel = MoveTemp(InPanel);
@@ -663,10 +664,12 @@ void FRemoteControlUIModule::UnbindRemoteControlCommands()
 		FUICommandList& ActionList = *MainFrame.GetMainFrameCommandBindings();
 
 		ActionList.UnmapAction(Commands.SavePreset);
-		ActionList.UnmapAction(Commands.FindPresetInContentBrowser);
-		ActionList.UnmapAction(Commands.ToggleProtocolMappings);
-		ActionList.UnmapAction(Commands.ToggleLogicEditor);
-		ActionList.UnmapAction(Commands.ToggleSignatureEditor);
+		ActionList.UnmapAction(FGlobalEditorCommonCommands::Get().FindInContentBrowser);
+		ActionList.UnmapAction(Commands.ActivateLogicMode);
+		ActionList.UnmapAction(Commands.ActivateDetailsMode);
+		ActionList.UnmapAction(Commands.ActivateSignatureMode);
+		ActionList.UnmapAction(Commands.ActivateProtocolsMode);
+		ActionList.UnmapAction(Commands.ActivateOutputLogMode);
 		ActionList.UnmapAction(Commands.DeleteEntity);
 		ActionList.UnmapAction(Commands.RenameEntity);
 		ActionList.UnmapAction(Commands.CopyItem);
diff --git a/Engine/Plugins/VirtualProduction/RemoteControl/Source/RemoteControlUI/Private/UI/BaseLogicUI/SRCLogicPanelListBase.cpp b/Engine/Plugins/VirtualProduction/RemoteControl/Source/RemoteControlUI/Private/UI/BaseLogicUI/SRCLogicPanelListBase.cpp
--- a/Engine/Plugins/VirtualProduction/RemoteControl/Source/RemoteControlUI/Private/UI/BaseLogicUI/SRCLogicPanelListBase.cpp
+++ b/Engine/Plugins/VirtualProduction/RemoteControl/Source/RemoteControlUI/Private/UI/BaseLogicUI/SRCLogicPanelListBase.cpp
@@ -31,10 +31,8 @@ TSharedPtr<SWidget> SRCLogicPanelListBase::GetContextMenuWidget()
 
 	const FRemoteControlCommands& Commands = FRemoteControlCommands::Get();
 
-	// 1. Copy
 	MenuBuilder.AddMenuEntry(Commands.CopyItem);
 
-	// 2. Paste
 	FText PasteItemLabel = LOCTEXT("Paste", "Paste");
 	if (TSharedPtr<SRemoteControlPanel> RemoteControlPanel = RemoteControlPanelWeakPtr.Pin())
 	{
@@ -48,23 +46,14 @@ TSharedPtr<SWidget> SRCLogicPanelListBase::GetContextMenuWidget()
 		}
 	}
 	MenuBuilder.AddMenuEntry(Commands.PasteItem, NAME_None, PasteItemLabel);
-
-	// 2. Duplicate
 	MenuBuilder.AddMenuEntry(Commands.DuplicateItem);
-
-	// 3. Update
 	MenuBuilder.AddMenuEntry(Commands.UpdateValue);
-
-	// 4. Delete
 	MenuBuilder.AddMenuEntry(Commands.DeleteEntity);
-
-	// 5. Delete All
-	FUIAction Action(FExecuteAction::CreateSP(this, &SRCLogicPanelListBase::RequestDeleteAllItems)
-	, FCanExecuteAction::CreateSP(this, &SRCLogicPanelListBase::CanDeleteAllItems));
-
 	MenuBuilder.AddMenuEntry(LOCTEXT("DeleteAll", "Delete All"),
 		LOCTEXT("ContextMenuEditTooltip", "Delete all the rows in this list"),
-		FSlateIcon(), Action);
+		FSlateIcon(),
+		FUIAction(FExecuteAction::CreateSP(this, &SRCLogicPanelListBase::RequestDeleteAllItems),
+			FCanExecuteAction::CreateSP(this, &SRCLogicPanelListBase::CanDeleteAllItems)));
 
 	MenuBuilder.EndSection();
 
@@ -78,7 +67,7 @@ bool SRCLogicPanelListBase::CanDeleteAllItems() const
 {
 	if (TSharedPtr<SRemoteControlPanel> RemoteControlPanel = RemoteControlPanelWeakPtr.Pin())
 	{
-		return !RemoteControlPanel->IsInLiveMode();
+		return !RemoteControlPanel->IsModeActive(ERCPanelMode::Live);
 	}
 
 	return false;
@@ -92,4 +81,4 @@ void SRCLogicPanelListBase::RequestDeleteAllItems()
 	}
 }
 
-#undef LOCTEXT_NAMESPACE
\ No newline at end of file
+#undef LOCTEXT_NAMESPACE
diff --git a/Engine/Plugins/VirtualProduction/RemoteControl/Source/RemoteControlUI/Private/UI/Drawers/SRCPanelDrawer.cpp b/Engine/Plugins/VirtualProduction/RemoteControl/Source/RemoteControlUI/Private/UI/Drawers/SRCPanelDrawer.cpp
deleted file mode 100644
diff --git a/Engine/Plugins/VirtualProduction/RemoteControl/Source/RemoteControlUI/Private/UI/Drawers/SRCPanelDrawer.h b/Engine/Plugins/VirtualProduction/RemoteControl/Source/RemoteControlUI/Private/UI/Drawers/SRCPanelDrawer.h
deleted file mode 100644
diff --git a/Engine/Plugins/VirtualProduction/RemoteControl/Source/RemoteControlUI/Private/UI/SRCPanelExposedEntitiesList.cpp b/Engine/Plugins/VirtualProduction/RemoteControl/Source/RemoteControlUI/Private/UI/SRCPanelExposedEntitiesList.cpp
--- a/Engine/Plugins/VirtualProduction/RemoteControl/Source/RemoteControlUI/Private/UI/SRCPanelExposedEntitiesList.cpp
+++ b/Engine/Plugins/VirtualProduction/RemoteControl/Source/RemoteControlUI/Private/UI/SRCPanelExposedEntitiesList.cpp
@@ -740,9 +740,9 @@ FText SRCPanelExposedEntitiesList::HandleEntityListHeaderLabel() const
 {
 	const FRemoteControlCommands& Commands = FRemoteControlCommands::Get();
 
-	if (bIsInProtocolsMode.Get() && Commands.ToggleProtocolMappings.IsValid())
+	if (bIsInProtocolsMode.Get() && Commands.ActivateProtocolsMode.IsValid())
 	{
-		return Commands.ToggleProtocolMappings->GetLabel();
+		return Commands.ActivateProtocolsMode->GetLabel();
 	}
 
 	return LOCTEXT("PropertiesLabel", "Properties");
diff --git a/Engine/Plugins/VirtualProduction/RemoteControl/Source/RemoteControlUI/Private/UI/SRemoteControlPanel.cpp b/Engine/Plugins/VirtualProduction/RemoteControl/Source/RemoteControlUI/Private/UI/SRemoteControlPanel.cpp
--- a/Engine/Plugins/VirtualProduction/RemoteControl/Source/RemoteControlUI/Private/UI/SRemoteControlPanel.cpp
+++ b/Engine/Plugins/VirtualProduction/RemoteControl/Source/RemoteControlUI/Private/UI/SRemoteControlPanel.cpp
@@ -1,7 +1,6 @@
 // Copyright Epic Games, Inc. All Rights Reserved.
 
 #include "SRemoteControlPanel.h"
-
 #include "Action/SRCActionPanel.h"
 #include "ActorEditorUtils.h"
 #include "Behaviour/Builtin/Bind/RCBehaviourBind.h"
@@ -63,10 +62,10 @@
 #include "Templates/SubclassOf.h"
 #include "Templates/UnrealTypeTraits.h"
 #include "ToolMenus.h"
+#include "Toolkits/GlobalEditorCommonCommands.h"
 #include "Toolkits/IToolkitHost.h"
 #include "UI/BaseLogicUI/RCLogicModeBase.h"
 #include "UI/BaseLogicUI/SRCLogicPanelListBase.h"
-#include "UI/Drawers/SRCPanelDrawer.h"
 #include "UI/Filters/SRCPanelFilter.h"
 #include "UI/Panels/SRCDockPanel.h"
 #include "Widgets/DeclarativeSyntaxSupport.h"
@@ -155,7 +154,7 @@ public:
 	{
 		if (const TSharedPtr<SRemoteControlPanel> PinnedOwner = Owner.Pin())
 		{
-			if (PinnedOwner->IsInLiveMode())
+			if (PinnedOwner->IsModeActive(ERCPanelMode::Live))
 			{
 				return false;
 			}
@@ -187,163 +186,6 @@ private:
 	TWeakPtr<SRemoteControlPanel> Owner;
 };
 
-/**
- * UI representation of a auto resizing button.
- */
-class SAutoResizeButton : public SCompoundWidget
-{
-	SLATE_BEGIN_ARGS(SAutoResizeButton)
-		: _ForceSmallIcons(false)
-		, _UICommand(nullptr)
-		, _IconOverride()
-	{}
-
-		SLATE_ATTRIBUTE(bool, ForceSmallIcons)
-
-		/** UI Command must be mapped to MainFrame Command List. */
-		SLATE_ARGUMENT(TSharedPtr<const FUICommandInfo>, UICommand)
-
-		SLATE_ARGUMENT(TAttribute<FSlateIcon>, IconOverride)
-
-	SLATE_END_ARGS()
-
-public:
-
-	void Construct(const FArguments& InArgs)
-	{
-		bForceSmallIcons = InArgs._ForceSmallIcons;
-
-		UICommand = InArgs._UICommand;
-
-		IconOverride = InArgs._IconOverride;
-
-		// Mimic Toolbar button style
-		const FToolBarStyle& ToolBarStyle = FAppStyle::Get().GetWidgetStyle<FToolBarStyle>("AssetEditorToolbar");
-
-		// Get the label & tooltip from the UI Command.
-		const TAttribute<FText> ActualLabel = UICommand.IsValid() ? UICommand->GetLabel() : FText::GetEmpty();
-		const TAttribute<FText> ActualToolTip = UICommand.IsValid() ? UICommand->GetDescription() : FText::GetEmpty();
-
-		// If we were supplied an image than go ahead and use that, otherwise we use a null widget
-		TSharedRef<SLayeredImage> IconWidget = SNew(SLayeredImage)
-			.ColorAndOpacity(this, &SAutoResizeButton::GetIconForegroundColor)
-			.Visibility(EVisibility::HitTestInvisible)
-			.Image(this, &SAutoResizeButton::GetIconBrush);
-
-		ChildSlot
-		.Padding(ToolBarStyle.ButtonPadding.Left, 0.f, ToolBarStyle.ButtonPadding.Right, 0.f)
-			[
-				SNew(SCheckBox)
-				.Padding(ToolBarStyle.CheckBoxPadding)
-				.Style(&ToolBarStyle.ToggleButton)
-				.IsChecked(this , &SAutoResizeButton::HandleIsChecked)
-				.OnCheckStateChanged(this, &SAutoResizeButton::OnCheckStateChanged)
-				.ToolTipText(ActualToolTip)
-				.Content()
-				[
-					SNew(SHorizontalBox)
-
-					// Icon Widget
-					+ SHorizontalBox::Slot()
-					.AutoWidth()
-					.VAlign(VAlign_Center)
-					.HAlign(HAlign_Center)
-					.Padding(2.f, 4.f)
-					[
-						IconWidget
-					]
-
-					// Label Text
-					+ SHorizontalBox::Slot()
-					.AutoWidth()
-					.Padding(ToolBarStyle.LabelPadding)
-					.VAlign(VAlign_Center)
-					[
-						SNew(STextBlock)
-						.Visibility(this, &SAutoResizeButton::GetLabelVisibility)
-						.Text(ActualLabel)
-						.TextStyle(&ToolBarStyle.LabelStyle)	// Smaller font for tool tip labels
-					]
-				]
-			];
-	}
-
-private:
-
-	/** Called by Slate to determine whether labels are visible */
-	EVisibility GetLabelVisibility() const
-	{
-		const bool bUseSmallIcons = bForceSmallIcons.IsSet() ? bForceSmallIcons.Get() : false;
-
-		return bUseSmallIcons ? EVisibility::Collapsed : EVisibility::Visible;
-	}
-
-	/** Gets the icon brush for the toolbar block widget */
-	const FSlateBrush* GetIconBrush() const
-	{
-		const FSlateIcon ActionIcon = UICommand.IsValid() ?  UICommand->GetIcon() : FSlateIcon();
-		const FSlateIcon& ActualIcon = IconOverride.IsSet() ? IconOverride.Get() : ActionIcon;
-
-		return ActualIcon.GetIcon();
-	}
-
-	/** Retrieves the color used by icon brush. */
-	FSlateColor GetIconForegroundColor() const
-	{
-		// If any brush has a tint, don't assume it should be subdued
-		const FSlateBrush* Brush = GetIconBrush();
-		if (Brush && Brush->TintColor != FLinearColor::White)
-		{
-			return FLinearColor::White;
-		}
-
-		return FSlateColor::UseForeground();
-	}
-
-	/**
-	 * Called by Slate to check whether this toolbar button is checked or not.
-	 */
-	ECheckBoxState HandleIsChecked() const
-	{
-		IMainFrameModule& MainFrame = FModuleManager::Get().LoadModuleChecked<IMainFrameModule>("MainFrame");
-
-		TSharedPtr<FUICommandList> ActionList = MainFrame.GetMainFrameCommandBindings();
-
-		if (ActionList.IsValid() && UICommand.IsValid())
-		{
-			return ActionList->GetCheckState(UICommand.ToSharedRef());
-		}
-
-		return ECheckBoxState::Undetermined;
-	}
-
-	/**
-	 * Called by Slate when this tool bar button's button is clicked
-	 */
-	void OnCheckStateChanged(ECheckBoxState NewState)
-	{
-		IMainFrameModule& MainFrame = FModuleManager::Get().LoadModuleChecked<IMainFrameModule>("MainFrame");
-
-		TSharedPtr<FUICommandList> ActionList = MainFrame.GetMainFrameCommandBindings();
-
-		if (ActionList.IsValid() && UICommand.IsValid())
-		{
-			ActionList->TryExecuteAction(UICommand.ToSharedRef());
-		}
-	}
-
-private:
-
-	/** Should we use small icons or not. */
-	TAttribute<bool> bForceSmallIcons;
-
-	/** Holds the UI Command information of this button. */
-	TSharedPtr<const FUICommandInfo> UICommand;
-
-	/** Overriden icon to be used instead of default one. */
-	TAttribute<FSlateIcon> IconOverride;
-};
-
 void SRemoteControlPanel::Construct(const FArguments& InArgs, URemoteControlPreset* InPreset, TSharedPtr<IToolkitHost> InToolkitHost)
 {
 	OnLiveModeChange = InArgs._OnLiveModeChange;
@@ -351,12 +193,9 @@ void SRemoteControlPanel::Construct(const FArguments& InArgs, URemoteControlPres
 	WidgetRegistry = MakeShared<FRCPanelWidgetRegistry>();
 	ToolkitHost = InToolkitHost;
 
-	ActivePanel = ERCPanels::RCP_Properties;
+	ActiveMode = ERCPanelMode::Controller;
 
 	RCPanelStyle = &FRemoteControlPanelStyle::Get()->GetWidgetStyle<FRCPanelStyle>("RemoteControlPanel.MinorPanel");
-	
-	const URemoteControlSettings* RemoteControlSettings = GetMutableDefault<URemoteControlSettings>();
-	bIsLogicPanelEnabled = RemoteControlSettings->bLogicPanelVisibility;
 
 	TArray<TSharedRef<SWidget>> ExtensionWidgets;
 	FRemoteControlUIModule::Get().GetExtensionGenerators().Broadcast(ExtensionWidgets);
@@ -369,22 +208,22 @@ void SRemoteControlPanel::Construct(const FArguments& InArgs, URemoteControlPres
 
 	// Settings
 	AddToolbarWidget(SNew(SButton)
-			.ButtonStyle(&RCPanelStyle->FlatButtonStyle)
-			.ContentPadding(2.0f)
-			.TextStyle(FRemoteControlPanelStyle::Get(), "RemoteControlPanel.Button.TextStyle")
-			.OnClicked(this, &SRemoteControlPanel::OnClickSettingsButton)
-			.ToolTipText(LOCTEXT("OpenRemoteControlSettings", "Open Remote Control settings."))
-			[
-				SNew(SImage)
-				.Image(FAppStyle::Get().GetBrush("Icons.Toolbar.Settings"))
-				.ColorAndOpacity(FSlateColor::UseForeground())
-			]);
+		.ButtonStyle(&RCPanelStyle->FlatButtonStyle)
+		.ContentPadding(2.0f)
+		.TextStyle(FRemoteControlPanelStyle::Get(), "RemoteControlPanel.Button.TextStyle")
+		.OnClicked(this, &SRemoteControlPanel::OnClickSettingsButton)
+		.ToolTipText(LOCTEXT("OpenRemoteControlSettings", "Open Remote Control settings."))
+		[
+			SNew(SImage)
+			.Image(FAppStyle::Get().GetBrush("Icons.Toolbar.Settings"))
+			.ColorAndOpacity(FSlateColor::UseForeground())
+		]);
 
 	// Show Log
 	AddToolbarWidget(SNew(SCheckBox)
 		.Style(&RCPanelStyle->ToggleButtonStyle)
 		.ToolTipText(LOCTEXT("ShowLogTooltip", "Show/Hide remote control log."))
-		.IsChecked_Lambda([this]() { return (FRemoteControlLogger::Get().IsEnabled() && !bIsInLiveMode) ? ECheckBoxState::Checked : ECheckBoxState::Unchecked; })
+		.IsChecked_Lambda([]() { return FRemoteControlLogger::Get().IsEnabled() ? ECheckBoxState::Checked : ECheckBoxState::Unchecked; })
 		.OnCheckStateChanged(this, &SRemoteControlPanel::OnLogCheckboxToggle)
 		.Padding(4.f)
 		[
@@ -409,15 +248,14 @@ void SRemoteControlPanel::Construct(const FArguments& InArgs, URemoteControlPres
 	EntityList = SNew(SRCPanelExposedEntitiesList, Preset.Get(), WidgetRegistry)
 		.ExposeActorsComboButton(CreateExposeActorsButton())
 		.ExposeFunctionsComboButton(CreateExposeFunctionsButton())
-		.OnEntityListUpdated_Lambda([this] ()
+		.LiveMode(this, &SRemoteControlPanel::IsModeActive, ERCPanelMode::Live)
+		.ProtocolsMode(this, &SRemoteControlPanel::IsModeActive, ERCPanelMode::Protocols)
+		.OnEntityListUpdated_Lambda([this]()
 			{
 				UpdateEntityDetailsView(EntityList->GetSelectedEntity());
 				UpdateRebindButtonVisibility();
 				CachedExposedPropertyArgs.Reset();
-			}
-		)
-		.LiveMode_Lambda([this]() {return bIsInLiveMode; })
-		.ProtocolsMode_Lambda([this]() { return IsInProtocolsMode() && ActivePanel == ERCPanels::RCP_Protocols; });
+			});
 
 	EntityList->OnSelectionChange().AddSP(this, &SRemoteControlPanel::UpdateEntityDetailsView);
 
@@ -487,24 +325,20 @@ void SRemoteControlPanel::Construct(const FArguments& InArgs, URemoteControlPres
 		.EnableHeader(false)
 		.EnableFooter(false);
 
-	ControllerPanel = SNew(SRCControllerPanel, SharedThis(this))
-		.LiveMode_Lambda([this] { return bIsInLiveMode; })
-		.Visibility_Lambda([this] { return (bIsLogicPanelEnabled || bIsInLiveMode) && (ActivePanel == ERCPanels::RCP_Properties || ActivePanel == ERCPanels::RCP_Live || ActivePanel == ERCPanels::RCP_None) ? EVisibility::Visible : EVisibility::Collapsed; });
+	TSharedRef<SRemoteControlPanel> This = SharedThis(this);
 
-	BehaviourPanel = SNew(SRCBehaviourPanel, SharedThis(this))
-		.Visibility_Lambda([this] { return !bIsInLiveMode && bIsLogicPanelEnabled && (ActivePanel == ERCPanels::RCP_Properties || ActivePanel == ERCPanels::RCP_None) ? EVisibility::Visible : EVisibility::Collapsed; });
-
-	SignaturePanel = SNew(SRCSignaturePanel, SharedThis(this))
-		.LiveMode_Lambda([this] { return bIsInLiveMode; });
+	// Controller Mode
+	ActionPanel = SNew(SRCActionPanel, This);
+	BehaviorPanel = SNew(SRCBehaviourPanel, This);
+	ControllerPanel = SNew(SRCControllerPanel, This)
+		.LiveMode(this, &SRemoteControlPanel::IsModeActive, ERCPanelMode::Live);
 
 	ControllersAndBehavioursPanel->AddPanel(ControllerPanel.ToSharedRef(), 0.8f);
+	ControllersAndBehavioursPanel->AddPanel(BehaviorPanel.ToSharedRef(), 0.f, /*bResizable*/false);
 
-	constexpr bool bResizable = false;
-	ControllersAndBehavioursPanel->AddPanel(BehaviourPanel.ToSharedRef(), 0.f, bResizable);
-
-	// Actions Panel.
-	ActionPanel = SNew(SRCActionPanel, SharedThis(this))
-		.Visibility_Lambda([this] { return !bIsInLiveMode && bIsLogicPanelEnabled && (ActivePanel == ERCPanels::RCP_Properties || ActivePanel == ERCPanels::RCP_None) ? EVisibility::Visible : EVisibility::Collapsed; });
+	// Signature Mode
+	SignaturePanel = SNew(SRCSignaturePanel, This)
+		.LiveMode(this, &SRemoteControlPanel::IsModeActive, ERCPanelMode::Live);
 
 	// Retrieve Action Panel Split ratio from RC Settings.
 	// We can't just get the value, since it will update in real time as the user resizes the slot: this value needs to update in real time
@@ -514,7 +348,7 @@ void SRemoteControlPanel::Construct(const FArguments& InArgs, URemoteControlPres
 			const URemoteControlSettings* Settings = GetMutableDefault<URemoteControlSettings>();
 			return Settings->ActionPanelSplitRatio;
 		}));
-	
+
 	int32 ActionPanelSlotIndex = LogicPanel->AddPanel(ControllersAndBehavioursPanel, ActionPanelSplitRatioAttribute, true);
 	LogicPanel->AddPanel(ActionPanel.ToSharedRef(), 1.0 - ActionPanelSplitRatioAttribute.Get(), true);
 
@@ -541,204 +375,37 @@ void SRemoteControlPanel::Construct(const FArguments& InArgs, URemoteControlPres
 	}
 
 	// Make 2 Columns with Panel Drawer + Main Panel
-	TSharedRef<SSplitter> ContentPanel = SNew(SSplitter)
-		.Orientation(Orient_Horizontal)
-		.HitDetectionSplitterHandleSize(RCPanelStyle->SplitterHandleSize)
-		.PhysicalSplitterHandleSize(RCPanelStyle->SplitterHandleSize)
-
-		// Panel Drawer
-		+SSplitter::Slot()
-		.Value(0.05f)
-		.Resizable(false)
-		.SizeRule(SSplitter::ESizeRule::SizeToContent)
+	TSharedRef<SWidgetSwitcher> ContentPanel = SNew(SWidgetSwitcher)
+		.WidgetIndex_Lambda([this](){ return static_cast<int32>(ActiveMode); })
+		// 1. Logic Mode (Index: 0)
+		+ SWidgetSwitcher::Slot()
 		[
-			SAssignNew(PanelDrawer, SRCPanelDrawer)
+			BuildLogicModeContent(LogicPanel)
 		]
-
-		// Main Panel
-		+SSplitter::Slot()
-		.Value(0.35f)
+		// 2. Entity Details Mode (Index: 1)
+		+ SWidgetSwitcher::Slot()
 		[
-			SNew(SWidgetSwitcher)
-			.WidgetIndex_Lambda([this](){ return (int32)ActivePanel; })
-
-			// 1. Nothing to be shown as all panels might be hidden (Index : 0)
-			// Since this is widget switcher we should show NullWidget in such case.
-			+ SWidgetSwitcher::Slot()
-			[
-				SNew(SWidgetSwitcher)
-				.WidgetIndex_Lambda([this]() { return !bIsLogicPanelEnabled ? 0 : 1; })
-
-				+ SWidgetSwitcher::Slot()
-				[
-					SNullWidget::NullWidget
-				]
-
-				+ SWidgetSwitcher::Slot()
-				[
-					LogicPanel
-				]
-			]
-
-			// 2. Exposed entities List with Logic Editor (Index : 1).
-			+ SWidgetSwitcher::Slot()
-			[
-				SNew(SWidgetSwitcher)
-				.WidgetIndex_Lambda([this]() { return !bIsLogicPanelEnabled ? 0 : 1; })
-
-				+ SWidgetSwitcher::Slot()
-				[
-					SNew(SWidgetSwitcher)
-					.WidgetIndex_Lambda([this]() { return !bIsInLiveMode ? 0 : 1; })
-
-					+ SWidgetSwitcher::Slot()
-                    [
-						EntityList.ToSharedRef()
-					]
-
-					+ SWidgetSwitcher::Slot()
-					[
-						SNullWidget::NullWidget
-					]
-				]
-
-				+ SWidgetSwitcher::Slot()
-				[
-					SNew(SWidgetSwitcher)
-					.WidgetIndex_Lambda([this]() { return bIsInLiveMode ? 0 : 1; })
-
-					+ SWidgetSwitcher::Slot()
-					[
-						LogicPanel
-					]
-
-					+ SWidgetSwitcher::Slot()
-					[
-						SNew(SSplitter)
-						.Orientation(Orient_Horizontal)
-
-						+SSplitter::Slot()
-						.Value(0.4)
-						[
-							EntityList.ToSharedRef()
-						]
-
-						+SSplitter::Slot()
-						.Value(0.6)
-						[
-							LogicPanel
-						]
-					]
-				]
-			]
-
-			// 3. Properties with Entity Details (Index : 2).
-			+ SWidgetSwitcher::Slot()
-			[
-				SNew(SSplitter)
-				.Orientation(EOrientation::Orient_Vertical)
-
-				// Exposed entities List
-				+ SSplitter::Slot()
-				.Value(TreeBindingSplitRatioTop)
-				.OnSlotResized(SSplitter::FOnSlotResized::CreateLambda([](float InNewSize)
-				{
-					URemoteControlSettings* Settings = GetMutableDefault<URemoteControlSettings>();
-					Settings->TreeBindingSplitRatio = InNewSize;
-					Settings->PostEditChange();
-					Settings->SaveConfig();
-				}))
-				[
-					EntityList.ToSharedRef()
-				]
-
-				// Entity Details
-				+ SSplitter::Slot()
-				.Value(TreeBindingSplitRatioBottom)
-				[
-					CreateEntityDetailsView()
-				]
-			]
-
-			// 4. Properties with Protocols (Index : 3).
-			+ SWidgetSwitcher::Slot()
-			[
-				SNew(SSplitter)
-				.Orientation(EOrientation::Orient_Vertical)
-
-				// Exposed entities List
-				+ SSplitter::Slot()
-				.Value(TreeBindingSplitRatioTop)
-				.OnSlotResized(SSplitter::FOnSlotResized::CreateLambda([](float InNewSize)
-				{
-					URemoteControlSettings* Settings = GetMutableDefault<URemoteControlSettings>();
-					Settings->TreeBindingSplitRatio = InNewSize;
-					Settings->PostEditChange();
-					Settings->SaveConfig();
-				}))
-				[
-					EntityList.ToSharedRef()
-				]
-
-				// Protocol Details
-				+ SSplitter::Slot()
-				.Value(TreeBindingSplitRatioBottom)
-				[
-					EntityProtocolDetails.ToSharedRef()
-				]
-			]
-
-			// 5. Output Log (Index : 4).
-			+ SWidgetSwitcher::Slot()
-			[
-				OutputLogDockPanel
-			]
-
-			// 6. Live (Index : 5).
-			+ SWidgetSwitcher::Slot()
-			[
-				SNew(SSplitter)
-				.Orientation(EOrientation::Orient_Horizontal)
-
-				// Exposed entities List
-				+ SSplitter::Slot()
-				.Value(0.5f)
-				[
-					EntityList.ToSharedRef()
-				]
-
-				// Logic Panel
-				+ SSplitter::Slot()
-				.Value(0.5f)
-				[
-					LogicPanel
-				]
-			]
-
-			// 7. (unused) Count (Index : 6)
-			+ SWidgetSwitcher::Slot()
-			[
-				SNullWidget::NullWidget
-			]
-
-			// 8. Signatures (Index : 7).
-			+ SWidgetSwitcher::Slot()
-			[
-				SNew(SSplitter)
-				.Orientation(Orient_Horizontal)
-
-				+SSplitter::Slot()
-				.Value(0.4)
-				[
-					EntityList.ToSharedRef()
-				]
-
-				+SSplitter::Slot()
-				.Value(0.6)
-				[
-					SignaturePanel.ToSharedRef()
-				]
-			]
+			BuildEntityDetailsModeContent(TreeBindingSplitRatioTop, TreeBindingSplitRatioBottom)
+		]
+		// 3. Protocols Mode (Index: 2)
+		+ SWidgetSwitcher::Slot()
+		[
+			BuildProtocolsModeContent(TreeBindingSplitRatioTop, TreeBindingSplitRatioBottom)
+		]
+		// 4. Output Log Mode (Index : 3)
+		+ SWidgetSwitcher::Slot()
+		[
+			OutputLogDockPanel
+		]
+		// 5. Live Mode (Index : 4)
+		+ SWidgetSwitcher::Slot()
+		[
+			BuildLiveModeContent(LogicPanel)
+		]
+		// 7. Signatures Mode (Index : 5)
+		+ SWidgetSwitcher::Slot()
+		[
+			BuildSignaturesModeContent()
 		];
 
 	ChildSlot
@@ -784,12 +451,10 @@ void SRemoteControlPanel::Construct(const FArguments& InArgs, URemoteControlPres
 			]
 		];
 
-
 	InputProcessor = MakeShared<FRemoteControlPanelInputProcessor>(SharedThis(this));
 	FSlateApplication::Get().RegisterInputPreProcessor(InputProcessor);
 
 	RegisterEvents();
-	RegisterPanels();
 	CacheLevelClasses();
 	Refresh();
 	LoadSettings(InPreset->GetPresetId());
@@ -798,8 +463,6 @@ void SRemoteControlPanel::Construct(const FArguments& InArgs, URemoteControlPres
 SRemoteControlPanel::~SRemoteControlPanel()
 {
 	SaveSettings();
-
-	UnregisterPanels();
 	UnregisterEvents();
 
 	FSlateApplication::Get().UnregisterInputPreProcessor(InputProcessor);
@@ -982,6 +645,40 @@ FGuid SRemoteControlPanel::GetSelectedGroup() const
 	return FGuid();
 }
 
+bool SRemoteControlPanel::CanActivateMode(ERCPanelMode InPanelMode) const
+{
+	if (InPanelMode == ERCPanelMode::OutputLog)
+	{
+		return FRemoteControlLogger::Get().IsEnabled();
+	}
+	return true;
+}
+
+bool SRemoteControlPanel::IsModeActive(ERCPanelMode InPanelMode) const
+{
+	return ActiveMode == InPanelMode;
+}
+
+void SRemoteControlPanel::SetActiveMode(ERCPanelMode InPanelMode)
+{
+	if (InPanelMode == ActiveMode)
+	{
+		return;
+	}
+
+	ActiveMode = InPanelMode;
+
+	if (EntityList.IsValid())
+	{
+		EEntitiesListMode EntitiesListMode = EEntitiesListMode::Default;
+		if (IsModeActive(ERCPanelMode::Protocols))
+		{
+			EntitiesListMode = EEntitiesListMode::Protocols;
+		}
+		EntityList->RebuildListWithColumns(EntitiesListMode);
+	}
+}
+
 FReply SRemoteControlPanel::OnClickDisableUseLessCPU() const
 {
 	UEditorPerformanceSettings* Settings = GetMutableDefault<UEditorPerformanceSettings>();
@@ -1134,7 +831,7 @@ TSharedRef<SWidget> SRemoteControlPanel::CreateExposeFunctionsButton()
 
 	return SAssignNew(ExposeFunctionsComboButton, SComboButton)
 		.AddMetaData<FTagMetaData>(FTagMetaData(TEXT("Expose Functions")))
-		.IsEnabled_Lambda([this]() { return !this->bIsInLiveMode; })
+		.IsEnabled_Lambda([this]() { return !IsModeActive(ERCPanelMode::Live); })
 		.HAlign(HAlign_Center)
 		.VAlign(VAlign_Center)
 		.ButtonStyle(&RCPanelStyle->FlatButtonStyle)
@@ -1233,7 +930,7 @@ TSharedRef<SWidget> SRemoteControlPanel::CreateExposeActorsButton()
 
 	return SAssignNew(ExposeActorsComboButton, SComboButton)
 		.AddMetaData<FTagMetaData>(FTagMetaData(TEXT("Expose Actors")))
-		.IsEnabled_Lambda([this]() { return !this->bIsInLiveMode; })
+		.IsEnabled_Lambda([this]() { return !IsModeActive(ERCPanelMode::Live); })
 		.HAlign(HAlign_Center)
 		.VAlign(VAlign_Center)
 		.ButtonStyle(&RCPanelStyle->FlatButtonStyle)
@@ -1407,6 +1104,115 @@ void SRemoteControlPanel::OnMapChange(uint32)
 	Refresh();
 }
 
+TSharedRef<SWidget> SRemoteControlPanel::BuildLogicModeContent(const TSharedRef<SWidget>& InLogicPanel)
+{
+	return SNew(SWidgetSwitcher)
+		.WidgetIndex_Lambda([this]() { return IsModeActive(ERCPanelMode::Live) ? 0 : 1; })
+		+ SWidgetSwitcher::Slot()
+		[
+			InLogicPanel
+		]
+		+ SWidgetSwitcher::Slot()
+		[
+			SNew(SSplitter)
+			.Orientation(Orient_Horizontal)
+			+SSplitter::Slot()
+			.Value(0.4)
+			[
+				EntityList.ToSharedRef()
+			]
+			+SSplitter::Slot()
+			.Value(0.6)
+			[
+				InLogicPanel
+			]
+		];
+}
+
+TSharedRef<SWidget> SRemoteControlPanel::BuildEntityDetailsModeContent(const TAttribute<float>& InRatioTop, const TAttribute<float>& InRatioBottom)
+{
+	return SNew(SSplitter)
+		.Orientation(EOrientation::Orient_Vertical)
+		// Exposed entities List
+		+ SSplitter::Slot()
+		.Value(InRatioTop)
+		.OnSlotResized(SSplitter::FOnSlotResized::CreateLambda([](float InNewSize)
+		{
+			URemoteControlSettings* Settings = GetMutableDefault<URemoteControlSettings>();
+			Settings->TreeBindingSplitRatio = InNewSize;
+			Settings->PostEditChange();
+			Settings->SaveConfig();
+		}))
+		[
+			EntityList.ToSharedRef()
+		]
+		// Entity Details
+		+ SSplitter::Slot()
+		.Value(InRatioBottom)
+		[
+			CreateEntityDetailsView()
+		];
+}
+
+TSharedRef<SWidget> SRemoteControlPanel::BuildProtocolsModeContent(const TAttribute<float>& InRatioTop, const TAttribute<float>& InRatioBottom)
+{
+	return SNew(SSplitter)
+		.Orientation(EOrientation::Orient_Vertical)
+		// Exposed entities List
+		+ SSplitter::Slot()
+		.Value(InRatioTop)
+		.OnSlotResized(SSplitter::FOnSlotResized::CreateLambda([](float InNewSize)
+		{
+			URemoteControlSettings* Settings = GetMutableDefault<URemoteControlSettings>();
+			Settings->TreeBindingSplitRatio = InNewSize;
+			Settings->PostEditChange();
+			Settings->SaveConfig();
+		}))
+		[
+			EntityList.ToSharedRef()
+		]
+		// Protocol Details
+		+ SSplitter::Slot()
+		.Value(InRatioBottom)
+		[
+			EntityProtocolDetails.ToSharedRef()
+		];
+}
+
+TSharedRef<SWidget> SRemoteControlPanel::BuildLiveModeContent(const TSharedRef<SWidget>& InLogicPanel)
+{
+	return SNew(SSplitter)
+		.Orientation(EOrientation::Orient_Horizontal)
+		// Exposed entities List
+		+ SSplitter::Slot()
+		.Value(0.5f)
+		[
+			EntityList.ToSharedRef()
+		]
+		// Logic Panel
+		+ SSplitter::Slot()
+		.Value(0.5f)
+		[
+			InLogicPanel
+		];
+}
+
+TSharedRef<SWidget> SRemoteControlPanel::BuildSignaturesModeContent()
+{
+	return SNew(SSplitter)
+		.Orientation(Orient_Horizontal)
+		+SSplitter::Slot()
+		.Value(0.4)
+		[
+			EntityList.ToSharedRef()
+		]
+		+SSplitter::Slot()
+		.Value(0.6)
+		[
+			SignaturePanel.ToSharedRef()
+		];
+}
+
 void SRemoteControlPanel::BindRemoteControlCommands()
 {
 	const FRemoteControlCommands& Commands = FRemoteControlCommands::Get();
@@ -1417,36 +1223,31 @@ void SRemoteControlPanel::BindRemoteControlCommands()
 
 	ActionList.MapAction(
 		Commands.SavePreset,
-		FExecuteAction::CreateSP(this, &SRemoteControlPanel::SaveAsset_Execute),
+		FExecuteAction::CreateSP(this, &SRemoteControlPanel::SaveAsset),
 		FCanExecuteAction(),
 		FGetActionCheckState(),
 		FIsActionButtonVisible::CreateSP(this, &SRemoteControlPanel::CanSaveAsset));
 
 	ActionList.MapAction(
-		Commands.FindPresetInContentBrowser,
-		FExecuteAction::CreateSP(this, &SRemoteControlPanel::FindInContentBrowser_Execute),
+		FGlobalEditorCommonCommands::Get().FindInContentBrowser,
+		FExecuteAction::CreateSP(this, &SRemoteControlPanel::FindInContentBrowser),
 		FCanExecuteAction::CreateSP(this, &SRemoteControlPanel::CanFindInContentBrowser));
 
-	ActionList.MapAction(
-		Commands.ToggleProtocolMappings,
-		FExecuteAction::CreateSP(this, &SRemoteControlPanel::ToggleProtocolMappings_Execute),
-		FCanExecuteAction::CreateSP(this, &SRemoteControlPanel::CanToggleProtocolsMode),
-		FIsActionChecked::CreateSP(this, &SRemoteControlPanel::IsInProtocolsMode),
-		FIsActionButtonVisible::CreateSP(this, &SRemoteControlPanel::CanToggleProtocolsMode));
+	auto MapModeAction = [&ActionList, this](const TSharedPtr<FUICommandInfo>& InCommand, ERCPanelMode InMode)
+		{
+			ActionList.MapAction(
+				InCommand,
+				FExecuteAction::CreateSP(this, &SRemoteControlPanel::SetActiveMode, InMode),
+				FCanExecuteAction::CreateSP(this, &SRemoteControlPanel::CanActivateMode, InMode),
+				FIsActionChecked::CreateSP(this, &SRemoteControlPanel::IsModeActive, InMode),
+				FIsActionButtonVisible::CreateSP(this, &SRemoteControlPanel::CanActivateMode, InMode));
+		};
 
-	ActionList.MapAction(
-		Commands.ToggleLogicEditor,
-		FExecuteAction::CreateSP(this, &SRemoteControlPanel::ToggleLogicEditor_Execute),
-		FCanExecuteAction::CreateSP(this, &SRemoteControlPanel::CanToggleLogicPanel),
-		FIsActionChecked::CreateSP(this, &SRemoteControlPanel::IsLogicPanelEnabled),
-		FIsActionButtonVisible::CreateSP(this, &SRemoteControlPanel::CanToggleLogicPanel));
-
-	ActionList.MapAction(
-		Commands.ToggleSignatureEditor,
-		FExecuteAction::CreateSP(this, &SRemoteControlPanel::ToggleSignatureEditor_Execute),
-		FCanExecuteAction::CreateSP(this, &SRemoteControlPanel::CanToggleSignaturePanel),
-		FIsActionChecked::CreateSP(this, &SRemoteControlPanel::IsSignaturePanelEnabled),
-		FIsActionButtonVisible::CreateSP(this, &SRemoteControlPanel::CanToggleSignaturePanel));
+	MapModeAction(Commands.ActivateLogicMode    , ERCPanelMode::Controller);
+	MapModeAction(Commands.ActivateDetailsMode  , ERCPanelMode::EntityDetails);
+	MapModeAction(Commands.ActivateSignatureMode, ERCPanelMode::Signature);
+	MapModeAction(Commands.ActivateProtocolsMode, ERCPanelMode::Protocols);
+	MapModeAction(Commands.ActivateOutputLogMode, ERCPanelMode::OutputLog);
 
 	ActionList.MapAction(
 		Commands.DeleteEntity,
@@ -1613,125 +1414,6 @@ void SRemoteControlPanel::UnregisterEvents()
 	FEditorDelegates::EndPIE.RemoveAll(this);
 }
 
-void SRemoteControlPanel::RegisterPanels()
-{
-	if (PanelDrawer.IsValid())
-	{
-		PanelDrawer->OnRCPanelToggled().BindSP(this, &SRemoteControlPanel::OnRCPanelToggled);
-
-		PanelDrawer->CanToggleRCPanel().BindLambda([this]()
-			{
-				return !bIsInLiveMode;
-			}
-		);
-
-		{// Properties Panel
-			TSharedRef<FRCPanelDrawerArgs> PropertiesPanel = MakeShared<FRCPanelDrawerArgs>(ERCPanels::RCP_Properties);
-
-			PropertiesPanel->bDrawnByDefault = true;
-			PropertiesPanel->bRotateIconBy90 = true;
-			PropertiesPanel->Label = LOCTEXT("PropertiesPanelLabel", "Expose");
-			PropertiesPanel->ToolTip = LOCTEXT("PropertiesPanelTooltip", "Open exposed properties panel.");
-			PropertiesPanel->Icon = FSlateIcon(FAppStyle::GetAppStyleSetName(), "EditorPreferences.TabIcon");
-
-			PanelDrawer->RegisterPanel(PropertiesPanel);
-
-			RegisteredDrawers.Add(PropertiesPanel->GetPanelID(), PropertiesPanel);
-		}
-
-		{// Signature Panel
-			TSharedRef<FRCPanelDrawerArgs> SignaturePanelDrawer = MakeShared<FRCPanelDrawerArgs>(ERCPanels::RCP_Signature);
-
-			SignaturePanelDrawer->bDrawnByDefault = true;
-			SignaturePanelDrawer->bRotateIconBy90 = false;
-			SignaturePanelDrawer->DrawerVisibility = bIsSignaturePanelEnabled ? EVisibility::Visible : EVisibility::Collapsed;
-			SignaturePanelDrawer->Label = LOCTEXT("SignaturePanelLabel", "Signature");
-			SignaturePanelDrawer->ToolTip = LOCTEXT("SignaturePanelTooltip", "Open Signature panel.");
-			SignaturePanelDrawer->Icon = FSlateIcon(FAppStyle::GetAppStyleSetName(), TEXT("GraphEditor.Function_16x"));
-
-			PanelDrawer->RegisterPanel(SignaturePanelDrawer);
-
-			RegisteredDrawers.Add(SignaturePanelDrawer->GetPanelID(), SignaturePanelDrawer);
-		}
-
-		{// Properties With Details Panel
-			TSharedRef<FRCPanelDrawerArgs> EntityDetailsPanel = MakeShared<FRCPanelDrawerArgs>(ERCPanels::RCP_EntityDetails);
-
-			EntityDetailsPanel->bDrawnByDefault = false;
-			EntityDetailsPanel->bRotateIconBy90 = false;
-			EntityDetailsPanel->Label = LOCTEXT("EntityDetailsPanelLabel", "Details");
-			EntityDetailsPanel->ToolTip = LOCTEXT("EntityDetailsPanelTooltip", "Open entity details panel.");
-			EntityDetailsPanel->Icon = FSlateIcon(FAppStyle::GetAppStyleSetName(), "LevelEditor.Tabs.Details");
-
-			PanelDrawer->RegisterPanel(EntityDetailsPanel);
-
-			RegisteredDrawers.Add(EntityDetailsPanel->GetPanelID(), EntityDetailsPanel);
-		}
-
-		{// Properties With Protocols Panel
-			TSharedRef<FRCPanelDrawerArgs> EntityProtocolsPanel = MakeShared<FRCPanelDrawerArgs>(ERCPanels::RCP_Protocols);
-			const FRemoteControlCommands& Commands = FRemoteControlCommands::Get();
-
-			EntityProtocolsPanel->bDrawnByDefault = true;
-			EntityProtocolsPanel->bRotateIconBy90 = false;
-			EntityProtocolsPanel->DrawerVisibility = bIsInProtocolsMode ? EVisibility::Visible : EVisibility::Collapsed;
-			EntityProtocolsPanel->Label = Commands.ToggleProtocolMappings->GetLabel();
-			EntityProtocolsPanel->ToolTip = LOCTEXT("EntityProtocolsPanelTooltip", "Open entity protocols panel.");
-			EntityProtocolsPanel->Icon = FSlateIcon(FAppStyle::GetAppStyleSetName(), "LevelEditor.Tabs.StatsViewer");
-
-			PanelDrawer->RegisterPanel(EntityProtocolsPanel);
-
-			RegisteredDrawers.Add(EntityProtocolsPanel->GetPanelID(), EntityProtocolsPanel);
-		}
-
-		{// Output Log Panel
-			TSharedRef<FRCPanelDrawerArgs> OutputLogPanel = MakeShared<FRCPanelDrawerArgs>(ERCPanels::RCP_OutputLog);
-
-			OutputLogPanel->bDrawnByDefault = true;
-			OutputLogPanel->bRotateIconBy90 = false;
-			OutputLogPanel->DrawerVisibility = FRemoteControlLogger::Get().IsEnabled() ? EVisibility::Visible : EVisibility::Collapsed;
-			OutputLogPanel->Label = LOCTEXT("OutputLogPanelPanelLabel", "Log");
-			OutputLogPanel->ToolTip = LOCTEXT("OutputLogPanelPanelTooltip", "Open output log panel.");
-			OutputLogPanel->Icon = FSlateIcon(FAppStyle::GetAppStyleSetName(), "MessageLog.TabIcon");
-
-			PanelDrawer->RegisterPanel(OutputLogPanel);
-
-			RegisteredDrawers.Add(OutputLogPanel->GetPanelID(), OutputLogPanel);
-		}
-
-		{// Live Panel
-			TSharedRef<FRCPanelDrawerArgs> LivePanel = MakeShared<FRCPanelDrawerArgs>(ERCPanels::RCP_Live);
-
-			LivePanel->bDrawnByDefault = true;
-			LivePanel->bRotateIconBy90 = false;
-			LivePanel->DrawerVisibility = bIsInLiveMode ? EVisibility::Visible : EVisibility::Collapsed;
-			LivePanel->Label = LOCTEXT("LivePanelPanelLabel", "Live");
-			LivePanel->ToolTip = LOCTEXT("LivePanelTooltip", "Open live panel.");
-			LivePanel->Icon = FSlateIcon(FAppStyle::GetAppStyleSetName(), "LevelEditor.Tabs.StatsViewer");
-
-			PanelDrawer->RegisterPanel(LivePanel);
-
-			RegisteredDrawers.Add(LivePanel->GetPanelID(), LivePanel);
-		}
-	}
-}
-
-void SRemoteControlPanel::UnregisterPanels()
-{
-	if (PanelDrawer.IsValid())
-	{
-		PanelDrawer->OnRCPanelToggled().Unbind();
-		PanelDrawer->CanToggleRCPanel().Unbind();
-
-		for (TMap<ERCPanels, TSharedRef<FRCPanelDrawerArgs>>::TIterator RegisteredDrawer = RegisteredDrawers.CreateIterator(); RegisteredDrawer; ++RegisteredDrawer)
-		{
-			PanelDrawer->UnregisterPanel(RegisteredDrawer->Value);
-
-			RegisteredDrawer.RemoveCurrent();
-		}
-	}
-}
-
 void SRemoteControlPanel::Refresh()
 {
 	TRACE_CPUPROFILER_EVENT_SCOPE(SRemoteControlPanel::Refresh);
@@ -1814,27 +1496,11 @@ void SRemoteControlPanel::Unexpose(const FRCExposesPropertyArgs& InPropertyArgs)
 	}
 }
 
-void SRemoteControlPanel::OnLogCheckboxToggle(ECheckBoxState State)
+void SRemoteControlPanel::OnLogCheckboxToggle(ECheckBoxState InState)
 {
-	const bool bIsLogEnabled = (State == ECheckBoxState::Checked) ? true : false;
-	FRemoteControlLogger::Get().EnableLog(bIsLogEnabled);
-
-	if (PanelDrawer.IsValid())
-	{
-		TSharedRef<FRCPanelDrawerArgs> OutputLogPanel = RegisteredDrawers.FindChecked(ERCPanels::RCP_OutputLog);
-
-		OutputLogPanel->DrawerVisibility = bIsLogEnabled ? EVisibility::Visible : EVisibility::Collapsed;
-
-		// When we are not enabling log collapse the drawer.
-		PanelDrawer->TogglePanel(OutputLogPanel, !bIsLogEnabled);
-
-		if (!bIsLogEnabled)
-		{
-			TSharedRef<FRCPanelDrawerArgs> PropertiesPanel = RegisteredDrawers.FindChecked(ERCPanels::RCP_Properties);
-
-			PanelDrawer->TogglePanel(PropertiesPanel);
-		}
-	}
+	const bool bEnableLog = InState == ECheckBoxState::Checked;
+	FRemoteControlLogger::Get().EnableLog(bEnableLog);
+	SetActiveMode(bEnableLog ? ERCPanelMode::OutputLog : ERCPanelMode::Controller);
 }
 
 void SRemoteControlPanel::OnBlueprintReinstanced()
@@ -2104,15 +1770,14 @@ void SRemoteControlPanel::RegisterDefaultToolBar()
 	UToolMenus* ToolMenus = UToolMenus::Get();
 	if (!ToolMenus->IsMenuRegistered(DefaultRemoteControlPanelToolBarName))
 	{
-		UToolMenu* ToolbarBuilder = ToolMenus->RegisterMenu(DefaultRemoteControlPanelToolBarName, NAME_None, EMultiBoxType::SlimHorizontalToolBar);
-		ToolbarBuilder->StyleName = "ContentBrowser.ToolBar";
+		UToolMenu* ToolbarMenu = ToolMenus->RegisterMenu(DefaultRemoteControlPanelToolBarName, NAME_None, EMultiBoxType::SlimHorizontalToolBar);
+		check(ToolbarMenu);
+		ToolbarMenu->StyleName = "AssetEditorToolbar";
 
-		{
-			FToolMenuSection& AssetSection = ToolbarBuilder->AddSection("Asset");
-			AssetSection.AddEntry(FToolMenuEntry::InitToolBarButton(FRemoteControlCommands::Get().SavePreset, TAttribute<FText>(), TAttribute<FText>(), FSlateIcon(FAppStyle::GetAppStyleSetName(), TEXT("Icons.Save"))));
-			AssetSection.AddEntry(FToolMenuEntry::InitToolBarButton(FRemoteControlCommands::Get().FindPresetInContentBrowser, LOCTEXT("FindInContentBrowserButton", "Browse"), TAttribute<FText>(), FSlateIcon(FAppStyle::GetAppStyleSetName(), TEXT("LevelEditor.OpenContentBrowser"))));
-			AssetSection.AddSeparator("Common");
-		}
+		FToolMenuSection& AssetSection = ToolbarMenu->AddSection("Asset");
+		AssetSection.AddEntry(FToolMenuEntry::InitToolBarButton(FRemoteControlCommands::Get().SavePreset, TAttribute<FText>(), TAttribute<FText>(), FSlateIcon(FAppStyle::GetAppStyleSetName(), TEXT("Icons.Save"))));
+		AssetSection.AddEntry(FToolMenuEntry::InitToolBarButton(FGlobalEditorCommonCommands::Get().FindInContentBrowser));
+		AssetSection.AddSeparator("Common");
 	}
 }
 
@@ -2122,6 +1787,7 @@ void SRemoteControlPanel::GenerateToolbar()
 
 	ToolbarWidgetContent = SNew(SBorder)
 		.BorderImage(FAppStyle::Get().GetBrush("Brushes.Panel"))
+		.Padding(0)
 		[
 			SNullWidget::NullWidget
 		];
@@ -2163,17 +1829,9 @@ void SRemoteControlPanel::GenerateToolbar()
 				];
 		}
 
-		MiscWidgets = SNew(SVerticalBox)
-			+ SVerticalBox::Slot()
-			.AutoHeight()
-			[
-				MiscWidgetsHBox
-			];
+		MiscWidgets = MiscWidgetsHBox;
 	}
 
-	const URemoteControlSettings* Settings = GetMutableDefault<URemoteControlSettings>();
-	const FName& DefaultPanelMode = Settings->DefaultPanelMode;
-
 	if (!Preset->SelectedWorld.IsValid())
 	{
 		Preset->SelectedWorld = Preset->GetWorld(true);
@@ -2217,7 +1875,7 @@ void SRemoteControlPanel::GenerateToolbar()
 			.ButtonStyle(&RCPanelStyle->FlatButtonStyle)
 			.CollapseMenuOnParentFocus(true)
 			.HasDownArrow(false)
-			.ContentPadding(FMargin(4.f, 2.f))
+			.ContentPadding(FMargin(4.f, 1.f))
 			.OnGetMenuContent(this, &SRemoteControlPanel::OnGetSelectedWorldButtonContent)
 			.ButtonContent()
 			[
@@ -2233,76 +1891,6 @@ void SRemoteControlPanel::GenerateToolbar()
 		[
 			SNew(SSpacer)
 		]
-		+SHorizontalBox::Slot()
-		.Padding(5.f, 0.f)
-		.HAlign(HAlign_Fill)
-		.VAlign(VAlign_Center)
-		.AutoWidth()
-		[
-			SNew(SRCModeSwitcher)
-			.DefaultMode(DefaultPanelMode)
-			.OnModeSwitched_Lambda([this](const SRCModeSwitcher::FRCMode& NewMode)
-				{
-					if (NewMode.ModeId == TEXT("Operation"))
-					{
-						bIsInLiveMode = true;
-						bIsLogicPanelEnabled = true;
-					}
-					else if (NewMode.ModeId == TEXT("Setup"))
-					{
-						bIsInLiveMode = false;
-					}
-
-					if (PanelDrawer.IsValid())
-					{
-						TSharedRef<FRCPanelDrawerArgs> LivePanel = RegisteredDrawers.FindChecked(ERCPanels::RCP_Live);
-
-						LivePanel->DrawerVisibility = bIsInLiveMode ? EVisibility::Visible : EVisibility::Collapsed;
-
-						// When we are not in Live Mode collapse the drawer.
-						PanelDrawer->TogglePanel(LivePanel, !bIsInLiveMode);
-
-						if (!bIsInLiveMode)
-						{
-							TSharedRef<FRCPanelDrawerArgs> PropertiesPanel = RegisteredDrawers.FindChecked(ERCPanels::RCP_Properties);
-
-							PanelDrawer->TogglePanel(PropertiesPanel);
-						}
-					}
-
-					OnLiveModeChange.ExecuteIfBound(SharedThis(this), bIsInLiveMode);
-
-					URemoteControlSettings* Settings = GetMutableDefault<URemoteControlSettings>();					
-					Settings->DefaultPanelMode = NewMode.ModeId;
-					Settings->PostEditChange();
-					Settings->SaveConfig();
-				}
-			)
-
-			+ SRCModeSwitcher::Mode("Setup")
-			.DefaultLabel(FText::Format(LOCTEXT("SetupModeLabel", "{0}"), { FText::FromString("Setup") }))
-			.DefaultTooltip(FText::Format(LOCTEXT("SetupModeTooltip", "Switch to {0} mode."), { FText::FromString("Setup") }))
-			.HAlignCell(HAlign_Fill)
-			.VAlignCell(VAlign_Fill)
-			.FixedWidth(96.f)
-			.OptionalIcon(FSlateIcon(FAppStyle::GetAppStyleSetName(), "LevelEditor.Tabs.Outliner"))
-
-			+ SRCModeSwitcher::Mode("Operation")
-			.DefaultLabel(FText::Format(LOCTEXT("OpModeLabel", "{0}"), { FText::FromString("Operation") }))
-			.DefaultTooltip(FText::Format(LOCTEXT("OpModeTooltip", "Switch to {0} mode."), { FText::FromString("Operation") }))
-			.HAlignCell(HAlign_Fill)
-			.VAlignCell(VAlign_Fill)
-			.FixedWidth(96.f)
-			.OptionalIcon(FSlateIcon(FAppStyle::GetAppStyleSetName(), "Icons.Blueprint"))
-		]
-		+ SHorizontalBox::Slot()
-		.Padding(5.f, 0.f)
-		.HAlign(HAlign_Fill)
-		.VAlign(VAlign_Center)
-		.FillWidth(1.f)
-		[
-			SNew(SSpacer)
-		]
 		// Companion Separator (Rebind All)
 		+SHorizontalBox::Slot()
 		.HAlign(HAlign_Center)
@@ -2328,7 +1916,6 @@ void SRemoteControlPanel::GenerateToolbar()
 			.ContentPadding(2.f)
 			[
 				SNew(SHorizontalBox)
-
 				+SHorizontalBox::Slot()
 				.HAlign(HAlign_Center)
 				.VAlign(VAlign_Center)
@@ -2339,7 +1926,6 @@ void SRemoteControlPanel::GenerateToolbar()
 					.Image(FAppStyle::Get().GetBrush("Icons.Link"))
 					.ColorAndOpacity(FSlateColor::UseForeground())
 				]
-
 				+SHorizontalBox::Slot()
 				.VAlign(VAlign_Center)
 				.Padding(4.f, 2.f)
@@ -2354,6 +1940,7 @@ void SRemoteControlPanel::GenerateToolbar()
 		]
 		+ SHorizontalBox::Slot()
 		.HAlign(HAlign_Right)
+		.VAlign(VAlign_Fill)
 		.AutoWidth()
 		[
 			MiscWidgets
@@ -2368,59 +1955,33 @@ void SRemoteControlPanel::GenerateToolbar()
 void SRemoteControlPanel::RegisterAuxiliaryToolBar()
 {
 	UToolMenus* ToolMenus = UToolMenus::Get();
-	if (!ToolMenus->IsMenuRegistered(AuxiliaryRemoteControlPanelToolBarName))
+	if (!ToolMenus || ToolMenus->IsMenuRegistered(AuxiliaryRemoteControlPanelToolBarName))
 	{
-		UToolMenu* ToolbarBuilder = ToolMenus->RegisterMenu(AuxiliaryRemoteControlPanelToolBarName, NAME_None, EMultiBoxType::SlimHorizontalToolBar);
-		ToolbarBuilder->StyleName = "ContentBrowser.ToolBar";
+		return;
+	}
 
-#if 0
-		ToolbarBuilder->StyleName = "AssetEditorToolbar";
-#endif
-		{
-			FToolMenuSection& ToolsSection = ToolbarBuilder->AddSection("Tools");
+	UToolMenu* ToolbarBuilder = ToolMenus->RegisterMenu(AuxiliaryRemoteControlPanelToolBarName, NAME_None, EMultiBoxType::SlimHorizontalToolBar);
+	if (!ensure(ToolbarBuilder))
+	{
+		return;
+	}
 
-			const FRemoteControlCommands& Commands = FRemoteControlCommands::Get();
+	ToolbarBuilder->StyleName = "ContentBrowser.ToolBar";
 
-			ToolsSection.AddEntry(FToolMenuEntry::InitWidget("Signature",
-				SNew(SVerticalBox)
-				+ SVerticalBox::Slot()
-				.AutoHeight()
-				[
-					SNew(SAutoResizeButton)
-					.UICommand(FRemoteControlCommands::Get().ToggleSignatureEditor)
-					.ForceSmallIcons_Static(SRemoteControlPanel::ShouldForceSmallIcons)
-					.IconOverride(FSlateIcon(FAppStyle::GetAppStyleSetName(), TEXT("GraphEditor.Function_16x")))
-				]
-				, Commands.ToggleSignatureEditor->GetLabel())
-			);
+	FToolMenuSection& ToolsSection = ToolbarBuilder->AddSection("Tools");
 
-			ToolsSection.AddEntry(FToolMenuEntry::InitWidget("Logic",
-				SNew(SVerticalBox)
-				+ SVerticalBox::Slot()
-				.AutoHeight()
-				[
-					SNew(SAutoResizeButton)
-					.UICommand(FRemoteControlCommands::Get().ToggleLogicEditor)
-					.ForceSmallIcons_Static(SRemoteControlPanel::ShouldForceSmallIcons)
-					.IconOverride(FSlateIcon(FAppStyle::GetAppStyleSetName(), TEXT("GraphEditor.StateMachine_16x")))
-				]
-				, Commands.ToggleLogicEditor->GetLabel())
-			);
+	auto AddModeEntry = [&ToolsSection](const TSharedPtr<FUICommandInfo>& InCommand, FSlateIcon InIcon)
+		{
+			ToolsSection.AddEntry(FToolMenuEntry::InitToolBarButton(InCommand, TAttribute<FText>(), TAttribute<FText>(), InIcon));
+		};
 
-			ToolsSection.AddEntry(FToolMenuEntry::InitWidget("Protocols",
-				SNew(SVerticalBox)
-				+ SVerticalBox::Slot()
-				.AutoHeight()
-				[
-					SNew(SAutoResizeButton)
-					.UICommand(FRemoteControlCommands::Get().ToggleProtocolMappings)
-					.ForceSmallIcons_Static(SRemoteControlPanel::ShouldForceSmallIcons)
-					.IconOverride(FSlateIcon(FAppStyle::GetAppStyleSetName(), "LevelEditor.Tabs.StatsViewer"))
-				]
-				, Commands.ToggleProtocolMappings->GetLabel())
-			);
-		}
-	}
+	const FRemoteControlCommands& Commands = FRemoteControlCommands::Get();
+
+	AddModeEntry(Commands.ActivateLogicMode    , FSlateIcon(FAppStyle::GetAppStyleSetName(), TEXT("GraphEditor.StateMachine_16x")));
+	AddModeEntry(Commands.ActivateDetailsMode  , FSlateIcon(FAppStyle::GetAppStyleSetName(), TEXT("LevelEditor.Tabs.Details")));
+	AddModeEntry(Commands.ActivateSignatureMode, FSlateIcon(FAppStyle::GetAppStyleSetName(), TEXT("GraphEditor.Function_16x")));
+	AddModeEntry(Commands.ActivateProtocolsMode, FSlateIcon(FAppStyle::GetAppStyleSetName(), TEXT("LevelEditor.Tabs.StatsViewer")));
+	AddModeEntry(Commands.ActivateOutputLogMode, FSlateIcon(FAppStyle::GetAppStyleSetName(), TEXT("MessageLog.TabIcon")));
 }
 
 void SRemoteControlPanel::GenerateAuxiliaryToolbar()
@@ -2429,6 +1990,7 @@ void SRemoteControlPanel::GenerateAuxiliaryToolbar()
 
 	AuxiliaryToolbarWidgetContent = SNew(SBorder)
 		.BorderImage(FAppStyle::Get().GetBrush("Brushes.Panel"))
+		.Padding(0)
 		[
 			SNullWidget::NullWidget
 		];
@@ -2471,7 +2033,7 @@ bool SRemoteControlPanel::CanSaveAsset() const
 	return Preset.IsValid() && Preset->IsAsset();
 }
 
-void SRemoteControlPanel::SaveAsset_Execute() const
+void SRemoteControlPanel::SaveAsset() const
 {
 	if (Preset.IsValid())
 	{
@@ -2499,15 +2061,11 @@ bool SRemoteControlPanel::CanFindInContentBrowser() const
 	return Preset.IsValid();
 }
 
-void SRemoteControlPanel::FindInContentBrowser_Execute() const
+void SRemoteControlPanel::FindInContentBrowser() const
 {
-	if (Preset.IsValid())
+	if (URemoteControlPreset* ResolvedPreset = Preset.Get())
 	{
-		TArray<UObject*> ObjectsToSyncTo;
-
-		ObjectsToSyncTo.Add(Preset.Get());
-
-		GEditor->SyncBrowserToObjects(ObjectsToSyncTo);
+		GEditor->SyncBrowserToObject(ResolvedPreset);
 	}
 }
 
@@ -2531,125 +2089,15 @@ bool SRemoteControlPanel::ShouldForceSmallIcons()
 	return false;
 }
 
-void SRemoteControlPanel::ToggleProtocolMappings_Execute()
-{
-	bIsInProtocolsMode = !bIsInProtocolsMode;
-
-	if (PanelDrawer.IsValid())
-	{
-		TSharedRef<FRCPanelDrawerArgs> ProtocolsPanel = RegisteredDrawers.FindChecked(ERCPanels::RCP_Protocols);
-
-		ProtocolsPanel->DrawerVisibility = bIsInProtocolsMode ? EVisibility::Visible : EVisibility::Collapsed;
-
-		// When we are not in Protocols Mode collapse the drawer.
-		PanelDrawer->TogglePanel(ProtocolsPanel, !bIsInProtocolsMode);
-
-		if (!bIsInProtocolsMode)
-		{
-			TSharedRef<FRCPanelDrawerArgs> PropertiesPanel = RegisteredDrawers.FindChecked(ERCPanels::RCP_Properties);
-
-			PanelDrawer->TogglePanel(PropertiesPanel);
-		}
-
-		if (EntityList.IsValid())
-		{
-			const bool bToggleProtcolMode = IsInProtocolsMode() && ActivePanel == ERCPanels::RCP_Protocols;
-
-			EntityList->RebuildListWithColumns(bToggleProtcolMode ? EEntitiesListMode::Protocols : EEntitiesListMode::Default);
-		}
-	}
-}
-
-bool SRemoteControlPanel::CanToggleProtocolsMode() const
-{
-	return !bIsInLiveMode;
-}
-
-bool SRemoteControlPanel::IsInProtocolsMode() const
-{
-	return bIsInProtocolsMode;
-}
-
-void SRemoteControlPanel::ToggleLogicEditor_Execute()
-{
-	bIsLogicPanelEnabled = !bIsLogicPanelEnabled;
-
-	URemoteControlSettings* Settings = GetMutableDefault<URemoteControlSettings>();
-	Settings->bLogicPanelVisibility = bIsLogicPanelEnabled;
-	Settings->PostEditChange();
-	Settings->SaveConfig();	
-
-	if (PanelDrawer.IsValid() && (ActivePanel != ERCPanels::RCP_Properties))
-	{
-		TSharedRef<FRCPanelDrawerArgs> PropertiesPanel = RegisteredDrawers.FindChecked(ERCPanels::RCP_Properties);
-
-		PanelDrawer->TogglePanel(PropertiesPanel);
-	}
-}
-
-bool SRemoteControlPanel::CanToggleLogicPanel() const
-{
-	return true;
-}
-
-bool SRemoteControlPanel::IsLogicPanelEnabled() const
-{
-	return bIsLogicPanelEnabled;
-}
-
-void SRemoteControlPanel::ToggleSignatureEditor_Execute()
-{
-	bIsSignaturePanelEnabled = !bIsSignaturePanelEnabled;
-
-	if (PanelDrawer.IsValid())
-	{
-		TSharedRef<FRCPanelDrawerArgs> SignaturePanelDrawer = RegisteredDrawers.FindChecked(ERCPanels::RCP_Signature);
-
-		SignaturePanelDrawer->DrawerVisibility = bIsSignaturePanelEnabled ? EVisibility::Visible : EVisibility::Collapsed;
-		PanelDrawer->TogglePanel(SignaturePanelDrawer, !bIsSignaturePanelEnabled);
-
-		if (!bIsSignaturePanelEnabled)
-		{
-			TSharedRef<FRCPanelDrawerArgs> PropertiesPanel = RegisteredDrawers.FindChecked(ERCPanels::RCP_Properties);
-			PanelDrawer->TogglePanel(PropertiesPanel);
-		}
-	}
-}
-
-bool SRemoteControlPanel::CanToggleSignaturePanel() const
-{
-	return true;
-}
-
-bool SRemoteControlPanel::IsSignaturePanelEnabled() const
-{
-	return bIsSignaturePanelEnabled;
-}
-
-void SRemoteControlPanel::OnRCPanelToggled(ERCPanels InPanelID)
-{
-	if (InPanelID != ActivePanel)
-	{
-		ActivePanel = InPanelID;
-
-		if (EntityList.IsValid() && !bIsInLiveMode)
-		{
-			const bool bToggleProtcolMode = IsInProtocolsMode() && ActivePanel == ERCPanels::RCP_Protocols;
-
-			EntityList->RebuildListWithColumns(bToggleProtcolMode ? EEntitiesListMode::Protocols : EEntitiesListMode::Default);
-		}
-	}
-}
-
 TSharedPtr<class SRCLogicPanelBase> SRemoteControlPanel::GetActiveLogicPanel() const
 {
 	if (ControllerPanel->IsListFocused())
 	{
 		return ControllerPanel;
 	}
-	else if (BehaviourPanel->IsListFocused())
+	else if (BehaviorPanel->IsListFocused())
 	{
-		return BehaviourPanel;
+		return BehaviorPanel;
 	}
 	else if (ActionPanel->IsListFocused())
 	{
@@ -2710,7 +2158,7 @@ void SRemoteControlPanel::DeleteEntity_Execute()
 
 bool SRemoteControlPanel::CanDeleteEntity() const
 {
-	if (bIsInLiveMode)
+	if (IsModeActive(ERCPanelMode::Live))
 	{
 		return false;
 	}
@@ -2748,7 +2196,7 @@ void SRemoteControlPanel::RenameEntity_Execute() const
 
 bool SRemoteControlPanel::CanRenameEntity() const
 {
-	if (bIsInLiveMode)
+	if (IsModeActive(ERCPanelMode::Live))
 	{
 		return false;
 	}
@@ -2791,7 +2239,7 @@ void SRemoteControlPanel::ChangePropertyId_Execute() const
 
 bool SRemoteControlPanel::CanChangePropertyId() const
 {
-	if (bIsInLiveMode)
+	if (IsModeActive(ERCPanelMode::Live))
 	{
 		return false;
 	}
@@ -2825,7 +2273,7 @@ void SRemoteControlPanel::CopyItem_Execute()
 
 bool SRemoteControlPanel::CanCopyItem() const
 {
-	if (bIsInLiveMode)
+	if (IsModeActive(ERCPanelMode::Live))
 	{
 		return false;
 	}
@@ -2848,7 +2296,7 @@ void SRemoteControlPanel::PasteItem_Execute()
 
 bool SRemoteControlPanel::CanPasteItem() const
 {
-	if (bIsInLiveMode)
+	if (IsModeActive(ERCPanelMode::Live))
 	{
 		return false;
 	}
@@ -2880,7 +2328,7 @@ void SRemoteControlPanel::DuplicateItem_Execute()
 
 bool SRemoteControlPanel::CanDuplicateItem() const
 {
-	if (bIsInLiveMode)
+	if (IsModeActive(ERCPanelMode::Live))
 	{
 		return false;
 	}
@@ -2903,7 +2351,7 @@ void SRemoteControlPanel::UpdateValue_Execute()
 
 bool SRemoteControlPanel::CanUpdateValue() const
 {
-	if (bIsInLiveMode)
+	if (IsModeActive(ERCPanelMode::Live))
 	{
 		return false;
 	}
diff --git a/Engine/Plugins/VirtualProduction/RemoteControl/Source/RemoteControlUI/Private/UI/SRemoteControlPanel.h b/Engine/Plugins/VirtualProduction/RemoteControl/Source/RemoteControlUI/Private/UI/SRemoteControlPanel.h
--- a/Engine/Plugins/VirtualProduction/RemoteControl/Source/RemoteControlUI/Private/UI/SRemoteControlPanel.h
+++ b/Engine/Plugins/VirtualProduction/RemoteControl/Source/RemoteControlUI/Private/UI/SRemoteControlPanel.h
@@ -13,7 +13,22 @@
 #include "Widgets/SCompoundWidget.h"
 #include "Widgets/DeclarativeSyntaxSupport.h"
 
-enum class ERCPanels : uint8;
+/** The variety of panels we have in the RC Panel. */
+enum class ERCPanelMode : uint8
+{
+	Controller,
+
+	EntityDetails,
+
+	Protocols,
+
+	OutputLog,
+
+	Live,
+
+	Signature,
+};
+
 class AActor;
 struct EVisibility;
 struct FAssetData;
@@ -112,24 +127,14 @@ public:
 	 */
 	void ToggleProperty(const FRCExposesPropertyArgs& InArgs, FString InDesiredName = TEXT(""));
 
-	/**
-	 * @return Whether or not the panel is in live mode.
-	 */
-	bool IsInLiveMode() const { return bIsInLiveMode; }
-
 	/**
 	 * Get the selected group.
 	 */
 	FGuid GetSelectedGroup() const;
 
-	/**
-	 * Set the edit mode of the panel.
-	 * @param bEditMode The desired mode.
-	 */
-	void SetLiveMode(bool bLiveMode)
-	{
-		bIsInLiveMode = bLiveMode;
-	}
+	bool CanActivateMode(ERCPanelMode InPanelMode) const;
+	bool IsModeActive(ERCPanelMode InPanelMode) const;
+	void SetActiveMode(ERCPanelMode InPanelMode);
 
 	/**
 	 * Get the exposed entity list.
@@ -169,10 +174,15 @@ public:
 	virtual void AddReferencedObjects(FReferenceCollector& Collector) override;
 
 private:
-
 	/** Returns a Helper widget for entity details view and protocol details view. */
 	static TSharedRef<SBox> CreateNoneSelectedWidget();
 
+	TSharedRef<SWidget> BuildLogicModeContent(const TSharedRef<SWidget>& InLogicPanel);
+	TSharedRef<SWidget> BuildEntityDetailsModeContent(const TAttribute<float>& InRatioTop, const TAttribute<float>& InRatioBottom);
+	TSharedRef<SWidget> BuildProtocolsModeContent(const TAttribute<float>& InRatioTop, const TAttribute<float>& InRatioBottom);
+	TSharedRef<SWidget> BuildLiveModeContent(const TSharedRef<SWidget>& InLogicPanel);
+	TSharedRef<SWidget> BuildSignaturesModeContent();
+
 	//~ Remote Control Commands
 	void BindRemoteControlCommands();
 
@@ -190,11 +200,6 @@ private:
 	/** Unregister editor events */
 	void UnregisterEvents();
 
-	/** Register panels to the drawer. */
-	void RegisterPanels();
-	/** Unregister panels from the drawer. */
-	void UnregisterPanels();
-
 	/** Unexpose a field from the preset. */
 	void Unexpose(const FRCExposesPropertyArgs& InArgs);
 
@@ -268,7 +273,7 @@ private:
 	void OnEntityUnexposed(URemoteControlPreset* InPreset, const FGuid& InEntityId);
 
 	/** Toggles the logging part of UI */
-	void OnLogCheckboxToggle(ECheckBoxState State);
+	void OnLogCheckboxToggle(ECheckBoxState InState);
 
 	/** Triggers a next frame update of the actor function picker to ensure that added actors are valid. */
 	void UpdateActorFunctionPicker();
@@ -300,30 +305,16 @@ private:
 	bool CanSaveAsset() const;
 
 	/** Called when "Save" is clicked for this asset */
-	void SaveAsset_Execute() const;
+	void SaveAsset() const;
 
 	/** Called to test if "Find in Content Browser" should be enabled for this asset */
 	bool CanFindInContentBrowser() const;
 
 	/** Called when "Find in Content Browser" is clicked for this asset */
-	void FindInContentBrowser_Execute() const;
+	void FindInContentBrowser() const;
 
 	static bool ShouldForceSmallIcons();
 
-	void ToggleProtocolMappings_Execute();
-	bool CanToggleProtocolsMode() const;
-	bool IsInProtocolsMode() const;
-
-	void ToggleLogicEditor_Execute();
-	bool CanToggleLogicPanel() const;
-	bool IsLogicPanelEnabled() const;
-
-	void ToggleSignatureEditor_Execute();
-	bool CanToggleSignaturePanel() const;
-	bool IsSignaturePanelEnabled() const;
-
-	void OnRCPanelToggled(ERCPanels InPanelID);
-
 	/** Called when user attempts to delete a group/exposed entity. */
 	void DeleteEntity_Execute();
 
@@ -399,14 +390,6 @@ private:
 	static const FName TargetWorldRemoteControlPanelMenuName;
 	/** Holds the preset asset. */
 	TStrongObjectPtr<URemoteControlPreset> Preset;
-	/** Whether the panel is in protocols mode. */
-	bool bIsInProtocolsMode = false;
-	/** Whether the panel is in live (or) operation mode. */
-	bool bIsInLiveMode = false;
-	/** Whether the logic panel is enabled or not. */
-	bool bIsLogicPanelEnabled = false;
-	/** Whether the signature panel is enabled or not. */
-	bool bIsSignaturePanelEnabled = false;
 	/** Delegate called when the live mode changes. */
 	FOnLiveModeChange OnLiveModeChange;
 	/** Holds the blueprint library picker */
@@ -456,11 +439,11 @@ private:
 	/** Panel Drawer widget holds all docked panels. */
 	TSharedPtr<SRCPanelDrawer> PanelDrawer;
 	/** Map of Opened Drawers. */
-	TMap<ERCPanels, TSharedRef<FRCPanelDrawerArgs>> RegisteredDrawers;
+	TMap<ERCPanelMode, TSharedRef<FRCPanelDrawerArgs>> RegisteredDrawers;
 	/** Panel Style reference. */
 	const FRCPanelStyle* RCPanelStyle;
 	/** Stores the active panel that is drawn. */
-	ERCPanels ActivePanel;
+	ERCPanelMode ActiveMode = ERCPanelMode::Controller;
 	/** Input Preprocessor which catches the Delete Key when Docked. */
 	TSharedPtr<IInputProcessor> InputProcessor;
 	/** Currently selected world name */
@@ -470,7 +453,7 @@ private:
 	/** Controller panel UI widget for Remote Control Logic*/
 	TSharedPtr<class SRCControllerPanel> ControllerPanel;
 	/** Behaviour panel UI widget for Remote Control Logic*/
-	TSharedPtr<class SRCBehaviourPanel> BehaviourPanel;
+	TSharedPtr<class SRCBehaviourPanel> BehaviorPanel;
 	/** Action panel UI widget for Remote Control Logic*/
 	TSharedPtr<class SRCActionPanel> ActionPanel;
 	/** Signature panel UI widget */
