Remote Control: Signature UI update to keep expansion + selection state.
+ changed mechanism so that now instead of drag dropping from entity list to reuse the expose button

diff --git a/Engine/Plugins/VirtualProduction/RemoteControl/Source/RemoteControl/Private/RemoteControlSignature.cpp b/Engine/Plugins/VirtualProduction/RemoteControl/Source/RemoteControl/Private/RemoteControlSignature.cpp
--- a/Engine/Plugins/VirtualProduction/RemoteControl/Source/RemoteControl/Private/RemoteControlSignature.cpp
+++ b/Engine/Plugins/VirtualProduction/RemoteControl/Source/RemoteControl/Private/RemoteControlSignature.cpp
@@ -5,42 +5,55 @@
 #include "RemoteControlField.h"
 #include "RemoteControlPreset.h"
 
-int32 FRCSignature::AddFieldsFromEntities(URemoteControlPreset* InPreset, TConstArrayView<FGuid> InFieldEntityIds)
+FRCSignatureField FRCSignatureField::CreateField(const FRCFieldPathInfo& InFieldPathInfo, UObject* InOwnerObject, FProperty* InProperty)
 {
-	if (!InPreset || InFieldEntityIds.IsEmpty())
+	UClass* SupportedBindingClass = nullptr;
+
+	if (InProperty)
 	{
-		return 0;
+		if (UClass* PropertyOwnerClass = InProperty->GetOwnerClass())
+		{
+			SupportedBindingClass = PropertyOwnerClass;
+		}
 	}
 
-	const int32 PreviousFieldCount = Fields.Num();
+	return CreateField(InFieldPathInfo, InOwnerObject, SupportedBindingClass);
+}
+
+FRCSignatureField FRCSignatureField::CreateField(const FRCFieldPathInfo& InFieldPathInfo, UObject* InOwnerObject, UClass* InSupportedClass)
+{
+	if (!InSupportedClass && InFieldPathInfo.GetSegmentCount() > 0 && InFieldPathInfo.GetFieldSegment(0).IsResolved())
+	{
+		InSupportedClass = InFieldPathInfo.GetFieldSegment(0).ResolvedData.Field->GetOwnerClass();
+	}
 
-	Fields.Reserve(PreviousFieldCount + InFieldEntityIds.Num());
+	FRCSignatureField Field;
+	Field.bEnabled = true;
+	Field.FieldPath = InFieldPathInfo;
+	Field.SupportedClass = InSupportedClass;
 
-	for (const FGuid& EntityId : InFieldEntityIds)
+	if (UObject* BoundObject = InOwnerObject)
 	{
-		TSharedPtr<FRemoteControlProperty> ExposedProperty = InPreset->GetExposedEntity<FRemoteControlProperty>(EntityId).Pin();
-		if (!ExposedProperty.IsValid())
+		if (AActor* ActorOwner = BoundObject->GetTypedOuter<AActor>())
 		{
-			continue;
+			Field.ObjectRelativePath = BoundObject->GetPathName(ActorOwner);
 		}
+	}
 
-		FRCSignatureField Field;
-		Field.bEnabled = true;
-		Field.FieldPath = ExposedProperty->FieldPathInfo;
-		Field.SupportedClass = ExposedProperty->GetSupportedBindingClass();
+	return Field;
+}
 
-		if (UObject* BoundObject = ExposedProperty->GetBoundObject())
-		{
-			if (AActor* ActorOwner = BoundObject->GetTypedOuter<AActor>())
-			{
-				Field.ObjectRelativePath = BoundObject->GetPathName(ActorOwner);
-			}
-		}
+int32 FRCSignature::AddFields(TConstArrayView<FRCSignatureField> InFields)
+{
+	const int32 PreviousNum = Fields.Num();
+	Fields.Reserve(PreviousNum + InFields.Num());
 
-		Fields.AddUnique(MoveTemp(Field));
+	for (const FRCSignatureField& Field : InFields)
+	{
+		Fields.AddUnique(Field);
 	}
 
-	return Fields.Num() - PreviousFieldCount;
+	return Fields.Num() - PreviousNum;
 }
 
 int32 FRCSignature::ApplySignature(URemoteControlPreset* InPreset, TConstArrayView<TWeakObjectPtr<AActor>> InActors) const
diff --git a/Engine/Plugins/VirtualProduction/RemoteControl/Source/RemoteControl/Public/RemoteControlSignature.h b/Engine/Plugins/VirtualProduction/RemoteControl/Source/RemoteControl/Public/RemoteControlSignature.h
--- a/Engine/Plugins/VirtualProduction/RemoteControl/Source/RemoteControl/Public/RemoteControlSignature.h
+++ b/Engine/Plugins/VirtualProduction/RemoteControl/Source/RemoteControl/Public/RemoteControlSignature.h
@@ -14,6 +14,9 @@ struct FRCSignatureField
 
 	FRCSignatureField() = default;
 
+	REMOTECONTROL_API static FRCSignatureField CreateField(const FRCFieldPathInfo& InFieldPathInfo, UObject* InOwnerObject, FProperty* InProperty);
+	REMOTECONTROL_API static FRCSignatureField CreateField(const FRCFieldPathInfo& InFieldPathInfo, UObject* InOwnerObject, UClass* InSupportedClass);
+
 	bool operator==(const FRCSignatureField& InOtherField) const
 	{
 		return FieldPath == InOtherField.FieldPath
@@ -49,16 +52,16 @@ struct FRCSignature
 	}
 
 	/**
-	 * Gathers the Field Entities from the given preset and field ids and adds them as Signature Fields to this instance
-	 * @param InPreset the preset to check for
-	 * @param InFieldEntityIds the ids of the entities to find and add
+	 * Adds the given fields to this signature and ensuring no repeated field is present
+	 * @return the number of new fields added
 	 */
-	REMOTECONTROL_API int32 AddFieldsFromEntities(URemoteControlPreset* InPreset, TConstArrayView<FGuid> InFieldEntityIds);
+	REMOTECONTROL_API int32 AddFields(TConstArrayView<FRCSignatureField> InFields);
 
 	/**
 	 * Applies this Signature to the given Actors by exposing all this Signature's fields to the given preset
 	 * @param InPreset the preset where these properties will be exposed to
 	 * @param InActors the actors whose properties to expose
+	 * @return the number of properties affected in total
 	 */
 	REMOTECONTROL_API int32 ApplySignature(URemoteControlPreset* InPreset, TConstArrayView<TWeakObjectPtr<AActor>> InActors) const;
 
diff --git a/Engine/Plugins/VirtualProduction/RemoteControl/Source/RemoteControlUI/Private/RemoteControlUIModule.cpp b/Engine/Plugins/VirtualProduction/RemoteControl/Source/RemoteControlUI/Private/RemoteControlUIModule.cpp
--- a/Engine/Plugins/VirtualProduction/RemoteControl/Source/RemoteControlUI/Private/RemoteControlUIModule.cpp
+++ b/Engine/Plugins/VirtualProduction/RemoteControl/Source/RemoteControlUI/Private/RemoteControlUIModule.cpp
@@ -497,35 +497,23 @@ void FRemoteControlUIModule::FillRemoteControlRowContextSection(UToolMenu* InToo
 			, FToolMenuInsert(NAME_None, EToolMenuInsertType::First));
 	}
 
-	const TAttribute<FText> ExposeUnexposeLabel = TAttribute<FText>::Create(TAttribute<FText>::FGetter::CreateRaw(this, &FRemoteControlUIModule::GetExposePropertyButtonText, ExposesPropertyArgs));
-	const TAttribute<FText> ExposeUnexposeTooltip = TAttribute<FText>::Create(TAttribute<FText>::FGetter::CreateRaw(this, &FRemoteControlUIModule::GetExposePropertyButtonTooltip, ExposesPropertyArgs));
-	const TAttribute<FSlateIcon> ExposeUnexposeIcon = TAttribute<FSlateIcon>::Create(
-			[this, ExposesPropertyArgs]
-		{
-			return OnGetExposedIcon(ExposesPropertyArgs);
-		});;
-
 	// Unexpose/Expose entry for the Menu
 	RemoteControlSection->AddMenuEntry(
 		TEXT("RemoteControlExposeUnexposeEntry")
-		, ExposeUnexposeLabel
-		, ExposeUnexposeTooltip
-		, ExposeUnexposeIcon
+		, TAttribute<FText>::CreateRaw(this, &FRemoteControlUIModule::GetPropertyActionText, ExposesPropertyArgs)
+		, TAttribute<FText>::CreateRaw(this, &FRemoteControlUIModule::GetPropertyActionTooltip, ExposesPropertyArgs)
+		, TAttribute<FSlateIcon>::CreateRaw(this, &FRemoteControlUIModule::OnGetPropertyActionIcon, ExposesPropertyArgs)
 		, FUIAction(
-			FExecuteAction::CreateRaw(this, &FRemoteControlUIModule::OnToggleExposeProperty, ExposesPropertyArgs),
-			FCanExecuteAction::CreateRaw(this, &FRemoteControlUIModule::CanToggleExposeProperty, ExposesPropertyArgs),
-			FGetActionCheckState::CreateRaw(this, &FRemoteControlUIModule::GetPropertyExposedCheckState, ExposesPropertyArgs),
-			FIsActionButtonVisible::CreateRaw(this, &FRemoteControlUIModule::CanToggleExposeProperty, ExposesPropertyArgs)
+			FExecuteAction::CreateRaw(this, &FRemoteControlUIModule::ExecutePropertyAction, ExposesPropertyArgs),
+			FCanExecuteAction::CreateRaw(this, &FRemoteControlUIModule::CanExecutePropertyAction, ExposesPropertyArgs),
+			FGetActionCheckState(),
+			FIsActionButtonVisible::CreateRaw(this, &FRemoteControlUIModule::CanExecutePropertyAction, ExposesPropertyArgs)
 			)
 		);
 
 	// Unexpose/Expose SubProperty SubMenu for the Menu
 	if (ExposesPropertyArgs.PropertyHandle.IsValid() && HasChildProperties(ExposesPropertyArgs.PropertyHandle->GetProperty()))
 	{
-		const TAttribute<FSlateIcon> ExposeSubPropertyIcon = TAttribute<FSlateIcon>::Create([]()
-			{
-				return FSlateIcon(FRemoteControlExposeMenuStyle::GetStyleSetName(),"RemoteControlExposeMenu.Expand");
-			});
 		RemoteControlSection->AddSubMenu(
 			TEXT("ExposeUnexposeSubPropertyMenu")
 			, LOCTEXT("ExposeUnexposeSubPropertyMenu_Label", "Toggle Sub Property of this field")
@@ -591,9 +579,9 @@ void FRemoteControlUIModule::GetSubPropertySubMenu(UToolMenu* InToolMenu, FRCExp
 			FSlateIcon(),
 			FUIAction(
 				FExecuteAction::CreateRaw(this, &FRemoteControlUIModule::OnToggleExposeSubProperty, PropArgs.PropertyArgs, PropArgs.DesiredName),
-				FCanExecuteAction::CreateRaw(this, &FRemoteControlUIModule::CanToggleExposeProperty, PropArgs.PropertyArgs),
+				FCanExecuteAction::CreateRaw(this, &FRemoteControlUIModule::CanExecutePropertyAction, PropArgs.PropertyArgs),
 				FGetActionCheckState::CreateRaw(this, &FRemoteControlUIModule::GetPropertyExposedCheckState, PropArgs.PropertyArgs),
-				FIsActionButtonVisible::CreateRaw(this, &FRemoteControlUIModule::CanToggleExposeProperty, PropArgs.PropertyArgs)),
+				FIsActionButtonVisible::CreateRaw(this, &FRemoteControlUIModule::CanExecutePropertyAction, PropArgs.PropertyArgs)),
 			EUserInterfaceActionType::ToggleButton);
 	}
 }
@@ -692,49 +680,19 @@ void FRemoteControlUIModule::HandleCreatePropertyRowExtension(const FOnGenerateG
 
 		// Expose/Unexpose button.
 		FPropertyRowExtensionButton& ExposeButton = OutExtensions.AddDefaulted_GetRef();
-		ExposeButton.Icon = TAttribute<FSlateIcon>::Create(
-			[this, ExposesPropertyArgs]
-		{
-			return OnGetExposedIcon(ExposesPropertyArgs);
-		});
-
-		ExposeButton.Label = TAttribute<FText>::Create(TAttribute<FText>::FGetter::CreateRaw(this, &FRemoteControlUIModule::GetExposePropertyButtonText, ExposesPropertyArgs));
-		ExposeButton.ToolTip = TAttribute<FText>::Create(TAttribute<FText>::FGetter::CreateRaw(this, &FRemoteControlUIModule::GetExposePropertyButtonTooltip, ExposesPropertyArgs));
+		ExposeButton.Icon = TAttribute<FSlateIcon>::CreateRaw(this, &FRemoteControlUIModule::OnGetPropertyActionIcon, ExposesPropertyArgs);
+		ExposeButton.Label = TAttribute<FText>::CreateRaw(this, &FRemoteControlUIModule::GetPropertyActionText, ExposesPropertyArgs);
+		ExposeButton.ToolTip = TAttribute<FText>::CreateRaw(this, &FRemoteControlUIModule::GetPropertyActionTooltip, ExposesPropertyArgs);
 		ExposeButton.UIAction = FUIAction(
-			FExecuteAction::CreateRaw(this, &FRemoteControlUIModule::OnToggleExposeProperty, ExposesPropertyArgs),
-			FCanExecuteAction::CreateRaw(this, &FRemoteControlUIModule::CanToggleExposeProperty, ExposesPropertyArgs),
-			FGetActionCheckState::CreateRaw(this, &FRemoteControlUIModule::GetPropertyExposedCheckState, ExposesPropertyArgs),
-			FIsActionButtonVisible::CreateRaw(this, &FRemoteControlUIModule::CanToggleExposeProperty, ExposesPropertyArgs)
+			FExecuteAction::CreateRaw(this, &FRemoteControlUIModule::ExecutePropertyAction, ExposesPropertyArgs),
+			FCanExecuteAction::CreateRaw(this, &FRemoteControlUIModule::CanExecutePropertyAction, ExposesPropertyArgs),
+			FGetActionCheckState(),
+			FIsActionButtonVisible::CreateRaw(this, &FRemoteControlUIModule::CanExecutePropertyAction, ExposesPropertyArgs)
 		);
 
-		if (ExposesPropertyArgs.PropertyHandle.IsValid() && HasChildProperties(ExposesPropertyArgs.PropertyHandle->GetProperty()))
-		{
-			// Expose/Unexpose button.
-			FPropertyRowExtensionButton& ExposeSubPropertiesButton = OutExtensions.AddDefaulted_GetRef();
-			ExposeSubPropertiesButton.Icon = TAttribute<FSlateIcon>::Create([]()
-			{
-				return FSlateIcon(FRemoteControlExposeMenuStyle::GetStyleSetName(),"RemoteControlExposeMenu.Expand");
-			});
-
-			ExposeSubPropertiesButton.Label = LOCTEXT("RemoteControlExpandMenu_Title","Expose Sub-Property");
-			ExposeSubPropertiesButton.ToolTip = LOCTEXT("RemoteControlExpandMenu_Tooltip","Expose Sub-Property");
-			ExposeSubPropertiesButton.UIAction = FUIAction(
-				FExecuteAction::CreateRaw(this, &FRemoteControlUIModule::OnToggleExposePropertyWithChild, ExposesPropertyArgs),
-				FCanExecuteAction::CreateRaw(this, &FRemoteControlUIModule::CanToggleExposeProperty, ExposesPropertyArgs),
-				FGetActionCheckState(),
-				FIsActionButtonVisible::CreateRaw(this, &FRemoteControlUIModule::CanToggleExposeProperty, ExposesPropertyArgs));
-		}
-
 		// Override material(s) warning.
 		FPropertyRowExtensionButton& OverrideMaterialButton = OutExtensions.AddDefaulted_GetRef();
-
-		OverrideMaterialButton.Icon = TAttribute<FSlateIcon>::Create(
-			[this, ExposesPropertyArgs]
-		{
-			return OnGetOverrideMaterialsIcon(ExposesPropertyArgs);
-		}
-		);
-
+		OverrideMaterialButton.Icon = TAttribute<FSlateIcon>::CreateRaw(this, &FRemoteControlUIModule::OnGetOverrideMaterialsIcon, ExposesPropertyArgs);
 		OverrideMaterialButton.Label = LOCTEXT("OverrideMaterial", "Override Material");
 		OverrideMaterialButton.ToolTip = LOCTEXT("OverrideMaterialToolTip", "Click to override this material in order to expose this property to Remote Control.");
 		OverrideMaterialButton.UIAction = FUIAction(
@@ -775,6 +733,11 @@ TSharedPtr<SRemoteControlPanel> FRemoteControlUIModule::GetPanelForObject(const
 
 TSharedPtr<SRemoteControlPanel> FRemoteControlUIModule::GetPanelForProperty(const FRCExposesPropertyArgs& InPropertyArgs) const
 {
+	if (!InPropertyArgs.IsValid())
+	{
+		return nullptr;
+	}
+
 	const FRCExposesPropertyArgs::EType ExtensionArgsType = InPropertyArgs.GetType();
 
 	if (ExtensionArgsType == FRCExposesPropertyArgs::EType::E_Handle)
@@ -805,30 +768,33 @@ TSharedPtr<SRemoteControlPanel> FRemoteControlUIModule::GetPanelForPropertyChang
 	return GetPanelForObject(nullptr);
 }
 
-FSlateIcon FRemoteControlUIModule::OnGetExposedIcon(const FRCExposesPropertyArgs& InPropertyArgs) const
+FSlateIcon FRemoteControlUIModule::OnGetPropertyActionIcon(const FRCExposesPropertyArgs InPropertyArgs) const
 {
 	FName BrushName("RemoteControlExposeMenu.NoBrush");
 
-	if (TSharedPtr<SRemoteControlPanel> Panel = GetPanelForProperty(InPropertyArgs))
+	TSharedPtr<SRemoteControlPanel> Panel = GetPanelForProperty(InPropertyArgs);
+
+	if (Panel.IsValid() && Panel->GetPreset())
 	{
-		if (Panel->GetPreset())
+		if (Panel->IsModeActive(ERCPanelMode::Signature))
 		{
-			EPropertyExposeStatus Status = GetPropertyExposeStatus(InPropertyArgs);
-			if (Status == EPropertyExposeStatus::Exposed)
-			{
-				BrushName = HasChildPropertiesExposed(InPropertyArgs) ? "RemoteControlExposeMenu.VisibleAndVisibleChildren" : "RemoteControlExposeMenu.Visible";
-			}
-			else
-			{
-				BrushName = HasChildPropertiesExposed(InPropertyArgs) ? "RemoteControlExposeMenu.HiddenAndVisibleChildren" : "RemoteControlExposeMenu.Hidden";
-			}
+			return FSlateIcon(FAppStyle::GetAppStyleSetName(), TEXT("BlueprintEditor.AddNewFunction"));
+		}
+
+		if (Panel->IsExposed(InPropertyArgs))
+		{
+			BrushName = HasChildPropertiesExposed(InPropertyArgs) ? "RemoteControlExposeMenu.VisibleAndVisibleChildren" : "RemoteControlExposeMenu.Visible";
+		}
+		else
+		{
+			BrushName = HasChildPropertiesExposed(InPropertyArgs) ? "RemoteControlExposeMenu.HiddenAndVisibleChildren" : "RemoteControlExposeMenu.Hidden";
 		}
 	}
 
 	return FSlateIcon(FRemoteControlExposeMenuStyle::GetStyleSetName(), BrushName);
 }
 
-bool FRemoteControlUIModule::CanToggleExposeProperty(const FRCExposesPropertyArgs InPropertyArgs) const
+bool FRemoteControlUIModule::CanExecutePropertyAction(const FRCExposesPropertyArgs InPropertyArgs) const
 {
 	if (TSharedPtr<SRemoteControlPanel> Panel = GetPanelForProperty(InPropertyArgs))
 	{
@@ -841,34 +807,34 @@ bool FRemoteControlUIModule::CanToggleExposeProperty(const FRCExposesPropertyArg
 
 ECheckBoxState FRemoteControlUIModule::GetPropertyExposedCheckState(const FRCExposesPropertyArgs InPropertyArgs) const
 {
-	if (TSharedPtr<SRemoteControlPanel> Panel = GetPanelForProperty(InPropertyArgs))
+	TSharedPtr<SRemoteControlPanel> Panel = GetPanelForProperty(InPropertyArgs);
+	if (Panel.IsValid() && Panel->GetPreset() && Panel->IsExposed(InPropertyArgs))
 	{
-		if (Panel->GetPreset())
-		{
-			EPropertyExposeStatus ExposeStatus = GetPropertyExposeStatus(InPropertyArgs);
-			if (ExposeStatus == EPropertyExposeStatus::Exposed)
-			{
-				return ECheckBoxState::Checked;
-			}
-		}
+		return ECheckBoxState::Checked;
 	}
-
 	return ECheckBoxState::Unchecked;
 }
 
-void FRemoteControlUIModule::OnToggleExposeProperty(const FRCExposesPropertyArgs InPropertyArgs) const
+void FRemoteControlUIModule::ExecutePropertyAction(const FRCExposesPropertyArgs InPropertyArgs) const
 {
 	if (!ensureMsgf(InPropertyArgs.IsValid(), TEXT("Property could not be exposed because the extension args was invalid.")))
 	{
 		return;
 	}
-	if (ShouldCreateSubMenuForChildProperties(InPropertyArgs))
+
+	TSharedPtr<SRemoteControlPanel> Panel = GetPanelForProperty(InPropertyArgs);
+	if (!Panel.IsValid())
+	{
+		return;
+	}
+
+	if (ShouldCreateSubMenuForChildProperties(Panel, InPropertyArgs))
 	{
 		CreateSubMenuForChildProperties(InPropertyArgs);
 	}
-	else if (TSharedPtr<SRemoteControlPanel> Panel = GetPanelForProperty(InPropertyArgs))
+	else
 	{
-		Panel->ToggleProperty(InPropertyArgs);
+		Panel->ExecutePropertyAction(InPropertyArgs);
 	}
 }
 
@@ -880,7 +846,7 @@ void FRemoteControlUIModule::OnToggleExposeSubProperty(const FRCExposesPropertyA
 	}
 	if (TSharedPtr<SRemoteControlPanel> Panel = GetPanelForProperty(InPropertyArgs))
 	{
-		Panel->ToggleProperty(InPropertyArgs, InDesiredName);
+		Panel->ExecutePropertyAction(InPropertyArgs, InDesiredName);
 	}
 }
 
@@ -897,12 +863,10 @@ void FRemoteControlUIModule::OnExposeAll(const TArray<FRCExposesAllPropertiesArg
 {
 	for (const FRCExposesAllPropertiesArgs& PropInfo : InExposeAllArgs)
 	{
-		if (TSharedPtr<SRemoteControlPanel> Panel = GetPanelForProperty(PropInfo.PropertyArgs))
+		TSharedPtr<SRemoteControlPanel> Panel = GetPanelForProperty(PropInfo.PropertyArgs);
+		if (Panel.IsValid() && !Panel->IsExposed(PropInfo.PropertyArgs))
 		{
-			if (!Panel->IsExposed(PropInfo.PropertyArgs))
-			{
-				OnToggleExposeSubProperty(PropInfo.PropertyArgs, PropInfo.DesiredName);
-			}
+			OnToggleExposeSubProperty(PropInfo.PropertyArgs, PropInfo.DesiredName);
 		}
 	}
 }
@@ -911,12 +875,10 @@ void FRemoteControlUIModule::OnUnexposeAll(const TArray<FRCExposesAllPropertiesA
 {
 	for (const FRCExposesAllPropertiesArgs& PropInfo : InExposeAllArgs)
 	{
-		if (TSharedPtr<SRemoteControlPanel> Panel = GetPanelForProperty(PropInfo.PropertyArgs))
+		TSharedPtr<SRemoteControlPanel> Panel = GetPanelForProperty(PropInfo.PropertyArgs);
+		if (Panel.IsValid() && Panel->IsExposed(PropInfo.PropertyArgs))
 		{
-			if (Panel->IsExposed(PropInfo.PropertyArgs))
-			{
-				OnToggleExposeSubProperty(PropInfo.PropertyArgs, PropInfo.DesiredName);
-			}
+			OnToggleExposeSubProperty(PropInfo.PropertyArgs, PropInfo.DesiredName);
 		}
 	}
 }
@@ -935,7 +897,7 @@ FRemoteControlUIModule::EPropertyExposeStatus FRemoteControlUIModule::GetPropert
 }
 
 
-FSlateIcon FRemoteControlUIModule::OnGetOverrideMaterialsIcon(const FRCExposesPropertyArgs& InPropertyArgs) const
+FSlateIcon FRemoteControlUIModule::OnGetOverrideMaterialsIcon(const FRCExposesPropertyArgs InPropertyArgs) const
 {
 	FName BrushName("NoBrush");
 
@@ -1232,37 +1194,51 @@ void FRemoteControlUIModule::RegisterWidgetFactories()
 	RegisterWidgetFactoryForType(FRemoteControlInstanceMaterial::StaticStruct(), FOnGenerateRCWidget::CreateStatic(&SRCPanelExposedField::MakeInstance));
 }
 
-FText FRemoteControlUIModule::GetExposePropertyButtonTooltip(const FRCExposesPropertyArgs InPropertyArgs) const
+FText FRemoteControlUIModule::GetPropertyActionTooltip(const FRCExposesPropertyArgs InPropertyArgs) const
 {
 	if (const TSharedPtr<SRemoteControlPanel> Panel = GetPanelForProperty(InPropertyArgs))
 	{
 		if (const URemoteControlPreset* Preset = Panel->GetPreset())
 		{
 			const FText PresetName = FText::FromString(Preset->GetName());
-			if (GetPropertyExposeStatus(InPropertyArgs) == EPropertyExposeStatus::Exposed)
+
+			if (Panel->IsModeActive(ERCPanelMode::Signature))
 			{
-				return FText::Format(LOCTEXT("ExposePropertyToolTip", "Unexpose this property from RemoteControl Preset '{0}'."), PresetName);
+				return FText::Format(LOCTEXT("SignaturePropertyToolTip", "Add this property to the selected or a new signature in RemoteControl Preset '{0}'."), PresetName);
 			}
-			else
+
+			if (Panel->IsExposed(InPropertyArgs))
 			{
-				return FText::Format(LOCTEXT("UnexposePropertyToolTip", "Expose this property in RemoteControl Preset '{0}'."), PresetName);
+				return FText::Format(LOCTEXT("ExposePropertyToolTip", "Unexpose this property from RemoteControl Preset '{0}'."), PresetName);
 			}
+
+			return FText::Format(LOCTEXT("UnexposePropertyToolTip", "Expose this property in RemoteControl Preset '{0}'."), PresetName);
 		}
 	}
 
 	return LOCTEXT("InvalidExposePropertyTooltip", "Invalid Preset");
 }
 
-FText FRemoteControlUIModule::GetExposePropertyButtonText(const FRCExposesPropertyArgs InPropertyArgs) const
+FText FRemoteControlUIModule::GetPropertyActionText(const FRCExposesPropertyArgs InPropertyArgs) const
 {
-	if (GetPropertyExposeStatus(InPropertyArgs) == EPropertyExposeStatus::Exposed)
-	{
-		return LOCTEXT("ExposePropertyText", "Unexpose property");
-	}
-	else
+	TSharedPtr<SRemoteControlPanel> Panel = GetPanelForProperty(InPropertyArgs);
+	if (Panel.IsValid() && Panel->GetPreset())
 	{
-		return LOCTEXT("UnexposePropertyText", "Expose property");
+		if (Panel->IsModeActive(ERCPanelMode::Signature))
+		{
+			return LOCTEXT("AddSignaturePropertyText", "Add to Signature");
+		}
+
+		if (Panel->IsExposed(InPropertyArgs))
+		{
+			return LOCTEXT("ExposePropertyText", "Unexpose property");
+		}
+		else
+		{
+			return LOCTEXT("UnexposePropertyText", "Expose property");
+		}
 	}
+	return FText::GetEmpty();
 }
 
 void FRemoteControlUIModule::TryOverridingMaterials(const FRCExposesPropertyArgs InPropertyArgs)
@@ -1428,9 +1404,9 @@ void FRemoteControlUIModule::RefreshPanels()
 	}
 }
 
-bool FRemoteControlUIModule::ShouldCreateSubMenuForChildProperties(const FRCExposesPropertyArgs InPropertyArgs) const
+bool FRemoteControlUIModule::ShouldCreateSubMenuForChildProperties(const TSharedPtr<SRemoteControlPanel>& InPanel, const FRCExposesPropertyArgs InPropertyArgs) const
 {
-	if (InPropertyArgs.PropertyHandle.IsValid() && FSlateApplication::Get().GetModifierKeys().IsControlDown())
+	if (!InPanel->IsModeActive(ERCPanelMode::Signature) && InPropertyArgs.PropertyHandle.IsValid() && FSlateApplication::Get().GetModifierKeys().IsControlDown())
 	{
 		return HasChildProperties(InPropertyArgs.GetProperty());
 	}
@@ -1504,9 +1480,9 @@ void FRemoteControlUIModule::CreateSubMenuForChildProperties(const FRCExposesPro
 			FSlateIcon(),
 			FUIAction(
 				FExecuteAction::CreateRaw(this, &FRemoteControlUIModule::OnToggleExposeSubProperty, PropArgs.PropertyArgs, PropArgs.DesiredName),
-				FCanExecuteAction::CreateRaw(this, &FRemoteControlUIModule::CanToggleExposeProperty, PropArgs.PropertyArgs),
+				FCanExecuteAction::CreateRaw(this, &FRemoteControlUIModule::CanExecutePropertyAction, PropArgs.PropertyArgs),
 				FGetActionCheckState::CreateRaw(this, &FRemoteControlUIModule::GetPropertyExposedCheckState, PropArgs.PropertyArgs),
-				FIsActionButtonVisible::CreateRaw(this, &FRemoteControlUIModule::CanToggleExposeProperty, PropArgs.PropertyArgs)),
+				FIsActionButtonVisible::CreateRaw(this, &FRemoteControlUIModule::CanExecutePropertyAction, PropArgs.PropertyArgs)),
 			NAME_None,
 			EUserInterfaceActionType::ToggleButton);
 	}
@@ -1521,7 +1497,9 @@ void FRemoteControlUIModule::CreateSubMenuForChildProperties(const FRCExposesPro
 
 bool FRemoteControlUIModule::HasChildPropertiesExposed(const FRCExposesPropertyArgs& InPropertyArgs) const
 {
-	if (InPropertyArgs.PropertyHandle.IsValid())
+	TSharedPtr<SRemoteControlPanel> Panel = GetPanelForProperty(InPropertyArgs);
+
+	if (Panel.IsValid() && InPropertyArgs.PropertyHandle.IsValid())
 	{
 		if (FProperty* Property = InPropertyArgs.PropertyHandle->GetProperty())
 		{
@@ -1537,9 +1515,8 @@ bool FRemoteControlUIModule::HasChildPropertiesExposed(const FRCExposesPropertyA
 					ChildArgs.Property = *It;
 					ChildArgs.PropertyPath = (*It)->GetPathName(nullptr);
 					ChildArgs.PropertyHandle = InPropertyArgs.PropertyHandle->GetChildHandle(ChildHandleIndex++);
-					const EPropertyExposeStatus ChildStatus = GetPropertyExposeStatus(ChildArgs);
 
-					if (ChildStatus == EPropertyExposeStatus::Exposed)
+					if (Panel->IsExposed(ChildArgs))
 					{
 						return true;
 					}
diff --git a/Engine/Plugins/VirtualProduction/RemoteControl/Source/RemoteControlUI/Private/RemoteControlUIModule.h b/Engine/Plugins/VirtualProduction/RemoteControl/Source/RemoteControlUI/Private/RemoteControlUIModule.h
--- a/Engine/Plugins/VirtualProduction/RemoteControl/Source/RemoteControlUI/Private/RemoteControlUIModule.h
+++ b/Engine/Plugins/VirtualProduction/RemoteControl/Source/RemoteControlUI/Private/RemoteControlUIModule.h
@@ -145,16 +145,16 @@ private:
 	TSharedPtr<SRemoteControlPanel> GetPanelForPropertyChangeEvent(const FPropertyChangedEvent& InPropertyChangeEvent) const;
 
 	/** Handle getting the icon displayed in the row extension. */
-	FSlateIcon OnGetExposedIcon(const FRCExposesPropertyArgs& InArgs) const;
+	FSlateIcon OnGetPropertyActionIcon(const FRCExposesPropertyArgs InArgs) const;
 
 	/** Handle getting the expose button visibility. */
-	bool CanToggleExposeProperty(const FRCExposesPropertyArgs InArgs) const;
+	bool CanExecutePropertyAction(const FRCExposesPropertyArgs InArgs) const;
 
 	/** Is the property currently exposed? */
 	ECheckBoxState GetPropertyExposedCheckState(const FRCExposesPropertyArgs InArgs) const;
 
 	/** Handle clicking the expose button. */
-	void OnToggleExposeProperty(const FRCExposesPropertyArgs InArgs) const;
+	void ExecutePropertyAction(const FRCExposesPropertyArgs InArgs) const;
 
 	/** Handle clicking the expose SubProperty button. */
 	void OnToggleExposeSubProperty(const FRCExposesPropertyArgs InArgs, const FString InDesiredName = TEXT("")) const;
@@ -180,7 +180,7 @@ private:
 	EPropertyExposeStatus GetPropertyExposeStatus(const FRCExposesPropertyArgs& InArgs) const;
 
 	/** Handle getting the icon displayed in the row extension. */
-	FSlateIcon OnGetOverrideMaterialsIcon(const FRCExposesPropertyArgs& InArgs) const;
+	FSlateIcon OnGetOverrideMaterialsIcon(const FRCExposesPropertyArgs InArgs) const;
 
 	/** Handle getting the override materials button visibility. */
 	bool IsStaticOrSkeletalMaterialProperty(const FRCExposesPropertyArgs InArgs) const;
@@ -211,10 +211,10 @@ private:
 	void RegisterWidgetFactories();
 
 	/** Returns expose button tooltip based on exposed state */
-	FText GetExposePropertyButtonTooltip(const FRCExposesPropertyArgs InArgs) const;
+	FText GetPropertyActionTooltip(const FRCExposesPropertyArgs InArgs) const;
 
 	/** Returns expose button text based on exposed state */
-	FText GetExposePropertyButtonText(const FRCExposesPropertyArgs InArgs) const;
+	FText GetPropertyActionText(const FRCExposesPropertyArgs InArgs) const;
 
 	/** Attempts to replace the static or skeletal materials with their corressponding overrides. */
 	void TryOverridingMaterials(const FRCExposesPropertyArgs InArgs);
@@ -226,7 +226,7 @@ private:
 	void RefreshPanels();
 
 	/** Check if a sub menu should be created when ctrl+click on the eyeball icon in the details view */
-	bool ShouldCreateSubMenuForChildProperties(const FRCExposesPropertyArgs InPropertyArgs) const;
+	bool ShouldCreateSubMenuForChildProperties(const TSharedPtr<SRemoteControlPanel>& InPanel, const FRCExposesPropertyArgs InPropertyArgs) const;
 
 	/** Check if the property has child properties that can be exposed */
 	bool HasChildProperties(const FProperty* InProperty) const;
diff --git a/Engine/Plugins/VirtualProduction/RemoteControl/Source/RemoteControlUI/Private/UI/BaseLogicUI/SRCLogicPanelListBase.h b/Engine/Plugins/VirtualProduction/RemoteControl/Source/RemoteControlUI/Private/UI/BaseLogicUI/SRCLogicPanelListBase.h
--- a/Engine/Plugins/VirtualProduction/RemoteControl/Source/RemoteControlUI/Private/UI/BaseLogicUI/SRCLogicPanelListBase.h
+++ b/Engine/Plugins/VirtualProduction/RemoteControl/Source/RemoteControlUI/Private/UI/BaseLogicUI/SRCLogicPanelListBase.h
@@ -101,7 +101,7 @@ protected:
 			{
 				// Remove Model from Data Container
 				const int32 RemoveCount = RemoveModel(SelectedItem);
-				if (ensure(RemoveCount > 0))
+				if (RemoveCount > 0)
 				{
 					bIsDeleted = true;
 
diff --git a/Engine/Plugins/VirtualProduction/RemoteControl/Source/RemoteControlUI/Private/UI/SRemoteControlPanel.cpp b/Engine/Plugins/VirtualProduction/RemoteControl/Source/RemoteControlUI/Private/UI/SRemoteControlPanel.cpp
--- a/Engine/Plugins/VirtualProduction/RemoteControl/Source/RemoteControlUI/Private/UI/SRemoteControlPanel.cpp
+++ b/Engine/Plugins/VirtualProduction/RemoteControl/Source/RemoteControlUI/Private/UI/SRemoteControlPanel.cpp
@@ -576,13 +576,19 @@ bool SRemoteControlPanel::IsAllObjectsExposed(TArray<UObject*> InOuterObjects, c
 	return bAllObjectsExposed;
 }
 
-void SRemoteControlPanel::ToggleProperty(const FRCExposesPropertyArgs& InPropertyArgs, FString InDesiredName)
+void SRemoteControlPanel::ExecutePropertyAction(const FRCExposesPropertyArgs& InPropertyArgs, const FString& InDesiredName)
 {
 	if (!ensure(InPropertyArgs.IsValid()))
 	{
 		return;
 	}
 
+	if (IsModeActive(ERCPanelMode::Signature))
+	{
+		SignaturePanel->AddToSignature(InPropertyArgs);
+		return;
+	}
+
 	if (IsExposed(InPropertyArgs))
 	{
 		FScopedTransaction Transaction(LOCTEXT("UnexposeProperty", "Unexpose Property"));
@@ -1199,18 +1205,7 @@ TSharedRef<SWidget> SRemoteControlPanel::BuildLiveModeContent(const TSharedRef<S
 
 TSharedRef<SWidget> SRemoteControlPanel::BuildSignaturesModeContent()
 {
-	return SNew(SSplitter)
-		.Orientation(Orient_Horizontal)
-		+SSplitter::Slot()
-		.Value(0.4)
-		[
-			EntityList.ToSharedRef()
-		]
-		+SSplitter::Slot()
-		.Value(0.6)
-		[
-			SignaturePanel.ToSharedRef()
-		];
+	return SignaturePanel.ToSharedRef();
 }
 
 void SRemoteControlPanel::BindRemoteControlCommands()
diff --git a/Engine/Plugins/VirtualProduction/RemoteControl/Source/RemoteControlUI/Private/UI/SRemoteControlPanel.h b/Engine/Plugins/VirtualProduction/RemoteControl/Source/RemoteControlUI/Private/UI/SRemoteControlPanel.h
--- a/Engine/Plugins/VirtualProduction/RemoteControl/Source/RemoteControlUI/Private/UI/SRemoteControlPanel.h
+++ b/Engine/Plugins/VirtualProduction/RemoteControl/Source/RemoteControlUI/Private/UI/SRemoteControlPanel.h
@@ -125,7 +125,7 @@ public:
 	 * @param InArgs The extension arguments of the property to toggle.
 	 * @param InDesiredName Desired Name for this property if leaved empty RC will deduce it itself
 	 */
-	void ToggleProperty(const FRCExposesPropertyArgs& InArgs, FString InDesiredName = TEXT(""));
+	void ExecutePropertyAction(const FRCExposesPropertyArgs& InArgs, const FString& InDesiredName = TEXT(""));
 
 	/**
 	 * Get the selected group.
diff --git a/Engine/Plugins/VirtualProduction/RemoteControl/Source/RemoteControlUI/Private/UI/Signature/Items/RCSignatureTreeFieldItem.cpp b/Engine/Plugins/VirtualProduction/RemoteControl/Source/RemoteControlUI/Private/UI/Signature/Items/RCSignatureTreeFieldItem.cpp
--- a/Engine/Plugins/VirtualProduction/RemoteControl/Source/RemoteControlUI/Private/UI/Signature/Items/RCSignatureTreeFieldItem.cpp
+++ b/Engine/Plugins/VirtualProduction/RemoteControl/Source/RemoteControlUI/Private/UI/Signature/Items/RCSignatureTreeFieldItem.cpp
@@ -14,6 +14,11 @@ FRCSignatureTreeFieldItem::FRCSignatureTreeFieldItem(int32 InFieldIndex, const T
 {
 }
 
+void FRCSignatureTreeFieldItem::BuildPathSegment(FStringBuilderBase& InBuilder) const
+{
+	InBuilder << FieldIndex;
+}
+
 TOptional<bool> FRCSignatureTreeFieldItem::IsEnabled() const
 {
 	if (const FRCSignatureField* Field = FindField())
diff --git a/Engine/Plugins/VirtualProduction/RemoteControl/Source/RemoteControlUI/Private/UI/Signature/Items/RCSignatureTreeFieldItem.h b/Engine/Plugins/VirtualProduction/RemoteControl/Source/RemoteControlUI/Private/UI/Signature/Items/RCSignatureTreeFieldItem.h
--- a/Engine/Plugins/VirtualProduction/RemoteControl/Source/RemoteControlUI/Private/UI/Signature/Items/RCSignatureTreeFieldItem.h
+++ b/Engine/Plugins/VirtualProduction/RemoteControl/Source/RemoteControlUI/Private/UI/Signature/Items/RCSignatureTreeFieldItem.h
@@ -14,7 +14,9 @@ class FRCSignatureTreeFieldItem : public FRCSignatureTreeItemBase
 public:
 	explicit FRCSignatureTreeFieldItem(int32 InFieldIndex, const TSharedPtr<SRCSignatureTree>& InSignatureTree);
 
+protected:
 	//~ Begin FRCSignatureTreeItemBase
+	virtual void BuildPathSegment(FStringBuilderBase& InStringBuilder) const override;
 	virtual TOptional<bool> IsEnabled() const override;
 	virtual void SetEnabled(bool bInEnabled) override;
 	virtual FText GetDisplayNameText() const override;
diff --git a/Engine/Plugins/VirtualProduction/RemoteControl/Source/RemoteControlUI/Private/UI/Signature/Items/RCSignatureTreeItemBase.cpp b/Engine/Plugins/VirtualProduction/RemoteControl/Source/RemoteControlUI/Private/UI/Signature/Items/RCSignatureTreeItemBase.cpp
--- a/Engine/Plugins/VirtualProduction/RemoteControl/Source/RemoteControlUI/Private/UI/Signature/Items/RCSignatureTreeItemBase.cpp
+++ b/Engine/Plugins/VirtualProduction/RemoteControl/Source/RemoteControlUI/Private/UI/Signature/Items/RCSignatureTreeItemBase.cpp
@@ -9,15 +9,108 @@ FRCSignatureTreeItemBase::FRCSignatureTreeItemBase(const TSharedPtr<SRCSignature
 {
 }
 
+void FRCSignatureTreeItemBase::AddFlags(ERCSignatureTreeItemFlags InFlags)
+{
+	EnumAddFlags(Flags, InFlags);
+}
+
+void FRCSignatureTreeItemBase::RemoveFlags(ERCSignatureTreeItemFlags InFlags)
+{
+	EnumRemoveFlags(Flags, InFlags);
+}
+
+bool FRCSignatureTreeItemBase::HasAnyFlags(ERCSignatureTreeItemFlags InFlags) const
+{
+	return EnumHasAnyFlags(Flags, InFlags);
+}
+
 void FRCSignatureTreeItemBase::RebuildChildren()
 {
+	// Save current children in a map for restoring the flags
+	TMap<FName, TSharedPtr<FRCSignatureTreeItemBase>> OldChildren;
+	OldChildren.Reserve(Children.Num());
+	for (const TSharedPtr<FRCSignatureTreeItemBase>& Child : Children)
+	{
+		OldChildren.Add(Child->Path, Child);
+	}
+
 	Children.Reset();
 	GenerateChildren(Children);
 
 	TSharedRef<FRCSignatureTreeItemBase> This = SharedThis(this);
 	for (const TSharedPtr<FRCSignatureTreeItemBase>& Child : Children)
 	{
-		Child->ParentWeak = This;
+		Child->Initialize(This);
+		if (TSharedPtr<FRCSignatureTreeItemBase>* OldChild = OldChildren.Find(Child->Path))
+		{
+			Child->RestoreFrom(*OldChild);
+		}
 		Child->RebuildChildren();
 	}
 }
+
+void FRCSignatureTreeItemBase::VisitChildren(TFunctionRef<bool(const TSharedPtr<FRCSignatureTreeItemBase>&)> InCallable, bool bInRecursive)
+{
+	for (const TSharedPtr<FRCSignatureTreeItemBase>& Child : Children)
+	{
+		if (!InCallable(Child))
+		{
+			break;
+		}
+
+		if (bInRecursive)
+		{
+			Child->VisitChildren(InCallable, /*bRecursive*/true);
+		}
+	}
+}
+
+void FRCSignatureTreeItemBase::Initialize(const TSharedPtr<FRCSignatureTreeItemBase>& InParent)
+{
+	ParentWeak = InParent;
+	Path = BuildPath();
+}
+
+void FRCSignatureTreeItemBase::RestoreFrom(const TSharedPtr<FRCSignatureTreeItemBase>& InOldItem)
+{
+	if (InOldItem.IsValid())
+	{
+		// The only important things in restoration are flags & children.
+		// Children are only restored so they can restore their flags
+		Flags = InOldItem->GetFlags();
+		Children = InOldItem->GetChildren();
+	}
+}
+
+FName FRCSignatureTreeItemBase::BuildPath() const
+{
+	// For now, Signatures would have at most 2 items here, and Fields 3.
+	TArray<const FRCSignatureTreeItemBase*, TInlineAllocator<3>> AncestorItems;
+
+	// Generate Ancestor Path Array
+	{
+		const FRCSignatureTreeItemBase* CurrentItem = this;
+		while (CurrentItem)
+		{
+			AncestorItems.Add(CurrentItem);
+			CurrentItem = CurrentItem->GetParent().Get();
+		}
+	}
+
+	// Expectation is that signature items will generate a path from Guid (32 characters)
+	// with a dot delimiter to the field path which is an index (1-2 characters for the numbers in most use cases)
+	// the expected size of the string is 34 characters... rounded up to a multiple of 8 => 40
+	TStringBuilder<40> PathBuilder;
+
+	// Build Path starting from the farthest Ancestor.
+	for (const FRCSignatureTreeItemBase* Item : ReverseIterate(AncestorItems))
+	{
+		if (PathBuilder.Len() > 0)
+		{
+			PathBuilder << TEXT(".");
+		}
+		Item->BuildPathSegment(PathBuilder);
+	}
+
+	return PathBuilder.ToString();
+}
diff --git a/Engine/Plugins/VirtualProduction/RemoteControl/Source/RemoteControlUI/Private/UI/Signature/Items/RCSignatureTreeItemBase.h b/Engine/Plugins/VirtualProduction/RemoteControl/Source/RemoteControlUI/Private/UI/Signature/Items/RCSignatureTreeItemBase.h
--- a/Engine/Plugins/VirtualProduction/RemoteControl/Source/RemoteControlUI/Private/UI/Signature/Items/RCSignatureTreeItemBase.h
+++ b/Engine/Plugins/VirtualProduction/RemoteControl/Source/RemoteControlUI/Private/UI/Signature/Items/RCSignatureTreeItemBase.h
@@ -6,9 +6,19 @@
 
 class SRCSignatureTree;
 
+enum ERCSignatureTreeItemFlags : uint8
+{
+	None = 0,
+	Expanded = 1 << 0,
+	Selected = 1 << 1,
+};
+ENUM_CLASS_FLAGS(ERCSignatureTreeItemFlags)
+
 /** Base class for any Item represented in the Signature Tree */
 class FRCSignatureTreeItemBase : public FRCLogicModeBase
 {
+	friend class FRCSignatureTreeRootItem;
+
 public:
 	explicit FRCSignatureTreeItemBase(const TSharedPtr<SRCSignatureTree>& InSignatureTree);
 
@@ -21,6 +31,17 @@ public:
 	{
 	}
 
+	ERCSignatureTreeItemFlags GetFlags() const
+	{
+		return Flags;
+	}
+
+	void AddFlags(ERCSignatureTreeItemFlags InFlags);
+
+	void RemoveFlags(ERCSignatureTreeItemFlags InFlags);
+
+	bool HasAnyFlags(ERCSignatureTreeItemFlags InFlags) const;
+
 	virtual FText GetDisplayNameText() const = 0;
 
 	virtual bool CanEditDisplayNameText() const
@@ -58,15 +79,31 @@ public:
 
 	void RebuildChildren();
 
+	void VisitChildren(TFunctionRef<bool(const TSharedPtr<FRCSignatureTreeItemBase>&)> InCallable, bool bInRecursive);
+
 protected:
+	virtual void BuildPathSegment(FStringBuilderBase& InBuilder) const = 0;
+
 	virtual void GenerateChildren(TArray<TSharedPtr<FRCSignatureTreeItemBase>>& OutChildren) const
 	{
 	}
 
 private:
+	void Initialize(const TSharedPtr<FRCSignatureTreeItemBase>& InParent);
+
+	void RestoreFrom(const TSharedPtr<FRCSignatureTreeItemBase>& InOldItem);
+
+	/** Builds the path from the root to the item. Each item will be its own segment delimited by a dot */
+	FName BuildPath() const;
+
+	/** Unique path from the root to the item. Used to identify items in the Signature Tree */
+	FName Path;
+
 	TArray<TSharedPtr<FRCSignatureTreeItemBase>> Children;
 
 	TWeakPtr<FRCSignatureTreeItemBase> ParentWeak;
 
 	TWeakPtr<SRCSignatureTree> SignatureTreeWeak;
+
+	ERCSignatureTreeItemFlags Flags = ERCSignatureTreeItemFlags::Expanded;
 };
diff --git a/Engine/Plugins/VirtualProduction/RemoteControl/Source/RemoteControlUI/Private/UI/Signature/Items/RCSignatureTreeRootItem.cpp b/Engine/Plugins/VirtualProduction/RemoteControl/Source/RemoteControlUI/Private/UI/Signature/Items/RCSignatureTreeRootItem.cpp
new file mode 100644
--- /dev/null
+++ b/Engine/Plugins/VirtualProduction/RemoteControl/Source/RemoteControlUI/Private/UI/Signature/Items/RCSignatureTreeRootItem.cpp
@@ -0,0 +1,53 @@
+// Copyright Epic Games, Inc. All Rights Reserved.
+
+#include "RCSignatureTreeRootItem.h"
+#include "RCSignatureTreeSignatureItem.h"
+#include "RemoteControlSignature.h"
+#include "RemoteControlSignatureRegistry.h"
+#include "UI/Signature/SRCSignatureTree.h"
+
+FRCSignatureTreeRootItem::FRCSignatureTreeRootItem(const TSharedPtr<SRCSignatureTree>& InSignatureTree)
+	: FRCSignatureTreeItemBase(InSignatureTree)
+{
+}
+
+void FRCSignatureTreeRootItem::GenerateChildren(TArray<TSharedPtr<FRCSignatureTreeItemBase>>& OutChildren) const
+{
+	TSharedPtr<SRCSignatureTree> SignatureTree = GetSignatureTree();
+	if (!SignatureTree.IsValid())
+	{
+		return;
+	}
+
+	URemoteControlSignatureRegistry* SignatureRegistry = SignatureTree->GetSignatureRegistry();
+	if (!SignatureRegistry)
+	{
+		return;
+	}
+
+	TConstArrayView<FRCSignature> Signatures = SignatureRegistry->GetSignatures();
+	OutChildren.Reserve(OutChildren.Num() + Signatures.Num());
+
+	for (const FRCSignature& Signature : Signatures)
+	{
+		OutChildren.Add(MakeShared<FRCSignatureTreeSignatureItem>(Signature, SignatureTree));
+	}
+}
+
+FText FRCSignatureTreeRootItem::GetDisplayNameText() const
+{
+	checkNoEntry();
+	return FText::GetEmpty();
+}
+
+FText FRCSignatureTreeRootItem::GetDescription() const
+{
+	checkNoEntry();
+	return FText::GetEmpty();
+}
+
+int32 FRCSignatureTreeRootItem::RemoveFromRegistry()
+{
+	checkNoEntry();
+	return 0;
+}
diff --git a/Engine/Plugins/VirtualProduction/RemoteControl/Source/RemoteControlUI/Private/UI/Signature/Items/RCSignatureTreeRootItem.h b/Engine/Plugins/VirtualProduction/RemoteControl/Source/RemoteControlUI/Private/UI/Signature/Items/RCSignatureTreeRootItem.h
new file mode 100644
--- /dev/null
+++ b/Engine/Plugins/VirtualProduction/RemoteControl/Source/RemoteControlUI/Private/UI/Signature/Items/RCSignatureTreeRootItem.h
@@ -0,0 +1,29 @@
+// Copyright Epic Games, Inc. All Rights Reserved.
+
+#pragma once
+
+#include "RCSignatureTreeItemBase.h"
+
+/**
+ * Item class for the Root of all items
+ * This is primarily so that the Top level items can have the same set of functionalities as the rest of the items in the tree 
+ */
+class FRCSignatureTreeRootItem : public FRCSignatureTreeItemBase
+{
+public:
+	explicit FRCSignatureTreeRootItem(const TSharedPtr<SRCSignatureTree>& InSignatureTree);
+
+	TArray<TSharedPtr<FRCSignatureTreeItemBase>>& GetChildrenMutable()
+	{
+		return Children;
+	}
+
+protected:
+	//~ Begin FRCSignatureTreeItemBase
+	virtual void BuildPathSegment(FStringBuilderBase& InBuilder) const override {}
+	virtual void GenerateChildren(TArray<TSharedPtr<FRCSignatureTreeItemBase>>& OutChildren) const override;
+	virtual FText GetDisplayNameText() const override;
+	virtual FText GetDescription() const override;
+	virtual int32 RemoveFromRegistry() override;
+	//~ End FRCSignatureTreeItemBase
+};
diff --git a/Engine/Plugins/VirtualProduction/RemoteControl/Source/RemoteControlUI/Private/UI/Signature/Items/RCSignatureTreeSignatureItem.cpp b/Engine/Plugins/VirtualProduction/RemoteControl/Source/RemoteControlUI/Private/UI/Signature/Items/RCSignatureTreeSignatureItem.cpp
--- a/Engine/Plugins/VirtualProduction/RemoteControl/Source/RemoteControlUI/Private/UI/Signature/Items/RCSignatureTreeSignatureItem.cpp
+++ b/Engine/Plugins/VirtualProduction/RemoteControl/Source/RemoteControlUI/Private/UI/Signature/Items/RCSignatureTreeSignatureItem.cpp
@@ -1,6 +1,7 @@
 // Copyright Epic Games, Inc. All Rights Reserved.
 
 #include "RCSignatureTreeSignatureItem.h"
+#include "IRemoteControlUIModule.h"
 #include "RCSignatureTreeFieldItem.h"
 #include "RemoteControlPreset.h"
 #include "RemoteControlSignatureRegistry.h"
@@ -47,41 +48,21 @@ FRCSignature* FRCSignatureTreeSignatureItem::FindSignatureMutable(URemoteControl
 	return nullptr;
 }
 
-void FRCSignatureTreeSignatureItem::AddFieldEntities(TConstArrayView<FGuid> InFieldEntityIds)
+bool FRCSignatureTreeSignatureItem::AddField(URemoteControlSignatureRegistry* InRegistry, const FRCExposesPropertyArgs& InPropertyArgs)
 {
-	if (InFieldEntityIds.IsEmpty())
-	{
-		return;
-	}
-
-	URemoteControlPreset* Preset = GetPreset();
-	if (!Preset)
-	{
-		return;
-	}
-
-	URemoteControlSignatureRegistry* Registry = GetRegistry();
-
-	FRCSignature* Signature = FindSignatureMutable(Registry);
+	FRCSignature* Signature = FindSignatureMutable(InRegistry);
 	if (!Signature)
 	{
-		return;
+		return false;
 	}
 
-	FScopedTransaction Transaction(LOCTEXT("AddProperties", "Add Properties"));
-	Registry->Modify();
-
-	const int32 AddedFieldCount = Signature->AddFieldsFromEntities(Preset, InFieldEntityIds);
-
-	if (AddedFieldCount == 0)
-	{
-		// Nothing happened, cancel transaction
-		Transaction.Cancel();
-	}
-	else if (TSharedPtr<SRCSignatureTree> SignatureTree = GetSignatureTree())
+	TArray<FRCSignatureField, TInlineAllocator<1>> Fields;
 	{
-		SignatureTree->Refresh();
+		FRCFieldPathInfo PathInfo(InPropertyArgs.PropertyHandle->GeneratePathToProperty(), /*bSkipDuplicates*/true);
+		Fields.Emplace(FRCSignatureField:: CreateField(MoveTemp(PathInfo), InPropertyArgs.OwnerObject.Get(), InPropertyArgs.GetProperty()));
 	}
+
+	return Signature->AddFields(Fields) > 0;
 }
 
 void FRCSignatureTreeSignatureItem::ApplySignature(TConstArrayView<TWeakObjectPtr<AActor>> InActors)
@@ -114,6 +95,11 @@ void FRCSignatureTreeSignatureItem::ApplySignature(TConstArrayView<TWeakObjectPt
 	}
 }
 
+void FRCSignatureTreeSignatureItem::BuildPathSegment(FStringBuilderBase& InBuilder) const
+{
+	SignatureId.AppendString(InBuilder, EGuidFormats::DigitsLower);
+}
+
 TOptional<bool> FRCSignatureTreeSignatureItem::IsEnabled() const
 {
 	if (const FRCSignature* Signature = FindSignature())
diff --git a/Engine/Plugins/VirtualProduction/RemoteControl/Source/RemoteControlUI/Private/UI/Signature/Items/RCSignatureTreeSignatureItem.h b/Engine/Plugins/VirtualProduction/RemoteControl/Source/RemoteControlUI/Private/UI/Signature/Items/RCSignatureTreeSignatureItem.h
--- a/Engine/Plugins/VirtualProduction/RemoteControl/Source/RemoteControlUI/Private/UI/Signature/Items/RCSignatureTreeSignatureItem.h
+++ b/Engine/Plugins/VirtualProduction/RemoteControl/Source/RemoteControlUI/Private/UI/Signature/Items/RCSignatureTreeSignatureItem.h
@@ -9,6 +9,7 @@
 #include "UObject/WeakObjectPtrTemplates.h"
 
 class URemoteControlSignatureRegistry;
+struct FRCExposesPropertyArgs;
 struct FRCSignature;
 
 /** Item class representing an RC Signature */
@@ -25,13 +26,15 @@ public:
 
 	FRCSignature* FindSignatureMutable(URemoteControlSignatureRegistry* InRegistry);
 
+	bool AddField(URemoteControlSignatureRegistry* InRegistry, const FRCExposesPropertyArgs& InPropertyArgs);
+
 	//~ Begin IRCSignatureItem
-	virtual void AddFieldEntities(TConstArrayView<FGuid> InFieldEntityIds) override;
 	virtual void ApplySignature(TConstArrayView<TWeakObjectPtr<AActor>> InActors) override;
 	//~ End IRCSignatureItem
 
 protected:
 	//~ Begin FRCSignatureTreeItem
+	virtual void BuildPathSegment(FStringBuilderBase& InBuilder) const override;
 	virtual TOptional<bool> IsEnabled() const override;
 	virtual void SetEnabled(bool bInEnabled) override;
 	virtual FText GetDisplayNameText() const override;
diff --git a/Engine/Plugins/VirtualProduction/RemoteControl/Source/RemoteControlUI/Private/UI/Signature/SRCSignaturePanel.cpp b/Engine/Plugins/VirtualProduction/RemoteControl/Source/RemoteControlUI/Private/UI/Signature/SRCSignaturePanel.cpp
--- a/Engine/Plugins/VirtualProduction/RemoteControl/Source/RemoteControlUI/Private/UI/Signature/SRCSignaturePanel.cpp
+++ b/Engine/Plugins/VirtualProduction/RemoteControl/Source/RemoteControlUI/Private/UI/Signature/SRCSignaturePanel.cpp
@@ -3,6 +3,7 @@
 #include "SRCSignaturePanel.h"
 #include "Columns/RCSignatureDescriptionColumn.h"
 #include "Columns/RCSignatureLabelColumn.h"
+#include "Items/RCSignatureTreeSignatureItem.h"
 #include "Misc/MessageDialog.h"
 #include "RemoteControlPreset.h"
 #include "RemoteControlSignatureRegistry.h"
@@ -74,6 +75,58 @@ URemoteControlSignatureRegistry* SRCSignaturePanel::GetSignatureRegistry() const
 	return nullptr;
 }
 
+void SRCSignaturePanel::AddToSignature(const FRCExposesPropertyArgs& InPropertyArgs)
+{
+	if (!SignatureTreeView.IsValid())
+	{
+		return;
+	}
+
+	URemoteControlSignatureRegistry* Registry = GetSignatureRegistry();
+	if (!Registry)
+	{
+		return;
+	}
+
+	FScopedTransaction Transaction(LOCTEXT("AddToSignatureTransaction", "Add to Signature"));
+	Registry->Modify();
+
+	bool bFieldsAdded = false;
+
+	TArray<TSharedPtr<FRCSignatureTreeItemBase>> SelectedItems = SignatureTreeView->GetSelectedItems();
+	if (SelectedItems.IsEmpty())
+	{
+		// Make a new signature with temp view model item to add the field
+		TSharedRef<FRCSignatureTreeSignatureItem> SignatureItem = MakeShared<FRCSignatureTreeSignatureItem>(Registry->AddSignature(), SignatureTreeView);
+		if (SignatureItem->AddField(Registry, InPropertyArgs))
+		{
+			bFieldsAdded = true;
+		}
+	}
+	else
+	{
+		for (const TSharedPtr<FRCSignatureTreeItemBase>& Item : SelectedItems)
+		{
+			if (FRCSignatureTreeSignatureItem* SignatureItem = Item->AsSignatureItem())
+			{
+				if (SignatureItem->AddField(Registry, InPropertyArgs))
+				{
+					bFieldsAdded = true;
+				}
+			}
+		}
+	}
+
+	if (bFieldsAdded)
+	{
+		SignatureTreeView->Refresh();
+	}
+	else
+	{
+		Transaction.Cancel();
+	}
+}
+
 bool SRCSignaturePanel::IsListFocused() const
 {
 	return SignatureTreeView.IsValid() && SignatureTreeView->IsListFocused();
diff --git a/Engine/Plugins/VirtualProduction/RemoteControl/Source/RemoteControlUI/Private/UI/Signature/SRCSignaturePanel.h b/Engine/Plugins/VirtualProduction/RemoteControl/Source/RemoteControlUI/Private/UI/Signature/SRCSignaturePanel.h
--- a/Engine/Plugins/VirtualProduction/RemoteControl/Source/RemoteControlUI/Private/UI/Signature/SRCSignaturePanel.h
+++ b/Engine/Plugins/VirtualProduction/RemoteControl/Source/RemoteControlUI/Private/UI/Signature/SRCSignaturePanel.h
@@ -4,9 +4,11 @@
 
 #include "UI/BaseLogicUI/SRCLogicPanelBase.h"
 
+class FRCSignatureTreeItemBase;
 class IRCSignatureColumn;
-class URemoteControlSignatureRegistry;
 class SRCSignatureTree;
+class URemoteControlSignatureRegistry;
+struct FRCExposesPropertyArgs;
 
 class SRCSignaturePanel : public SRCLogicPanelBase
 {
@@ -20,6 +22,8 @@ public:
 	/** Gets the Signature Registry from the RC Preset */
 	URemoteControlSignatureRegistry* GetSignatureRegistry() const;
 
+	void AddToSignature(const FRCExposesPropertyArgs& InPropertyArgs);
+
 	/** Whether this widget currently has focus */
 	bool IsListFocused() const;
 
diff --git a/Engine/Plugins/VirtualProduction/RemoteControl/Source/RemoteControlUI/Private/UI/Signature/SRCSignatureRow.cpp b/Engine/Plugins/VirtualProduction/RemoteControl/Source/RemoteControlUI/Private/UI/Signature/SRCSignatureRow.cpp
--- a/Engine/Plugins/VirtualProduction/RemoteControl/Source/RemoteControlUI/Private/UI/Signature/SRCSignatureRow.cpp
+++ b/Engine/Plugins/VirtualProduction/RemoteControl/Source/RemoteControlUI/Private/UI/Signature/SRCSignatureRow.cpp
@@ -124,13 +124,6 @@ FReply SRCSignatureRow::OnRowAcceptDrop(const FDragDropEvent& InDragDropEvent
 			}
 		}
 
-		// Entity Drag Drop
-		if (TSharedPtr<FExposedEntityDragDrop> ExposedEntityDragDropOp = InDragDropEvent.GetOperationAs<FExposedEntityDragDrop>())
-		{
-			SignatureItem->AddFieldEntities(ExposedEntityDragDropOp->GetSelectedFieldsId());
-			return FReply::Handled();
-		}
-
 		// Actor Drag Drop (Only Actors dragged)
 		if (TSharedPtr<FActorDragDropOp> ActorDragDropOp = InDragDropEvent.GetOperationAs<FActorDragDropOp>())
 		{
diff --git a/Engine/Plugins/VirtualProduction/RemoteControl/Source/RemoteControlUI/Private/UI/Signature/SRCSignatureTree.cpp b/Engine/Plugins/VirtualProduction/RemoteControl/Source/RemoteControlUI/Private/UI/Signature/SRCSignatureTree.cpp
--- a/Engine/Plugins/VirtualProduction/RemoteControl/Source/RemoteControlUI/Private/UI/Signature/SRCSignatureTree.cpp
+++ b/Engine/Plugins/VirtualProduction/RemoteControl/Source/RemoteControlUI/Private/UI/Signature/SRCSignatureTree.cpp
@@ -3,6 +3,7 @@
 #include "SRCSignatureTree.h"
 #include "Columns/IRCSignatureColumn.h"
 #include "Items/RCSignatureTreeItemBase.h"
+#include "Items/RCSignatureTreeRootItem.h"
 #include "Items/RCSignatureTreeSignatureItem.h"
 #include "RemoteControlPreset.h"
 #include "RemoteControlSignatureRegistry.h"
@@ -19,9 +20,10 @@
 void SRCSignatureTree::Construct(const FArguments& InArgs, const TSharedRef<SRCSignaturePanel>& InSignaturePanel, const TSharedRef<SRemoteControlPanel>& InRCPanel)
 {
 	SRCLogicPanelListBase::Construct(SRCLogicPanelListBase::FArguments(), InSignaturePanel, InRCPanel);
-
 	SignaturePanelWeak = InSignaturePanel;
 
+	RootItem = MakeShared<FRCSignatureTreeRootItem>(SharedThis(this));
+
 	const FRCPanelStyle* RCPanelStyle = &FRemoteControlPanelStyle::Get()->GetWidgetStyle<FRCPanelStyle>("RemoteControlPanel.LogicControllersPanel");
 
 	HeaderRow = SNew(SHeaderRow)
@@ -34,10 +36,12 @@ void SRCSignatureTree::Construct(const FArguments& InArgs, const TSharedRef<SRCS
 	ChildSlot
 	[
 		SAssignNew(SignatureTreeView, STreeView<TSharedPtr<FRCSignatureTreeItemBase>>)
-		.TreeItemsSource(&SignatureItems)
+		.TreeItemsSource(&RootItem->GetChildrenMutable())
 		.HeaderRow(HeaderRow)
 		.OnGetChildren(this, &SRCSignatureTree::OnGetChildren)
 		.OnGenerateRow(this, &SRCSignatureTree::OnGenerateRow)
+		.OnExpansionChanged(this, &SRCSignatureTree::OnItemExpansionChanged)
+		.OnSelectionChanged(this, &SRCSignatureTree::OnItemSelectionChanged)
 		.OnContextMenuOpening(this, &SRCLogicPanelListBase::GetContextMenuWidget)
 		.SelectionMode(ESelectionMode::Multi)
 		.HighlightParentNodesForSelection(true)
@@ -72,29 +76,30 @@ TSharedPtr<IRCSignatureColumn> SRCSignatureTree::FindColumn(FName InColumnName)
 
 void SRCSignatureTree::Refresh()
 {
-	if (URemoteControlSignatureRegistry* SignatureRegistry = GetSignatureRegistry())
+	if (bRefreshing)
 	{
-		TConstArrayView<FRCSignature> Signatures = SignatureRegistry->GetSignatures();
-
-		SignatureItems.Reset(Signatures.Num());
+		return;
+	}
 
-		TSharedPtr<SRCSignatureTree> SignatureTree = SharedThis(this);
+	TGuardValue<bool> RefreshGuard(bRefreshing, true);
 
-		for (const FRCSignature& Signature : Signatures)
+	RootItem->RebuildChildren();
+	RootItem->VisitChildren([&TreeView = SignatureTreeView](const TSharedPtr<FRCSignatureTreeItemBase>& InItem)->bool
 		{
-			TSharedRef<FRCSignatureTreeSignatureItem> SignatureItem = MakeShared<FRCSignatureTreeSignatureItem>(Signature, SignatureTree);
-			SignatureItem->RebuildChildren();
-			SignatureItems.Emplace(MoveTemp(SignatureItem));
+			TreeView->SetItemExpansion(InItem, InItem->HasAnyFlags(ERCSignatureTreeItemFlags::Expanded));
+			TreeView->SetItemSelection(InItem, InItem->HasAnyFlags(ERCSignatureTreeItemFlags::Selected));
+			return true;
 		}
-	}
-	else
-	{
-		SignatureItems.Empty();
-	}
+		, /*bRecursive*/true);
 
 	SignatureTreeView->RequestTreeRefresh();
 }
 
+TArray<TSharedPtr<FRCSignatureTreeItemBase>> SRCSignatureTree::GetSelectedItems() const
+{
+	return SignatureTreeView->GetSelectedItems();
+}
+
 URemoteControlPreset* SRCSignatureTree::GetPreset()
 {
 	if (TSharedPtr<SRCSignaturePanel> SignaturePanel = SignaturePanelWeak.Pin())
@@ -106,14 +111,12 @@ URemoteControlPreset* SRCSignatureTree::GetPreset()
 
 TArray<TSharedPtr<FRCLogicModeBase>> SRCSignatureTree::GetSelectedLogicItems()
 {
-	TArray<TSharedPtr<FRCLogicModeBase>> SelectedItems;
-	SelectedItems.Append(SignatureTreeView->GetSelectedItems());
-	return SelectedItems;
+	return TArray<TSharedPtr<FRCLogicModeBase>>(GetSelectedItems());
 }
 
 bool SRCSignatureTree::IsEmpty() const
 {
-	return SignatureItems.IsEmpty();
+	return RootItem->GetChildren().IsEmpty();
 }
 
 bool SRCSignatureTree::IsListFocused() const
@@ -123,7 +126,7 @@ bool SRCSignatureTree::IsListFocused() const
 
 int32 SRCSignatureTree::Num() const
 {
-	return SignatureItems.Num();
+	return RootItem->GetChildren().Num();
 }
 
 int32 SRCSignatureTree::NumSelectedLogicItems() const
@@ -149,7 +152,7 @@ void SRCSignatureTree::DeleteSelectedPanelItems()
 	}
 
 	FScopedTransaction Transaction(LOCTEXT("RemoveSelectedSignatures", "Remove Selected Signatures"));
-	DeleteItemsFromLogicPanel<FRCSignatureTreeItemBase>(SignatureItems, SelectedSignatures);
+	DeleteItemsFromLogicPanel<FRCSignatureTreeItemBase>(RootItem->GetChildrenMutable(), SelectedSignatures);
 }
 
 void SRCSignatureTree::Reset()
@@ -184,4 +187,44 @@ void SRCSignatureTree::OnGetChildren(TSharedPtr<FRCSignatureTreeItemBase> InItem
 	}
 }
 
+void SRCSignatureTree::OnItemExpansionChanged(TSharedPtr<FRCSignatureTreeItemBase> InItem, bool bInIsExpanded)
+{
+	if (bRefreshing)
+	{
+		return;
+	}
+
+	if (InItem.IsValid())
+	{
+		if (bInIsExpanded)
+		{
+			InItem->AddFlags(ERCSignatureTreeItemFlags::Expanded);
+		}
+		else
+		{
+			InItem->RemoveFlags(ERCSignatureTreeItemFlags::Expanded);
+		}
+	}
+}
+
+void SRCSignatureTree::OnItemSelectionChanged(TSharedPtr<FRCSignatureTreeItemBase> InItem, ESelectInfo::Type InSelectionType)
+{
+	if (bRefreshing)
+	{
+		return;
+	}
+
+	RootItem->VisitChildren([](const TSharedPtr<FRCSignatureTreeItemBase>& InItem)->bool
+		{
+			InItem->RemoveFlags(ERCSignatureTreeItemFlags::Selected);
+			return true;
+		}
+		, /*bRecursive*/true);
+
+	for (const TSharedPtr<FRCSignatureTreeItemBase>& SelectedItem : SignatureTreeView->GetSelectedItems())
+	{
+		SelectedItem->AddFlags(ERCSignatureTreeItemFlags::Selected);
+	}
+}
+
 #undef LOCTEXT_NAMESPACE
diff --git a/Engine/Plugins/VirtualProduction/RemoteControl/Source/RemoteControlUI/Private/UI/Signature/SRCSignatureTree.h b/Engine/Plugins/VirtualProduction/RemoteControl/Source/RemoteControlUI/Private/UI/Signature/SRCSignatureTree.h
--- a/Engine/Plugins/VirtualProduction/RemoteControl/Source/RemoteControlUI/Private/UI/Signature/SRCSignatureTree.h
+++ b/Engine/Plugins/VirtualProduction/RemoteControl/Source/RemoteControlUI/Private/UI/Signature/SRCSignatureTree.h
@@ -5,6 +5,7 @@
 #include "UI/BaseLogicUI/SRCLogicPanelListBase.h"
 
 class FRCSignatureTreeItemBase;
+class FRCSignatureTreeRootItem;
 class IRCSignatureColumn;
 class SHeaderRow;
 class SRCSignaturePanel;
@@ -26,6 +27,8 @@ public:
 
 	void Refresh();
 
+	TArray<TSharedPtr<FRCSignatureTreeItemBase>> GetSelectedItems() const;
+
 	//~ Begin SRCLogicPanelListBase
 	virtual URemoteControlPreset* GetPreset() override;
 	virtual TArray<TSharedPtr<FRCLogicModeBase>> GetSelectedLogicItems() override;
@@ -46,8 +49,11 @@ private:
 
 	void OnGetChildren(TSharedPtr<FRCSignatureTreeItemBase> InItem, TArray<TSharedPtr<FRCSignatureTreeItemBase>>& OutChildren) const;
 
-	/** All Known Signatures in this Tree View */
-	TArray<TSharedPtr<FRCSignatureTreeItemBase>> SignatureItems;
+	void OnItemExpansionChanged(TSharedPtr<FRCSignatureTreeItemBase> InItem, bool bInIsExpanded);
+
+	void OnItemSelectionChanged(TSharedPtr<FRCSignatureTreeItemBase> InItem, ESelectInfo::Type InSelectionType);
+
+	TSharedPtr<FRCSignatureTreeRootItem> RootItem;
 
 	TMap<FName, TSharedRef<IRCSignatureColumn>> Columns;
 
@@ -56,4 +62,6 @@ private:
 	TSharedPtr<SHeaderRow> HeaderRow;
 
 	TWeakPtr<SRCSignaturePanel> SignaturePanelWeak;
+
+	bool bRefreshing = false;
 };
diff --git a/Engine/Plugins/VirtualProduction/RemoteControl/Source/RemoteControlUI/Public/UI/Signature/IRCSignatureItem.h b/Engine/Plugins/VirtualProduction/RemoteControl/Source/RemoteControlUI/Public/UI/Signature/IRCSignatureItem.h
--- a/Engine/Plugins/VirtualProduction/RemoteControl/Source/RemoteControlUI/Public/UI/Signature/IRCSignatureItem.h
+++ b/Engine/Plugins/VirtualProduction/RemoteControl/Source/RemoteControlUI/Public/UI/Signature/IRCSignatureItem.h
@@ -6,13 +6,10 @@
 #include "UObject/WeakObjectPtrTemplatesFwd.h"
 
 class AActor;
-struct FGuid;
 
 /** Interface for external implementations to interact with the Signature View Model */
 class IRCSignatureItem
 {
 public:
-	virtual void AddFieldEntities(TConstArrayView<FGuid> InFieldEntityIds) = 0;
-
 	virtual void ApplySignature(TConstArrayView<TWeakObjectPtr<AActor>> InActors) = 0;
 };
