Remote Control: Changed signature apply func to use UObject instead of AActor + some other small readability changes

diff --git a/Engine/Plugins/Experimental/Avalanche/Source/AvalancheEditor/Private/RemoteControl/AvaRCSignatureCustomization.cpp b/Engine/Plugins/Experimental/Avalanche/Source/AvalancheEditor/Private/RemoteControl/AvaRCSignatureCustomization.cpp
--- a/Engine/Plugins/Experimental/Avalanche/Source/AvalancheEditor/Private/RemoteControl/AvaRCSignatureCustomization.cpp
+++ b/Engine/Plugins/Experimental/Avalanche/Source/AvalancheEditor/Private/RemoteControl/AvaRCSignatureCustomization.cpp
@@ -2,6 +2,7 @@
 
 #include "AvaRCSignatureCustomization.h"
 #include "DragDropOps/AvaOutlinerItemDragDropOp.h"
+#include "GameFramework/Actor.h"
 #include "UI/Signature/IRCSignatureItem.h"
 
 bool FAvaRCSignatureCustomization::CanAcceptDrop(const FDragDropEvent& InDragDropEvent, IRCSignatureItem* InSignatureItem) const
@@ -16,7 +17,9 @@ FReply FAvaRCSignatureCustomization::AcceptDrop(const FDragDropEvent& InDragDrop
 	{
 		TArray<TWeakObjectPtr<AActor>> DragDropActors;
 		OutlinerDragDrop->GetDragDropOpActors(DragDropActors);
-		InSignatureItem->ApplySignature(DragDropActors);
+
+		TArray<TWeakObjectPtr<UObject>> DragDropObjects(MoveTemp(DragDropActors));
+		InSignatureItem->ApplySignature(DragDropObjects);
 		return FReply::Handled();
 	}
 
diff --git a/Engine/Plugins/VirtualProduction/RemoteControl/Source/RemoteControl/Private/RemoteControlSignature.cpp b/Engine/Plugins/VirtualProduction/RemoteControl/Source/RemoteControl/Private/RemoteControlSignature.cpp
--- a/Engine/Plugins/VirtualProduction/RemoteControl/Source/RemoteControl/Private/RemoteControlSignature.cpp
+++ b/Engine/Plugins/VirtualProduction/RemoteControl/Source/RemoteControl/Private/RemoteControlSignature.cpp
@@ -5,6 +5,56 @@
 #include "RemoteControlField.h"
 #include "RemoteControlPreset.h"
 
+namespace UE::RemoteControl::Private
+{
+	TArray<UObject*> ResolveObjects(TConstArrayView<TWeakObjectPtr<UObject>> InObjects)
+	{
+		TArray<UObject*> ResolvedObjects;
+		ResolvedObjects.Reserve(InObjects.Num());
+		for (const TWeakObjectPtr<UObject>& ObjectWeak : InObjects)
+		{
+			if (UObject* Object = ObjectWeak.Get())
+			{
+				ResolvedObjects.Add(Object);
+			}
+		}
+		return ResolvedObjects;
+	}
+
+	UObject* FindContext(const FRCSignatureField& InField, UObject* InOuterObject, UClass* InSupportedClass)
+	{
+		UObject* Context;
+
+		if (InField.ObjectRelativePath.IsEmpty())
+		{
+			Context = InOuterObject;
+		}
+		else
+		{
+			UClass* FindClass = InSupportedClass ? InSupportedClass : UObject::StaticClass();
+			Context = StaticFindObject(FindClass, InOuterObject, *InField.ObjectRelativePath);
+
+			// Slow Path: if Subobject Path did not find the object, try to find the first sub-object of the actor matching the class.
+			if (!Context && InSupportedClass)
+			{
+				ForEachObjectWithOuterBreakable(InOuterObject, [&Context, InSupportedClass](UObject* InSubobject)->bool
+					{
+						UClass* SubobjectClass = InSubobject->GetClass();
+						if (SubobjectClass && SubobjectClass->IsChildOf(InSupportedClass))
+						{
+							Context = InSubobject;
+							return false;
+						}
+						return true;
+					}
+					, /*bIncludeNestedObjects*/true);
+			}
+		}
+
+		return Context;
+	}
+}
+
 FRCSignatureField FRCSignatureField::CreateField(const FRCFieldPathInfo& InFieldPathInfo, UObject* InOwnerObject, FProperty* InProperty)
 {
 	UClass* SupportedBindingClass = nullptr;
@@ -56,25 +106,15 @@ int32 FRCSignature::AddFields(TConstArrayView<FRCSignatureField> InFields)
 	return Fields.Num() - PreviousNum;
 }
 
-int32 FRCSignature::ApplySignature(URemoteControlPreset* InPreset, TConstArrayView<TWeakObjectPtr<AActor>> InActors) const
+int32 FRCSignature::ApplySignature(URemoteControlPreset* InPreset, TConstArrayView<TWeakObjectPtr<UObject>> InObjects) const
 {
-	if (!InPreset || InActors.IsEmpty())
+	if (!InPreset || InObjects.IsEmpty())
 	{
 		return 0;
 	}
 
-	// Resolve Weak Actors
-	TArray<AActor*> ResolvedActors;
-	ResolvedActors.Reserve(InActors.Num());
-	for (const TWeakObjectPtr<AActor>& ActorWeak : InActors)
-	{
-		if (AActor* Actor = ActorWeak.Get())
-		{
-			ResolvedActors.Add(Actor);
-		}
-	}
-
-	if (ResolvedActors.IsEmpty())
+	TArray<UObject*> ResolvedObjects = UE::RemoteControl::Private::ResolveObjects(InObjects);
+	if (ResolvedObjects.IsEmpty())
 	{
 		return false;
 	}
@@ -90,44 +130,31 @@ int32 FRCSignature::ApplySignature(URemoteControlPreset* InPreset, TConstArrayVi
 			continue;
 		}
 
-		// Resolve, not load. It should be loaded already if relevant to the Actors to apply this signature to
+		// Resolve, not load. It should be loaded already if relevant to the Objects to apply this signature to
 		UClass* SupportedClass = Field.SupportedClass.ResolveClass();
-		UClass* FindClass = SupportedClass ? SupportedClass : UObject::StaticClass();
 
+		// Copy field path as it's going to be used to resolve (mutable)
 		FRCFieldPathInfo Path = Field.FieldPath;
-		for (AActor* Actor : ResolvedActors)
+
+		for (UObject* Object : ResolvedObjects)
 		{
-			UObject* Context;
-			if (Field.ObjectRelativePath.IsEmpty())
-			{
-				Context = Actor;
-			}
-			else
-			{
-				Context = StaticFindObject(FindClass, Actor, *Field.ObjectRelativePath);
+			UObject* Context = UE::RemoteControl::Private::FindContext(Field, Object, SupportedClass);
 
-				// Slow Path: if Subobject Path did not find the object, try to find the first sub-object of the actor matching the class.
-				if (!Context && SupportedClass)
-				{
-					ForEachObjectWithOuterBreakable(Actor, [&Context, SupportedClass](UObject* InSubobject)->bool
-						{
-							UClass* SubobjectClass = InSubobject->GetClass();
-							if (SubobjectClass && SubobjectClass->IsChildOf(SupportedClass))
-							{
-								Context = InSubobject;
-								return false;
-							}
-							return true;
-						}
-						, /*bIncludeNestedObjects*/true);
-				}
+			// Attempt to resolve the path from the given context
+			if (!Context || !Path.Resolve(Context))
+			{
+				continue;
 			}
 
-			if (Context && Path.Resolve(Context))
+			// If path resolved, expose the property
+			TSharedPtr<FRemoteControlProperty> Property = InPreset->ExposeProperty(Context, Field.FieldPath, ExposeArgs).Pin();
+			if (!Property.IsValid())
 			{
-				InPreset->ExposeProperty(Context, Field.FieldPath, ExposeArgs);
-				++ExposeCount;
+				continue;
 			}
+
+			// Property was exposed successfully, increase expose count
+			++ExposeCount;
 		}
 	}
 
diff --git a/Engine/Plugins/VirtualProduction/RemoteControl/Source/RemoteControl/Public/RemoteControlSignature.h b/Engine/Plugins/VirtualProduction/RemoteControl/Source/RemoteControl/Public/RemoteControlSignature.h
--- a/Engine/Plugins/VirtualProduction/RemoteControl/Source/RemoteControl/Public/RemoteControlSignature.h
+++ b/Engine/Plugins/VirtualProduction/RemoteControl/Source/RemoteControl/Public/RemoteControlSignature.h
@@ -7,6 +7,7 @@
 
 class URemoteControlPreset;
 
+/** Representation class of a Field (e.g. Property) in an Object */
 USTRUCT()
 struct FRCSignatureField
 {
@@ -60,10 +61,10 @@ struct FRCSignature
 	/**
 	 * Applies this Signature to the given Actors by exposing all this Signature's fields to the given preset
 	 * @param InPreset the preset where these properties will be exposed to
-	 * @param InActors the actors whose properties to expose
+	 * @param InObjects the objects whose properties or properties of its subojects to expose
 	 * @return the number of properties affected in total
 	 */
-	REMOTECONTROL_API int32 ApplySignature(URemoteControlPreset* InPreset, TConstArrayView<TWeakObjectPtr<AActor>> InActors) const;
+	REMOTECONTROL_API int32 ApplySignature(URemoteControlPreset* InPreset, TConstArrayView<TWeakObjectPtr<UObject>> InObjects) const;
 
 	/** User facing friendly name. Used as the Label when exposing */
 	UPROPERTY()
diff --git a/Engine/Plugins/VirtualProduction/RemoteControl/Source/RemoteControlUI/Private/RemoteControlUIModule.cpp b/Engine/Plugins/VirtualProduction/RemoteControl/Source/RemoteControlUI/Private/RemoteControlUIModule.cpp
--- a/Engine/Plugins/VirtualProduction/RemoteControl/Source/RemoteControlUI/Private/RemoteControlUIModule.cpp
+++ b/Engine/Plugins/VirtualProduction/RemoteControl/Source/RemoteControlUI/Private/RemoteControlUIModule.cpp
@@ -1516,7 +1516,7 @@ bool FRemoteControlUIModule::HasChildPropertiesExposed(const FRCExposesPropertyA
 					ChildArgs.PropertyPath = (*It)->GetPathName(nullptr);
 					ChildArgs.PropertyHandle = InPropertyArgs.PropertyHandle->GetChildHandle(ChildHandleIndex++);
 
-					if (Panel->IsExposed(ChildArgs))
+					if (ChildArgs.IsValid() && Panel->IsExposed(ChildArgs))
 					{
 						return true;
 					}
diff --git a/Engine/Plugins/VirtualProduction/RemoteControl/Source/RemoteControlUI/Private/UI/Signature/Items/RCSignatureTreeSignatureItem.cpp b/Engine/Plugins/VirtualProduction/RemoteControl/Source/RemoteControlUI/Private/UI/Signature/Items/RCSignatureTreeSignatureItem.cpp
--- a/Engine/Plugins/VirtualProduction/RemoteControl/Source/RemoteControlUI/Private/UI/Signature/Items/RCSignatureTreeSignatureItem.cpp
+++ b/Engine/Plugins/VirtualProduction/RemoteControl/Source/RemoteControlUI/Private/UI/Signature/Items/RCSignatureTreeSignatureItem.cpp
@@ -65,9 +65,9 @@ bool FRCSignatureTreeSignatureItem::AddField(URemoteControlSignatureRegistry* In
 	return Signature->AddFields(Fields) > 0;
 }
 
-void FRCSignatureTreeSignatureItem::ApplySignature(TConstArrayView<TWeakObjectPtr<AActor>> InActors)
+void FRCSignatureTreeSignatureItem::ApplySignature(TConstArrayView<TWeakObjectPtr<UObject>> InObjects)
 {
-	if (InActors.IsEmpty())
+	if (InObjects.IsEmpty())
 	{
 		return;
 	}
@@ -84,9 +84,9 @@ void FRCSignatureTreeSignatureItem::ApplySignature(TConstArrayView<TWeakObjectPt
 		return;
 	}
 
-	FScopedTransaction Transaction(LOCTEXT("ApplySignatureActors", "Apply Signature to Actors"));
+	FScopedTransaction Transaction(LOCTEXT("ApplySignature", "Apply Signature"));
 
-	int32 ExposedCount = Signature->ApplySignature(Preset, InActors);
+	int32 ExposedCount = Signature->ApplySignature(Preset, InObjects);
 
 	if (ExposedCount == 0)
 	{
diff --git a/Engine/Plugins/VirtualProduction/RemoteControl/Source/RemoteControlUI/Private/UI/Signature/Items/RCSignatureTreeSignatureItem.h b/Engine/Plugins/VirtualProduction/RemoteControl/Source/RemoteControlUI/Private/UI/Signature/Items/RCSignatureTreeSignatureItem.h
--- a/Engine/Plugins/VirtualProduction/RemoteControl/Source/RemoteControlUI/Private/UI/Signature/Items/RCSignatureTreeSignatureItem.h
+++ b/Engine/Plugins/VirtualProduction/RemoteControl/Source/RemoteControlUI/Private/UI/Signature/Items/RCSignatureTreeSignatureItem.h
@@ -29,7 +29,7 @@ public:
 	bool AddField(URemoteControlSignatureRegistry* InRegistry, const FRCExposesPropertyArgs& InPropertyArgs);
 
 	//~ Begin IRCSignatureItem
-	virtual void ApplySignature(TConstArrayView<TWeakObjectPtr<AActor>> InActors) override;
+	virtual void ApplySignature(TConstArrayView<TWeakObjectPtr<UObject>> InObjects) override;
 	//~ End IRCSignatureItem
 
 protected:
diff --git a/Engine/Plugins/VirtualProduction/RemoteControl/Source/RemoteControlUI/Private/UI/Signature/SRCSignatureRow.cpp b/Engine/Plugins/VirtualProduction/RemoteControl/Source/RemoteControlUI/Private/UI/Signature/SRCSignatureRow.cpp
--- a/Engine/Plugins/VirtualProduction/RemoteControl/Source/RemoteControlUI/Private/UI/Signature/SRCSignatureRow.cpp
+++ b/Engine/Plugins/VirtualProduction/RemoteControl/Source/RemoteControlUI/Private/UI/Signature/SRCSignatureRow.cpp
@@ -6,6 +6,7 @@
 #include "DragAndDrop/CompositeDragDropOp.h"
 #include "DragAndDrop/FolderDragDropOp.h"
 #include "EditorActorFolders.h"
+#include "GameFramework/Actor.h"
 #include "Items/RCSignatureTreeItemBase.h"
 #include "Items/RCSignatureTreeSignatureItem.h"
 #include "RemoteControlUIModule.h"
@@ -160,7 +161,8 @@ FReply SRCSignatureRow::OnRowAcceptDrop(const FDragDropEvent& InDragDropEvent
 
 void SRCSignatureRow::HandleActorDragDrop(IRCSignatureItem* InSignatureItem, const FActorDragDropOp& InActorDragDropOp)
 {
-	InSignatureItem->ApplySignature(InActorDragDropOp.Actors);
+	TArray<TWeakObjectPtr<UObject>> DragDropObjects(InActorDragDropOp.Actors);
+	InSignatureItem->ApplySignature(DragDropObjects);
 }
 
 void SRCSignatureRow::HandleFolderDragDrop(IRCSignatureItem* InSignatureItem, const FFolderDragDropOp& InFolderDragDropOp)
@@ -169,6 +171,8 @@ void SRCSignatureRow::HandleFolderDragDrop(IRCSignatureItem* InSignatureItem, co
 	{
 		TArray<TWeakObjectPtr<AActor>> Actors;
 		FActorFolders::GetWeakActorsFromFolders(*World, InFolderDragDropOp.Folders, Actors);
-		InSignatureItem->ApplySignature(Actors);
+
+		TArray<TWeakObjectPtr<UObject>> DragDropObjects(MoveTemp(Actors));
+		InSignatureItem->ApplySignature(DragDropObjects);
 	}
 }
diff --git a/Engine/Plugins/VirtualProduction/RemoteControl/Source/RemoteControlUI/Public/UI/Signature/IRCSignatureItem.h b/Engine/Plugins/VirtualProduction/RemoteControl/Source/RemoteControlUI/Public/UI/Signature/IRCSignatureItem.h
--- a/Engine/Plugins/VirtualProduction/RemoteControl/Source/RemoteControlUI/Public/UI/Signature/IRCSignatureItem.h
+++ b/Engine/Plugins/VirtualProduction/RemoteControl/Source/RemoteControlUI/Public/UI/Signature/IRCSignatureItem.h
@@ -5,11 +5,11 @@
 #include "Containers/ContainersFwd.h"
 #include "UObject/WeakObjectPtrTemplatesFwd.h"
 
-class AActor;
+class UObject;
 
 /** Interface for external implementations to interact with the Signature View Model */
 class IRCSignatureItem
 {
 public:
-	virtual void ApplySignature(TConstArrayView<TWeakObjectPtr<AActor>> InActors) = 0;
+	virtual void ApplySignature(TConstArrayView<TWeakObjectPtr<UObject>> InObjects) = 0;
 };
