Remote Control: Trivial signature label/desc files reorganization to their own directory (for readability)

diff --git a/Engine/Plugins/VirtualProduction/RemoteControl/Source/RemoteControlUI/Private/UI/Signature/Columns/IRCSignatureColumn.h b/Engine/Plugins/VirtualProduction/RemoteControl/Source/RemoteControlUI/Private/UI/Signature/Columns/IRCSignatureColumn.h
deleted file mode 100644
diff --git a/Engine/Plugins/VirtualProduction/RemoteControl/Source/RemoteControlUI/Private/UI/Signature/Columns/RCSignatureDescriptionColumn.cpp b/Engine/Plugins/VirtualProduction/RemoteControl/Source/RemoteControlUI/Private/UI/Signature/Columns/RCSignatureDescriptionColumn.cpp
deleted file mode 100644
diff --git a/Engine/Plugins/VirtualProduction/RemoteControl/Source/RemoteControlUI/Private/UI/Signature/Columns/RCSignatureDescriptionColumn.h b/Engine/Plugins/VirtualProduction/RemoteControl/Source/RemoteControlUI/Private/UI/Signature/Columns/RCSignatureDescriptionColumn.h
deleted file mode 100644
diff --git a/Engine/Plugins/VirtualProduction/RemoteControl/Source/RemoteControlUI/Private/UI/Signature/Columns/RCSignatureLabelColumn.cpp b/Engine/Plugins/VirtualProduction/RemoteControl/Source/RemoteControlUI/Private/UI/Signature/Columns/RCSignatureLabelColumn.cpp
deleted file mode 100644
diff --git a/Engine/Plugins/VirtualProduction/RemoteControl/Source/RemoteControlUI/Private/UI/Signature/Columns/RCSignatureLabelColumn.h b/Engine/Plugins/VirtualProduction/RemoteControl/Source/RemoteControlUI/Private/UI/Signature/Columns/RCSignatureLabelColumn.h
deleted file mode 100644
diff --git a/Engine/Plugins/VirtualProduction/RemoteControl/Source/RemoteControlUI/Private/UI/Signature/Columns/SRCSignatureDescription.cpp b/Engine/Plugins/VirtualProduction/RemoteControl/Source/RemoteControlUI/Private/UI/Signature/Columns/SRCSignatureDescription.cpp
deleted file mode 100644
diff --git a/Engine/Plugins/VirtualProduction/RemoteControl/Source/RemoteControlUI/Private/UI/Signature/Columns/SRCSignatureDescription.h b/Engine/Plugins/VirtualProduction/RemoteControl/Source/RemoteControlUI/Private/UI/Signature/Columns/SRCSignatureDescription.h
deleted file mode 100644
diff --git a/Engine/Plugins/VirtualProduction/RemoteControl/Source/RemoteControlUI/Private/UI/Signature/Columns/SRCSignatureLabel.cpp b/Engine/Plugins/VirtualProduction/RemoteControl/Source/RemoteControlUI/Private/UI/Signature/Columns/SRCSignatureLabel.cpp
deleted file mode 100644
diff --git a/Engine/Plugins/VirtualProduction/RemoteControl/Source/RemoteControlUI/Private/UI/Signature/Columns/SRCSignatureLabel.h b/Engine/Plugins/VirtualProduction/RemoteControl/Source/RemoteControlUI/Private/UI/Signature/Columns/SRCSignatureLabel.h
deleted file mode 100644
diff --git a/Engine/Plugins/VirtualProduction/RemoteControl/Source/RemoteControlUI/Private/UI/Signature/Description/RCSignatureDescriptionColumn.cpp b/Engine/Plugins/VirtualProduction/RemoteControl/Source/RemoteControlUI/Private/UI/Signature/Description/RCSignatureDescriptionColumn.cpp
new file mode 100644
--- /dev/null
+++ b/Engine/Plugins/VirtualProduction/RemoteControl/Source/RemoteControlUI/Private/UI/Signature/Description/RCSignatureDescriptionColumn.cpp
@@ -0,0 +1,32 @@
+// Copyright Epic Games, Inc. All Rights Reserved.
+
+#include "RCSignatureDescriptionColumn.h"
+#include "SRCSignatureDescription.h"
+
+#define LOCTEXT_NAMESPACE "RCSignatureDescriptionColumn"
+
+FName FRCSignatureDescriptionColumn::GetColumnId() const
+{
+	return TEXT("FRCSignatureDescriptionColumn");
+}
+
+bool FRCSignatureDescriptionColumn::ShouldShowColumnByDefault() const
+{
+	return false;
+}
+
+SHeaderRow::FColumn::FArguments FRCSignatureDescriptionColumn::ConstructHeaderRowColumn()
+{
+	return SHeaderRow::Column(GetColumnId())
+		.FillWidth(0.5f)
+		.DefaultLabel(LOCTEXT("DisplayName", "Description"));
+}
+
+TSharedRef<SWidget> FRCSignatureDescriptionColumn::ConstructRowWidget(TSharedPtr<FRCSignatureTreeItemBase> InItem
+	, const TSharedRef<SRCSignatureTree>& InList
+	, const TSharedRef<SRCSignatureRow>& InRow)
+{
+	return SNew(SRCSignatureDescription, InItem.ToSharedRef());
+}
+
+#undef LOCTEXT_NAMESPACE
diff --git a/Engine/Plugins/VirtualProduction/RemoteControl/Source/RemoteControlUI/Private/UI/Signature/Description/RCSignatureDescriptionColumn.h b/Engine/Plugins/VirtualProduction/RemoteControl/Source/RemoteControlUI/Private/UI/Signature/Description/RCSignatureDescriptionColumn.h
new file mode 100644
--- /dev/null
+++ b/Engine/Plugins/VirtualProduction/RemoteControl/Source/RemoteControlUI/Private/UI/Signature/Description/RCSignatureDescriptionColumn.h
@@ -0,0 +1,15 @@
+// Copyright Epic Games, Inc. All Rights Reserved.
+
+#pragma once
+
+#include "UI/Signature/IRCSignatureColumn.h"
+
+class FRCSignatureDescriptionColumn : public IRCSignatureColumn
+{
+	//~ Begin IRCSignatureColumn
+	virtual FName GetColumnId() const override;
+	virtual bool ShouldShowColumnByDefault() const override;
+	virtual SHeaderRow::FColumn::FArguments ConstructHeaderRowColumn() override;
+	virtual TSharedRef<SWidget> ConstructRowWidget(TSharedPtr<FRCSignatureTreeItemBase> InItem, const TSharedRef<SRCSignatureTree>& InList, const TSharedRef<SRCSignatureRow>& InRow) override;
+	//~ End IRCSignatureColumn
+};
diff --git a/Engine/Plugins/VirtualProduction/RemoteControl/Source/RemoteControlUI/Private/UI/Signature/Description/SRCSignatureDescription.cpp b/Engine/Plugins/VirtualProduction/RemoteControl/Source/RemoteControlUI/Private/UI/Signature/Description/SRCSignatureDescription.cpp
new file mode 100644
--- /dev/null
+++ b/Engine/Plugins/VirtualProduction/RemoteControl/Source/RemoteControlUI/Private/UI/Signature/Description/SRCSignatureDescription.cpp
@@ -0,0 +1,29 @@
+// Copyright Epic Games, Inc. All Rights Reserved.
+
+#include "SRCSignatureDescription.h"
+#include "UI/Signature/Items/RCSignatureTreeItemBase.h"
+#include "Widgets/Text/STextBlock.h"
+
+void SRCSignatureDescription::Construct(const FArguments& InArgs, const TSharedRef<FRCSignatureTreeItemBase>& InItem)
+{
+	ItemWeak = InItem;
+
+	ChildSlot
+	.VAlign(VAlign_Center)
+	.Padding(5.f, 0.f, 0.f, 0.f)
+	[
+		SNew(STextBlock)
+		.Text(this, &SRCSignatureDescription::GetDescription)
+		.Justification(ETextJustify::Left)
+	];
+}
+
+FText SRCSignatureDescription::GetDescription() const
+{
+	if (!CachedDescription.IsSet())
+	{
+		TSharedPtr<FRCSignatureTreeItemBase> Item = ItemWeak.Pin();
+		CachedDescription = Item.IsValid() ? Item->GetDescription() : FText::GetEmpty();
+	}
+	return *CachedDescription;
+}
diff --git a/Engine/Plugins/VirtualProduction/RemoteControl/Source/RemoteControlUI/Private/UI/Signature/Description/SRCSignatureDescription.h b/Engine/Plugins/VirtualProduction/RemoteControl/Source/RemoteControlUI/Private/UI/Signature/Description/SRCSignatureDescription.h
new file mode 100644
--- /dev/null
+++ b/Engine/Plugins/VirtualProduction/RemoteControl/Source/RemoteControlUI/Private/UI/Signature/Description/SRCSignatureDescription.h
@@ -0,0 +1,23 @@
+// Copyright Epic Games, Inc. All Rights Reserved.
+
+#pragma once
+
+#include "Widgets/SCompoundWidget.h"
+
+class FRCSignatureTreeItemBase;
+
+class SRCSignatureDescription : public SCompoundWidget
+{
+public:
+	SLATE_BEGIN_ARGS(SRCSignatureDescription) {}
+	SLATE_END_ARGS()
+
+	void Construct(const FArguments& InArgs, const TSharedRef<FRCSignatureTreeItemBase>& InItem);
+
+private:
+	FText GetDescription() const;
+
+	TWeakPtr<FRCSignatureTreeItemBase> ItemWeak;
+
+	mutable TOptional<FText> CachedDescription;
+};
diff --git a/Engine/Plugins/VirtualProduction/RemoteControl/Source/RemoteControlUI/Private/UI/Signature/IRCSignatureColumn.h b/Engine/Plugins/VirtualProduction/RemoteControl/Source/RemoteControlUI/Private/UI/Signature/IRCSignatureColumn.h
new file mode 100644
--- /dev/null
+++ b/Engine/Plugins/VirtualProduction/RemoteControl/Source/RemoteControlUI/Private/UI/Signature/IRCSignatureColumn.h
@@ -0,0 +1,33 @@
+// Copyright Epic Games, Inc. All Rights Reserved.
+
+#pragma once
+
+#include "Templates/SharedPointer.h"
+#include "Widgets/Views/SHeaderRow.h"
+
+class FRCSignatureTreeItemBase;
+class SRCSignatureRow;
+class SRCSignatureTree;
+
+class IRCSignatureColumn : public TSharedFromThis<IRCSignatureColumn>
+{
+public:
+	virtual ~IRCSignatureColumn() = default;
+
+	virtual FName GetColumnId() const = 0;
+
+	/**
+	 * Determines whether the Column should be Showing by Default while still be able to toggle it on/off.
+	 * Used when calling SHeaderRow::SetShowGeneratedColumn (requires ShouldGenerateWidget to not be set).
+	 */
+	virtual bool ShouldShowColumnByDefault() const
+	{
+		return false;
+	}
+
+	virtual SHeaderRow::FColumn::FArguments ConstructHeaderRowColumn() = 0;
+
+	virtual TSharedRef<SWidget> ConstructRowWidget(TSharedPtr<FRCSignatureTreeItemBase> InItem
+		, const TSharedRef<SRCSignatureTree>& InList
+		, const TSharedRef<SRCSignatureRow>& InRow) = 0;
+};
diff --git a/Engine/Plugins/VirtualProduction/RemoteControl/Source/RemoteControlUI/Private/UI/Signature/Label/RCSignatureLabelColumn.cpp b/Engine/Plugins/VirtualProduction/RemoteControl/Source/RemoteControlUI/Private/UI/Signature/Label/RCSignatureLabelColumn.cpp
new file mode 100644
--- /dev/null
+++ b/Engine/Plugins/VirtualProduction/RemoteControl/Source/RemoteControlUI/Private/UI/Signature/Label/RCSignatureLabelColumn.cpp
@@ -0,0 +1,28 @@
+// Copyright Epic Games, Inc. All Rights Reserved.
+
+#include "RCSignatureLabelColumn.h"
+#include "SRCSignatureLabel.h"
+
+#define LOCTEXT_NAMESPACE "RCSignatureLabelColumn"
+
+FName FRCSignatureLabelColumn::GetColumnId() const
+{
+	return TEXT("FRCSignatureLabelColumn");
+}
+
+SHeaderRow::FColumn::FArguments FRCSignatureLabelColumn::ConstructHeaderRowColumn()
+{
+	return SHeaderRow::Column(GetColumnId())
+		.FillWidth(0.5f)
+		.DefaultLabel(LOCTEXT("DisplayName", "Label"))
+		.ShouldGenerateWidget(true);
+}
+
+TSharedRef<SWidget> FRCSignatureLabelColumn::ConstructRowWidget(TSharedPtr<FRCSignatureTreeItemBase> InItem
+	, const TSharedRef<SRCSignatureTree>& InList
+	, const TSharedRef<SRCSignatureRow>& InRow)
+{
+	return SNew(SRCSignatureLabel, InItem.ToSharedRef(), InRow);
+}
+
+#undef LOCTEXT_NAMESPACE
diff --git a/Engine/Plugins/VirtualProduction/RemoteControl/Source/RemoteControlUI/Private/UI/Signature/Label/RCSignatureLabelColumn.h b/Engine/Plugins/VirtualProduction/RemoteControl/Source/RemoteControlUI/Private/UI/Signature/Label/RCSignatureLabelColumn.h
new file mode 100644
--- /dev/null
+++ b/Engine/Plugins/VirtualProduction/RemoteControl/Source/RemoteControlUI/Private/UI/Signature/Label/RCSignatureLabelColumn.h
@@ -0,0 +1,14 @@
+// Copyright Epic Games, Inc. All Rights Reserved.
+
+#pragma once
+
+#include "UI/Signature/IRCSignatureColumn.h"
+
+class FRCSignatureLabelColumn : public IRCSignatureColumn
+{
+	//~ Begin IRCSignatureColumn
+	virtual FName GetColumnId() const override;
+	virtual SHeaderRow::FColumn::FArguments ConstructHeaderRowColumn() override;
+	virtual TSharedRef<SWidget> ConstructRowWidget(TSharedPtr<FRCSignatureTreeItemBase> InItem, const TSharedRef<SRCSignatureTree>& InList, const TSharedRef<SRCSignatureRow>& InRow) override;
+	//~ End IRCSignatureColumn
+};
diff --git a/Engine/Plugins/VirtualProduction/RemoteControl/Source/RemoteControlUI/Private/UI/Signature/Label/SRCSignatureLabel.cpp b/Engine/Plugins/VirtualProduction/RemoteControl/Source/RemoteControlUI/Private/UI/Signature/Label/SRCSignatureLabel.cpp
new file mode 100644
--- /dev/null
+++ b/Engine/Plugins/VirtualProduction/RemoteControl/Source/RemoteControlUI/Private/UI/Signature/Label/SRCSignatureLabel.cpp
@@ -0,0 +1,101 @@
+// Copyright Epic Games, Inc. All Rights Reserved.
+
+#include "SRCSignatureLabel.h"
+#include "Styling/SlateTypes.h"
+#include "UI/Signature/Items/RCSignatureTreeItemBase.h"
+#include "UI/Signature/SRCSignatureRow.h"
+#include "Widgets/Input/SCheckBox.h"
+#include "Widgets/Text/SInlineEditableTextBlock.h"
+#include "Widgets/Views/SExpanderArrow.h"
+
+#define LOCTEXT_NAMESPACE "SRCSignatureLabel"
+
+void SRCSignatureLabel::Construct(const FArguments& InArgs
+	, const TSharedRef<FRCSignatureTreeItemBase>& InItem
+	, const TSharedRef<SRCSignatureRow>& InRow)
+{
+	ItemWeak = InItem;
+
+	ChildSlot
+	[
+		SNew(SBox)
+		.MinDesiredHeight(25)
+		[
+			SNew(SHorizontalBox)
+			+ SHorizontalBox::Slot()
+			.AutoWidth()
+			.Padding(6.f, 0.f, 0.f, 0.f)
+			[
+				SNew(SExpanderArrow, InRow)
+				.IndentAmount(12)
+				.ShouldDrawWires(true)
+			]
+			+ SHorizontalBox::Slot()
+			.AutoWidth()
+			[
+				SNew(SCheckBox)
+				.IsChecked(this, &SRCSignatureLabel::GetItemEnabledState)
+				.OnCheckStateChanged(this, &SRCSignatureLabel::SetItemEnabledState)
+				.Padding(FMargin(4.0f, 0.0f))
+				.ToolTipText(LOCTEXT("ItemEnableCheckBoxTooltip", "Determines whether the entry is enabled or not"))
+			]
+			+ SHorizontalBox::Slot()
+			.FillWidth(1.0f)
+			[
+				SNew(SInlineEditableTextBlock)
+				.Text(this, &SRCSignatureLabel::GetSignatureDisplayName)
+				.OnTextCommitted(this, &SRCSignatureLabel::OnSignatureDisplayNameCommitted)
+				.IsSelected(InRow, &SRCSignatureRow::IsSelected)
+			]
+		]
+	];
+}
+
+ECheckBoxState SRCSignatureLabel::GetItemEnabledState() const
+{
+	if (!CachedCheckBoxState.IsSet())
+	{
+		TSharedPtr<FRCSignatureTreeItemBase> Item = ItemWeak.Pin();
+		TOptional<bool> ItemEnabled = Item.IsValid() ? Item->IsEnabled() : TOptional<bool>();
+
+		if (ItemEnabled.IsSet())
+		{
+			CachedCheckBoxState = *ItemEnabled ? ECheckBoxState::Checked : ECheckBoxState::Unchecked;
+		}
+		else
+		{
+			CachedCheckBoxState = ECheckBoxState::Undetermined;
+		}
+	}
+	return *CachedCheckBoxState;
+}
+
+void SRCSignatureLabel::SetItemEnabledState(ECheckBoxState InState)
+{
+	if (TSharedPtr<FRCSignatureTreeItemBase> Item = ItemWeak.Pin())
+	{
+		Item->SetEnabled(InState == ECheckBoxState::Checked);
+		CachedCheckBoxState.Reset();
+	}
+}
+
+FText SRCSignatureLabel::GetSignatureDisplayName() const
+{
+	if (!CachedDisplayName.IsSet())
+	{
+		TSharedPtr<FRCSignatureTreeItemBase> Item = ItemWeak.Pin();
+		CachedDisplayName = Item.IsValid() ? Item->GetDisplayNameText() : FText::GetEmpty();
+	}
+	return *CachedDisplayName;
+}
+
+void SRCSignatureLabel::OnSignatureDisplayNameCommitted(const FText& InText, ETextCommit::Type InCommitType)
+{
+	if (TSharedPtr<FRCSignatureTreeItemBase> Item = ItemWeak.Pin())
+	{
+		Item->SetDisplayNameText(InText);
+		CachedDisplayName.Reset();
+	}
+}
+
+#undef LOCTEXT_NAMESPACE
diff --git a/Engine/Plugins/VirtualProduction/RemoteControl/Source/RemoteControlUI/Private/UI/Signature/Label/SRCSignatureLabel.h b/Engine/Plugins/VirtualProduction/RemoteControl/Source/RemoteControlUI/Private/UI/Signature/Label/SRCSignatureLabel.h
new file mode 100644
--- /dev/null
+++ b/Engine/Plugins/VirtualProduction/RemoteControl/Source/RemoteControlUI/Private/UI/Signature/Label/SRCSignatureLabel.h
@@ -0,0 +1,35 @@
+// Copyright Epic Games, Inc. All Rights Reserved.
+
+#pragma once
+
+#include "Widgets/SCompoundWidget.h"
+
+class FRCSignatureTreeItemBase;
+class SRCSignatureRow;
+enum class ECheckBoxState : uint8;
+
+class SRCSignatureLabel : public SCompoundWidget
+{
+public:
+	SLATE_BEGIN_ARGS(SRCSignatureLabel) {}
+	SLATE_END_ARGS()
+
+	void Construct(const FArguments& InArgs
+		, const TSharedRef<FRCSignatureTreeItemBase>& InItem
+		, const TSharedRef<SRCSignatureRow>& InRow);
+
+private:
+	ECheckBoxState GetItemEnabledState() const;
+
+	void SetItemEnabledState(ECheckBoxState InState);
+
+	FText GetSignatureDisplayName() const;
+
+	void OnSignatureDisplayNameCommitted(const FText& InText, ETextCommit::Type InCommitType);
+
+	TWeakPtr<FRCSignatureTreeItemBase> ItemWeak;
+
+	mutable TOptional<ECheckBoxState> CachedCheckBoxState;
+
+	mutable TOptional<FText> CachedDisplayName;
+};
diff --git a/Engine/Plugins/VirtualProduction/RemoteControl/Source/RemoteControlUI/Private/UI/Signature/SRCSignaturePanel.cpp b/Engine/Plugins/VirtualProduction/RemoteControl/Source/RemoteControlUI/Private/UI/Signature/SRCSignaturePanel.cpp
--- a/Engine/Plugins/VirtualProduction/RemoteControl/Source/RemoteControlUI/Private/UI/Signature/SRCSignaturePanel.cpp
+++ b/Engine/Plugins/VirtualProduction/RemoteControl/Source/RemoteControlUI/Private/UI/Signature/SRCSignaturePanel.cpp
@@ -1,9 +1,9 @@
 // Copyright Epic Games, Inc. All Rights Reserved.
 
 #include "SRCSignaturePanel.h"
-#include "Columns/RCSignatureDescriptionColumn.h"
-#include "Columns/RCSignatureLabelColumn.h"
+#include "Description/RCSignatureDescriptionColumn.h"
 #include "Items/RCSignatureTreeSignatureItem.h"
+#include "Label/RCSignatureLabelColumn.h"
 #include "Misc/MessageDialog.h"
 #include "RemoteControlPreset.h"
 #include "RemoteControlSignatureRegistry.h"
diff --git a/Engine/Plugins/VirtualProduction/RemoteControl/Source/RemoteControlUI/Private/UI/Signature/SRCSignatureRow.cpp b/Engine/Plugins/VirtualProduction/RemoteControl/Source/RemoteControlUI/Private/UI/Signature/SRCSignatureRow.cpp
--- a/Engine/Plugins/VirtualProduction/RemoteControl/Source/RemoteControlUI/Private/UI/Signature/SRCSignatureRow.cpp
+++ b/Engine/Plugins/VirtualProduction/RemoteControl/Source/RemoteControlUI/Private/UI/Signature/SRCSignatureRow.cpp
@@ -1,12 +1,12 @@
 // Copyright Epic Games, Inc. All Rights Reserved.
 
 #include "SRCSignatureRow.h"
-#include "Columns/IRCSignatureColumn.h"
 #include "DragAndDrop/ActorDragDropOp.h"
 #include "DragAndDrop/CompositeDragDropOp.h"
 #include "DragAndDrop/FolderDragDropOp.h"
 #include "EditorActorFolders.h"
 #include "GameFramework/Actor.h"
+#include "IRCSignatureColumn.h"
 #include "Items/RCSignatureTreeItemBase.h"
 #include "Items/RCSignatureTreeSignatureItem.h"
 #include "RemoteControlUIModule.h"
diff --git a/Engine/Plugins/VirtualProduction/RemoteControl/Source/RemoteControlUI/Private/UI/Signature/SRCSignatureTree.cpp b/Engine/Plugins/VirtualProduction/RemoteControl/Source/RemoteControlUI/Private/UI/Signature/SRCSignatureTree.cpp
--- a/Engine/Plugins/VirtualProduction/RemoteControl/Source/RemoteControlUI/Private/UI/Signature/SRCSignatureTree.cpp
+++ b/Engine/Plugins/VirtualProduction/RemoteControl/Source/RemoteControlUI/Private/UI/Signature/SRCSignatureTree.cpp
@@ -1,7 +1,7 @@
 // Copyright Epic Games, Inc. All Rights Reserved.
 
 #include "SRCSignatureTree.h"
-#include "Columns/IRCSignatureColumn.h"
+#include "IRCSignatureColumn.h"
 #include "Items/RCSignatureTreeItemBase.h"
 #include "Items/RCSignatureTreeRootItem.h"
 #include "Items/RCSignatureTreeSignatureItem.h"
