Remote Control: Added Signature Actions that allow to execute custom logic on the exposed property of the signature.

diff --git a/Engine/Plugins/VirtualProduction/RemoteControl/Source/RemoteControl/Private/RemoteControlSignature.cpp b/Engine/Plugins/VirtualProduction/RemoteControl/Source/RemoteControl/Private/RemoteControlSignature.cpp
--- a/Engine/Plugins/VirtualProduction/RemoteControl/Source/RemoteControl/Private/RemoteControlSignature.cpp
+++ b/Engine/Plugins/VirtualProduction/RemoteControl/Source/RemoteControl/Private/RemoteControlSignature.cpp
@@ -2,6 +2,8 @@
 
 #include "RemoteControlSignature.h"
 #include "GameFramework/Actor.h"
+#include "RCSignatureAction.h"
+#include "RCSignatureActionDefinition.h"
 #include "RemoteControlField.h"
 #include "RemoteControlPreset.h"
 
@@ -34,7 +36,7 @@ namespace UE::RemoteControl::Private
 			UClass* FindClass = InSupportedClass ? InSupportedClass : UObject::StaticClass();
 			Context = StaticFindObject(FindClass, InOuterObject, *InField.ObjectRelativePath);
 
-			// Slow Path: if Subobject Path did not find the object, try to find the first sub-object of the actor matching the class.
+			// Slow Path: if Subobject Path did not find the object, try to find the first sub-object of the object matching the class.
 			if (!Context && InSupportedClass)
 			{
 				ForEachObjectWithOuterBreakable(InOuterObject, [&Context, InSupportedClass](UObject* InSubobject)->bool
@@ -55,38 +57,29 @@ namespace UE::RemoteControl::Private
 	}
 }
 
-FRCSignatureField FRCSignatureField::CreateField(const FRCFieldPathInfo& InFieldPathInfo, UObject* InOwnerObject, FProperty* InProperty)
+FRCSignatureField FRCSignatureField::CreateField(const FRCFieldPathInfo& InFieldPathInfo, const UObject* InOwnerObject, const FProperty* InProperty)
 {
-	UClass* SupportedBindingClass = nullptr;
-
+	const UClass* SupportedClass = nullptr;
 	if (InProperty)
 	{
-		if (UClass* PropertyOwnerClass = InProperty->GetOwnerClass())
-		{
-			SupportedBindingClass = PropertyOwnerClass;
-		}
+		SupportedClass = InProperty->GetOwnerClass();
 	}
 
-	return CreateField(InFieldPathInfo, InOwnerObject, SupportedBindingClass);
-}
-
-FRCSignatureField FRCSignatureField::CreateField(const FRCFieldPathInfo& InFieldPathInfo, UObject* InOwnerObject, UClass* InSupportedClass)
-{
-	if (!InSupportedClass && InFieldPathInfo.GetSegmentCount() > 0 && InFieldPathInfo.GetFieldSegment(0).IsResolved())
+	if (!SupportedClass && InFieldPathInfo.GetSegmentCount() > 0 && InFieldPathInfo.GetFieldSegment(0).IsResolved())
 	{
-		InSupportedClass = InFieldPathInfo.GetFieldSegment(0).ResolvedData.Field->GetOwnerClass();
+		SupportedClass = InFieldPathInfo.GetFieldSegment(0).ResolvedData.Field->GetOwnerClass();
 	}
 
 	FRCSignatureField Field;
 	Field.bEnabled = true;
 	Field.FieldPath = InFieldPathInfo;
-	Field.SupportedClass = InSupportedClass;
+	Field.SupportedClass = SupportedClass;
 
-	if (UObject* BoundObject = InOwnerObject)
+	if (InOwnerObject)
 	{
-		if (AActor* ActorOwner = BoundObject->GetTypedOuter<AActor>())
+		if (AActor* ActorOwner = InOwnerObject->GetTypedOuter<AActor>())
 		{
-			Field.ObjectRelativePath = BoundObject->GetPathName(ActorOwner);
+			Field.ObjectRelativePath = InOwnerObject->GetPathName(ActorOwner);
 		}
 	}
 
@@ -147,14 +140,23 @@ int32 FRCSignature::ApplySignature(URemoteControlPreset* InPreset, TConstArrayVi
 			}
 
 			// If path resolved, expose the property
-			TSharedPtr<FRemoteControlProperty> Property = InPreset->ExposeProperty(Context, Field.FieldPath, ExposeArgs).Pin();
-			if (!Property.IsValid())
+			FRCSignatureActionContext ActionContext;
+			ActionContext.Preset = InPreset;
+			ActionContext.Object = Object;
+			ActionContext.Property = InPreset->ExposeProperty(Context, Field.FieldPath, ExposeArgs).Pin();
+			if (!ActionContext.Property.IsValid())
 			{
 				continue;
 			}
 
 			// Property was exposed successfully, increase expose count
 			++ExposeCount;
+
+			// Execute the actions for the newly exposed property
+			for (const FRCSignatureActionDefinition& Action : Field.ActionDefinitions)
+			{
+				Action.Execute(ActionContext);
+			}
 		}
 	}
 
diff --git a/Engine/Plugins/VirtualProduction/RemoteControl/Source/RemoteControl/Private/RemoteControlSignatureActionDefinition.cpp b/Engine/Plugins/VirtualProduction/RemoteControl/Source/RemoteControl/Private/RemoteControlSignatureActionDefinition.cpp
new file mode 100644
--- /dev/null
+++ b/Engine/Plugins/VirtualProduction/RemoteControl/Source/RemoteControl/Private/RemoteControlSignatureActionDefinition.cpp
@@ -0,0 +1,34 @@
+// Copyright Epic Games, Inc. All Rights Reserved.
+
+#include "RCSignatureActionDefinition.h"
+#include "RCSignatureAction.h"
+#include "UObject/StructOnScope.h"
+
+FRCSignatureActionDefinition::FRCSignatureActionDefinition(const UScriptStruct* InScriptStruct, const FRCSignatureField& InFieldOwner)
+{
+	ActionInstance.InitializeAsScriptStruct(InScriptStruct, /*StructMemory*/nullptr);
+
+	if (FRCSignatureAction* Action = ActionInstance.GetMutablePtr())
+	{
+		Action->Initialize(InFieldOwner);
+	}
+}
+
+const FRCSignatureAction* FRCSignatureActionDefinition::GetAction() const
+{
+	return ActionInstance.GetPtr();
+}
+
+TSharedRef<FStructOnScope> FRCSignatureActionDefinition::MakeStructOnScope()
+{
+	return MakeShared<FStructOnScope>(ActionInstance.GetScriptStruct(), ActionInstance.GetMutableMemory());
+}
+
+bool FRCSignatureActionDefinition::Execute(const FRCSignatureActionContext& InContext) const
+{
+	if (const FRCSignatureAction* Action = ActionInstance.GetPtr())
+	{
+		return Action->Execute(InContext);
+	}
+	return false;
+}
diff --git a/Engine/Plugins/VirtualProduction/RemoteControl/Source/RemoteControl/Public/RCSignatureAction.h b/Engine/Plugins/VirtualProduction/RemoteControl/Source/RemoteControl/Public/RCSignatureAction.h
new file mode 100644
--- /dev/null
+++ b/Engine/Plugins/VirtualProduction/RemoteControl/Source/RemoteControl/Public/RCSignatureAction.h
@@ -0,0 +1,83 @@
+// Copyright Epic Games, Inc. All Rights Reserved.
+
+#pragma once
+
+#include "Templates/SharedPointer.h"
+#include "UObject/ObjectPtr.h"
+#include "RCSignatureAction.generated.h"
+
+class URemoteControlPreset;
+struct FRCSignatureField;
+struct FRemoteControlProperty;
+
+/** The context for a Signature Action to execute */
+struct FRCSignatureActionContext
+{
+	/** The preset where the Signature is being applied */
+	TObjectPtr<URemoteControlPreset> Preset = nullptr;
+
+	/** The object that the Signature is applying */
+	TObjectPtr<UObject> Object = nullptr;
+
+	/** The exposed property from the Signature */
+	TSharedPtr<FRemoteControlProperty> Property;
+};
+
+#if WITH_EDITOR
+/**
+ * Editor information on the Icon of a given Action
+ * Names are explicitly used over FSlateIcon to avoid a dependency on Slate Core API
+ */
+struct FRCSignatureActionIcon
+{
+	/** Name of the style set the icon can be found in */
+	FName StyleSetName;
+
+	/** Name of the style for the icon */
+	FName StyleName;
+
+	/** Name of the style for the overlay icon (if any) */
+	FName OverlayStyleName;
+
+	/** Color of the Base Icon */
+	FLinearColor BaseColor = FLinearColor::White;
+
+	/** Optional Color of the Overlay Icon (if any). Uses the Icon Color if not set */
+	TOptional<FLinearColor> OverlayColor;
+};
+#endif
+
+/** Base class of Actions that execute when Applying a Signature Field */
+USTRUCT(meta=(Hidden))
+struct FRCSignatureAction
+{
+	GENERATED_BODY()
+
+	virtual ~FRCSignatureAction() = default;
+
+	/**
+	 * Called only once when the Signature Action is first added to the action list
+	 * @param InFieldOwner the owner of the action list
+	 */
+	virtual void Initialize(const FRCSignatureField& InFieldOwner)
+	{
+	}
+
+	/**
+	 * Executes the Action logic
+	 * @param InContext context of the action (preset, exposed property and uobject being used)
+	 * @return true if the action did anything to the preset
+	 */
+	virtual bool Execute(const FRCSignatureActionContext& InContext) const
+	{
+		return false;
+	}
+
+#if WITH_EDITOR
+	/** Retrieves the icon to use to represent this Signature Action */
+	virtual FRCSignatureActionIcon GetIcon() const
+	{
+		return FRCSignatureActionIcon();
+	}
+#endif
+};
diff --git a/Engine/Plugins/VirtualProduction/RemoteControl/Source/RemoteControl/Public/RCSignatureActionDefinition.h b/Engine/Plugins/VirtualProduction/RemoteControl/Source/RemoteControl/Public/RCSignatureActionDefinition.h
new file mode 100644
--- /dev/null
+++ b/Engine/Plugins/VirtualProduction/RemoteControl/Source/RemoteControl/Public/RCSignatureActionDefinition.h
@@ -0,0 +1,32 @@
+// Copyright Epic Games, Inc. All Rights Reserved.
+
+#pragma once
+
+#include "StructUtils/InstancedStruct.h"
+#include "RCSignatureActionDefinition.generated.h"
+
+class FStructOnScope;
+struct FRCSignatureAction;
+struct FRCSignatureActionContext;
+struct FRCSignatureField;
+
+/** Struct containing an Action Instance and handling its Execution */
+USTRUCT()
+struct FRCSignatureActionDefinition
+{
+	GENERATED_BODY()
+
+	FRCSignatureActionDefinition() = default;
+
+	REMOTECONTROL_API explicit FRCSignatureActionDefinition(const UScriptStruct* InScriptStruct, const FRCSignatureField& InFieldOwner);
+
+	REMOTECONTROL_API const FRCSignatureAction* GetAction() const;
+
+	REMOTECONTROL_API TSharedRef<FStructOnScope> MakeStructOnScope();
+
+	bool Execute(const FRCSignatureActionContext& InContext) const;
+
+private:
+	UPROPERTY()
+	TInstancedStruct<FRCSignatureAction> ActionInstance;
+};
diff --git a/Engine/Plugins/VirtualProduction/RemoteControl/Source/RemoteControl/Public/RemoteControlSignature.h b/Engine/Plugins/VirtualProduction/RemoteControl/Source/RemoteControl/Public/RemoteControlSignature.h
--- a/Engine/Plugins/VirtualProduction/RemoteControl/Source/RemoteControl/Public/RemoteControlSignature.h
+++ b/Engine/Plugins/VirtualProduction/RemoteControl/Source/RemoteControl/Public/RemoteControlSignature.h
@@ -2,9 +2,11 @@
 
 #pragma once
 
+#include "RCSignatureActionDefinition.h"
 #include "RemoteControlFieldPath.h"
 #include "RemoteControlSignature.generated.h"
 
+enum class EPropertyBagPropertyType : uint8;
 class URemoteControlPreset;
 
 /** Representation class of a Field (e.g. Property) in an Object */
@@ -15,8 +17,7 @@ struct FRCSignatureField
 
 	FRCSignatureField() = default;
 
-	REMOTECONTROL_API static FRCSignatureField CreateField(const FRCFieldPathInfo& InFieldPathInfo, UObject* InOwnerObject, FProperty* InProperty);
-	REMOTECONTROL_API static FRCSignatureField CreateField(const FRCFieldPathInfo& InFieldPathInfo, UObject* InOwnerObject, UClass* InSupportedClass);
+	REMOTECONTROL_API static FRCSignatureField CreateField(const FRCFieldPathInfo& InFieldPathInfo, const UObject* InOwnerObject, const FProperty* InProperty);
 
 	bool operator==(const FRCSignatureField& InOtherField) const
 	{
@@ -37,6 +38,10 @@ struct FRCSignatureField
 	UPROPERTY()
 	FSoftClassPath SupportedClass;
 
+	/** Container holding the action definitions for the field */
+	UPROPERTY()
+	TArray<FRCSignatureActionDefinition> ActionDefinitions;
+
 	/** Whether to consider this field when applying a Signature */
 	UPROPERTY()
 	bool bEnabled = true;
@@ -59,7 +64,7 @@ struct FRCSignature
 	REMOTECONTROL_API int32 AddFields(TConstArrayView<FRCSignatureField> InFields);
 
 	/**
-	 * Applies this Signature to the given Actors by exposing all this Signature's fields to the given preset
+	 * Applies this Signature to the given Objects by exposing all this Signature's fields to the given preset
 	 * @param InPreset the preset where these properties will be exposed to
 	 * @param InObjects the objects whose properties or properties of its subojects to expose
 	 * @return the number of properties affected in total
diff --git a/Engine/Plugins/VirtualProduction/RemoteControl/Source/RemoteControlUI/Private/UI/Signature/Actions/RCSignatureActionColumn.cpp b/Engine/Plugins/VirtualProduction/RemoteControl/Source/RemoteControlUI/Private/UI/Signature/Actions/RCSignatureActionColumn.cpp
new file mode 100644
--- /dev/null
+++ b/Engine/Plugins/VirtualProduction/RemoteControl/Source/RemoteControlUI/Private/UI/Signature/Actions/RCSignatureActionColumn.cpp
@@ -0,0 +1,80 @@
+// Copyright Epic Games, Inc. All Rights Reserved.
+
+#include "RCSignatureActionColumn.h"
+#include "RCSignatureAction.h"
+#include "RCSignatureActionType.h"
+#include "SRCSignatureActionBox.h"
+#include "StructUtils/InstancedStruct.h"
+#include "UObject/Class.h"
+#include "UObject/UObjectIterator.h"
+
+#define LOCTEXT_NAMESPACE "RCSignatureActionColumn"
+
+FRCSignatureActionColumn::FRCSignatureActionColumn(const TAttribute<bool>& InLiveMode)
+	: LiveMode(InLiveMode)
+{
+	RefreshActionTypes();
+}
+
+FName FRCSignatureActionColumn::GetColumnId() const
+{
+	return TEXT("FRCSignatureActionColumn");
+}
+
+bool FRCSignatureActionColumn::ShouldShowColumnByDefault() const
+{
+	return true;
+}
+
+SHeaderRow::FColumn::FArguments FRCSignatureActionColumn::ConstructHeaderRowColumn()
+{
+	return SHeaderRow::Column(GetColumnId())
+		.FillWidth(0.5f)
+		.DefaultLabel(LOCTEXT("DisplayName", "Actions"));
+}
+
+TSharedRef<SWidget> FRCSignatureActionColumn::ConstructRowWidget(TSharedPtr<FRCSignatureTreeItemBase> InItem, const TSharedRef<SRCSignatureTree>& InList, const TSharedRef<SRCSignatureRow>& InRow)
+{
+	return SNew(SRCSignatureActionBox, InItem.ToSharedRef(), InRow)
+		.LiveMode(LiveMode)
+		.ActionTypesSource(&ActionTypes)
+		.OnActionTypesComboBoxOpening(this, &FRCSignatureActionColumn::RefreshActionTypes);
+}
+
+void FRCSignatureActionColumn::RefreshActionTypes()
+{
+	ActionTypes.Reset();
+
+	if (!UObjectInitialized())
+	{
+		return;
+	}
+
+	const FName HiddenMetaData(TEXT("Hidden"));
+
+	for (UScriptStruct* ScriptStruct : TObjectRange<UScriptStruct>())
+	{
+		if (ScriptStruct->HasMetaData(HiddenMetaData))
+		{
+			continue;				
+		}
+
+		if (ScriptStruct->IsChildOf(TBaseStructure<FRCSignatureAction>::Get()))
+		{
+			TSharedRef<FRCSignatureActionType> ActionType = MakeShared<FRCSignatureActionType>();
+			ActionType->Type = ScriptStruct;
+			ActionType->Title = ScriptStruct->GetDisplayNameText();
+
+			// Initialize a Temp Instance to get the Icon to use
+			{
+				TInstancedStruct<FRCSignatureAction> Instance;
+				Instance.InitializeAsScriptStruct(ScriptStruct, /*StructMemory*/nullptr);
+				ActionType->Icon = Instance.Get().GetIcon();
+			}
+
+			ActionTypes.Add(MoveTemp(ActionType));
+		}
+	}
+}
+
+#undef LOCTEXT_NAMESPACE
diff --git a/Engine/Plugins/VirtualProduction/RemoteControl/Source/RemoteControlUI/Private/UI/Signature/Actions/RCSignatureActionColumn.h b/Engine/Plugins/VirtualProduction/RemoteControl/Source/RemoteControlUI/Private/UI/Signature/Actions/RCSignatureActionColumn.h
new file mode 100644
--- /dev/null
+++ b/Engine/Plugins/VirtualProduction/RemoteControl/Source/RemoteControlUI/Private/UI/Signature/Actions/RCSignatureActionColumn.h
@@ -0,0 +1,27 @@
+// Copyright Epic Games, Inc. All Rights Reserved.
+
+#pragma once
+
+#include "UI/Signature/IRCSignatureColumn.h"
+
+struct FRCSignatureActionType;
+
+class FRCSignatureActionColumn : public IRCSignatureColumn
+{
+public:
+	explicit FRCSignatureActionColumn(const TAttribute<bool>& InLiveMode);
+
+private:
+	//~ Begin IRCSignatureColumn
+	virtual FName GetColumnId() const override;
+	virtual bool ShouldShowColumnByDefault() const override;
+	virtual SHeaderRow::FColumn::FArguments ConstructHeaderRowColumn() override;
+	virtual TSharedRef<SWidget> ConstructRowWidget(TSharedPtr<FRCSignatureTreeItemBase> InItem, const TSharedRef<SRCSignatureTree>& InList, const TSharedRef<SRCSignatureRow>& InRow) override;
+	//~ End IRCSignatureColumn
+
+	void RefreshActionTypes();
+
+	TArray<TSharedPtr<FRCSignatureActionType>> ActionTypes; 
+
+	TAttribute<bool> LiveMode;
+};
diff --git a/Engine/Plugins/VirtualProduction/RemoteControl/Source/RemoteControlUI/Private/UI/Signature/Actions/RCSignatureActionType.h b/Engine/Plugins/VirtualProduction/RemoteControl/Source/RemoteControlUI/Private/UI/Signature/Actions/RCSignatureActionType.h
new file mode 100644
--- /dev/null
+++ b/Engine/Plugins/VirtualProduction/RemoteControl/Source/RemoteControlUI/Private/UI/Signature/Actions/RCSignatureActionType.h
@@ -0,0 +1,22 @@
+// Copyright Epic Games, Inc. All Rights Reserved.
+
+#pragma once
+
+#include "Internationalization/Text.h"
+#include "RCSignatureAction.h"
+#include "Textures/SlateIcon.h"
+
+class UScriptStruct;
+
+/** Struct to describe an Item in the Action Type combo box */
+struct FRCSignatureActionType
+{
+	/** Struct type of the Action to create if selected */
+	const UScriptStruct* Type = nullptr;
+
+	/** Display Name Text in the combo box entry */
+	FText Title;
+
+	/** Icon and color in the combo box entry */
+	FRCSignatureActionIcon Icon;
+};
diff --git a/Engine/Plugins/VirtualProduction/RemoteControl/Source/RemoteControlUI/Private/UI/Signature/Actions/SRCSignatureAction.cpp b/Engine/Plugins/VirtualProduction/RemoteControl/Source/RemoteControlUI/Private/UI/Signature/Actions/SRCSignatureAction.cpp
new file mode 100644
--- /dev/null
+++ b/Engine/Plugins/VirtualProduction/RemoteControl/Source/RemoteControlUI/Private/UI/Signature/Actions/SRCSignatureAction.cpp
@@ -0,0 +1,127 @@
+// Copyright Epic Games, Inc. All Rights Reserved.
+
+#include "SRCSignatureAction.h"
+#include "RCSignatureAction.h"
+#include "SRCSignatureActionIcon.h"
+#include "Styling/AppStyle.h"
+#include "Styling/ToolBarStyle.h"
+#include "UI/Signature/Items/RCSignatureTreeActionItem.h"
+#include "UI/Signature/Items/RCSignatureTreeItemBase.h"
+#include "Widgets/Layout/SBorder.h"
+
+void SRCSignatureAction::Construct(const FArguments& InArgs, const TSharedRef<FRCSignatureTreeActionItem>& InActionItem)
+{
+	ActionItemWeak = InActionItem;
+	CheckBoxStyle = &FAppStyle::Get().GetWidgetStyle<FToolBarStyle>(TEXT("SlimToolBar")).ToggleButton;
+
+	ChildSlot
+	[
+		SNew(SBorder)
+		.BorderImage(this, &SRCSignatureAction::GetBackgroundBrush)
+		.Padding(5.f, 4.f)
+		[
+			SAssignNew(ActionImage, SRCSignatureActionIcon)
+			.ActionIcon(InActionItem->GetIcon())
+		]
+	];
+}
+
+FReply SRCSignatureAction::OnMouseButtonDown(const FGeometry& InGeometry, const FPointerEvent& InMouseEvent)
+{
+	if (InMouseEvent.GetEffectingButton() == EKeys::LeftMouseButton)
+	{
+		SetPressed(true);
+
+		TSharedRef<SRCSignatureAction> This = SharedThis(this);
+
+		return FReply::Handled()
+			.DetectDrag(This, InMouseEvent.GetEffectingButton())
+			.CaptureMouse(This)
+			.SetUserFocus(This, EFocusCause::Mouse);
+	}
+	return FReply::Unhandled();
+}
+
+FReply SRCSignatureAction::OnMouseButtonUp(const FGeometry& InGeometry, const FPointerEvent& InMouseEvent)
+{
+	FReply Reply = FReply::Unhandled();
+
+	if (IsPressed() && InMouseEvent.GetEffectingButton() == EKeys::LeftMouseButton)
+	{
+		SetPressed(false);
+
+		bool bEventOverButton = IsHovered();
+		if (!bEventOverButton && InMouseEvent.IsTouchEvent())
+		{
+			bEventOverButton = InGeometry.IsUnderLocation(InMouseEvent.GetScreenSpacePosition());
+		}
+
+		if (bEventOverButton && HasMouseCapture())
+		{
+			if (TSharedPtr<FRCSignatureTreeActionItem> ActionItem = ActionItemWeak.Pin())
+			{
+				bool bShouldSelect = true;
+
+				const bool bIsMultiSelection = InMouseEvent.IsControlDown();
+
+				// if Multi-selecting and already selected, toggle off
+				if (bIsMultiSelection && ActionItem->IsSelected())
+				{
+					bShouldSelect = false;
+				}
+
+				ActionItem->SetSelected(bShouldSelect, bIsMultiSelection);
+			}
+			Reply = FReply::Handled();
+		}
+	}
+
+	if (HasMouseCapture())
+	{
+		Reply.ReleaseMouseCapture();
+	}
+
+	return Reply;
+}
+
+void SRCSignatureAction::OnMouseCaptureLost(const FCaptureLostEvent& InCaptureLostEvent)
+{
+	SetPressed(false);
+}
+
+bool SRCSignatureAction::IsPressed() const
+{
+	return bIsPressed;
+}
+
+void SRCSignatureAction::SetPressed(bool bInPressed)
+{
+	bIsPressed = bInPressed;
+}
+
+const FSlateBrush* SRCSignatureAction::GetBackgroundBrush() const
+{
+	check(CheckBoxStyle);
+
+	TSharedPtr<FRCSignatureTreeActionItem> ActionItem = ActionItemWeak.Pin();
+
+	const bool bIsSelected = ActionItem.IsValid() && ActionItem->IsSelected();
+
+	if (IsPressed())
+	{
+		return bIsSelected
+			? &CheckBoxStyle->CheckedPressedImage
+			: &CheckBoxStyle->UncheckedPressedImage;
+	}
+
+	if (IsHovered())
+	{
+		return bIsSelected
+			? &CheckBoxStyle->CheckedHoveredImage
+			: &CheckBoxStyle->UncheckedHoveredImage;
+	}
+
+	return bIsSelected
+		? &CheckBoxStyle->CheckedImage
+		: &CheckBoxStyle->UncheckedImage;
+}
diff --git a/Engine/Plugins/VirtualProduction/RemoteControl/Source/RemoteControlUI/Private/UI/Signature/Actions/SRCSignatureAction.h b/Engine/Plugins/VirtualProduction/RemoteControl/Source/RemoteControlUI/Private/UI/Signature/Actions/SRCSignatureAction.h
new file mode 100644
--- /dev/null
+++ b/Engine/Plugins/VirtualProduction/RemoteControl/Source/RemoteControlUI/Private/UI/Signature/Actions/SRCSignatureAction.h
@@ -0,0 +1,39 @@
+// Copyright Epic Games, Inc. All Rights Reserved.
+
+#pragma once
+
+#include "Widgets/SCompoundWidget.h"
+
+class FRCSignatureTreeActionItem;
+class SRCSignatureActionIcon;
+struct FCheckBoxStyle;
+
+class SRCSignatureAction : public SCompoundWidget
+{
+public:
+	SLATE_BEGIN_ARGS(SRCSignatureAction) {}
+	SLATE_END_ARGS()
+
+	void Construct(const FArguments& InArgs, const TSharedRef<FRCSignatureTreeActionItem>& InActionItem);
+
+protected:
+	//~ Begin SWidget
+	virtual FReply OnMouseButtonDown(const FGeometry& InGeometry, const FPointerEvent& InMouseEvent) override;
+	virtual FReply OnMouseButtonUp(const FGeometry& InGeometry, const FPointerEvent& InMouseEvent) override;
+	virtual void OnMouseCaptureLost(const FCaptureLostEvent& InCaptureLostEvent) override;
+	//~ End SWidget
+
+private:
+	bool IsPressed() const;
+	void SetPressed(bool bInPressed);
+
+	const FSlateBrush* GetBackgroundBrush() const;
+
+	TWeakPtr<FRCSignatureTreeActionItem> ActionItemWeak;
+
+	TSharedPtr<SRCSignatureActionIcon> ActionImage;
+
+	const FCheckBoxStyle* CheckBoxStyle = nullptr;
+
+	bool bIsPressed = false;
+};
diff --git a/Engine/Plugins/VirtualProduction/RemoteControl/Source/RemoteControlUI/Private/UI/Signature/Actions/SRCSignatureActionBox.cpp b/Engine/Plugins/VirtualProduction/RemoteControl/Source/RemoteControlUI/Private/UI/Signature/Actions/SRCSignatureActionBox.cpp
new file mode 100644
--- /dev/null
+++ b/Engine/Plugins/VirtualProduction/RemoteControl/Source/RemoteControlUI/Private/UI/Signature/Actions/SRCSignatureActionBox.cpp
@@ -0,0 +1,173 @@
+// Copyright Epic Games, Inc. All Rights Reserved.
+
+#include "SRCSignatureActionBox.h"
+#include "DetailLayoutBuilder.h"
+#include "RCSignatureActionType.h"
+#include "RemoteControlSignature.h"
+#include "SRCSignatureAction.h"
+#include "SRCSignatureActionIcon.h"
+#include "Styling/RemoteControlStyles.h"
+#include "UI/RemoteControlPanelStyle.h"
+#include "UI/Signature/Items/RCSignatureTreeActionItem.h"
+#include "UI/Signature/Items/RCSignatureTreeFieldItem.h"
+#include "UI/Signature/Items/RCSignatureTreeItemBase.h"
+#include "UI/Signature/SRCSignatureRow.h"
+#include "Widgets/Input/SComboBox.h"
+#include "Widgets/Layout/SScrollBox.h"
+
+void SRCSignatureActionBox::Construct(const FArguments& InArgs, const TSharedRef<FRCSignatureTreeItemBase>& InItem, const TSharedRef<SRCSignatureRow>& InRow)
+{
+	ItemWeak = InItem;
+	LiveMode = InArgs._LiveMode;
+	IsHovered = TAttribute<bool>(InRow, &SRCSignatureRow::IsHovered);
+	IsSelected = TAttribute<bool>(InRow, &SRCSignatureRow::IsSelected);
+
+	ActionTypesComboBox = SNew(SComboBox<TSharedPtr<FRCSignatureActionType>>)
+		.IsEnabled(this, &SRCSignatureActionBox::CanAddAction)
+		.Visibility(this, &SRCSignatureActionBox::GetAddActionVisibility)
+		.ComboBoxStyle(&FAppStyle::Get().GetWidgetStyle<FComboBoxStyle>("SimpleComboBox"))
+		.OptionsSource(InArgs.GetActionTypesSource())
+		.OnComboBoxOpening(InArgs._OnActionTypesComboBoxOpening)
+		.OnGenerateWidget(this, &SRCSignatureActionBox::GenerateActionTypeWidget)
+		.OnSelectionChanged(this, &SRCSignatureActionBox::OnActionTypeSelected)
+		.ContentPadding(0)
+		.HasDownArrow(false)
+		[
+			SNew(SImage)
+			.Image(FAppStyle::GetBrush("Icons.PlusCircle"))
+			.DesiredSizeOverride(FVector2D(SRCSignatureActionIcon::IconSize))
+		];
+
+	ChildSlot
+		[
+			SAssignNew(ActionListBox, SScrollBox)
+			.ConsumeMouseWheel(EConsumeMouseWheel::Always)
+			.Orientation(EOrientation::Orient_Horizontal)
+			.ScrollBarThickness(FVector2D(2.f))
+		];
+
+	Refresh();
+}
+
+void SRCSignatureActionBox::Refresh()
+{
+	ActionListBox->ClearChildren();
+
+	TSharedPtr<FRCSignatureTreeItemBase> Item = ItemWeak.Pin();
+	if (!Item.IsValid())
+	{
+		return;
+	}
+
+	TSharedPtr<FRCSignatureTreeFieldItem> FieldItem = Item->MutableCast<FRCSignatureTreeFieldItem>();
+	if (!FieldItem.IsValid())
+	{
+		return;
+	}
+
+	// Force rebuild children to discover any new actions
+	FieldItem->RebuildChildren();
+
+	TConstArrayView<TSharedPtr<FRCSignatureTreeItemBase>> FieldChildren = Item->GetChildren();
+
+	for (const TSharedPtr<FRCSignatureTreeItemBase>& FieldChild : FieldChildren)
+	{
+		TSharedPtr<FRCSignatureTreeActionItem> ActionItem = FieldChild->MutableCast<FRCSignatureTreeActionItem>();
+		if (!ActionItem.IsValid())
+		{
+			continue;
+		}
+
+		ActionListBox->AddSlot()
+			.Padding(0.f, 1.f)
+			[
+				SNew(SRCSignatureAction, ActionItem.ToSharedRef())
+			];
+	}
+
+	// Add (+) button at the end to add new actions
+	ActionListBox->AddSlot()
+		[
+			ActionTypesComboBox.ToSharedRef()
+		];
+}
+
+bool SRCSignatureActionBox::CanAddAction() const
+{
+	const bool bIsLive = LiveMode.Get(false);
+	return !bIsLive;
+}
+
+EVisibility SRCSignatureActionBox::GetAddActionVisibility() const
+{
+	if (!CanAddAction())
+	{
+		return EVisibility::Collapsed;
+	}
+
+	if (ActionTypesComboBox->IsOpen())
+	{
+		return EVisibility::Visible;
+	}
+
+	if (IsSelected.Get(false) || IsHovered.Get(false))
+	{
+		return EVisibility::Visible;
+	}
+
+	return EVisibility::Collapsed;
+}
+
+TSharedRef<SWidget> SRCSignatureActionBox::GenerateActionTypeWidget(TSharedPtr<FRCSignatureActionType> InActionType)
+{
+	check(InActionType.IsValid());
+	return SNew(SBox)
+		.WidthOverride(200.f)
+		.Padding(0.f, 5.f)
+		[
+			SNew(SHorizontalBox)
+			+ SHorizontalBox::Slot()
+			.AutoWidth()
+			.Padding(5.f, 0.f)
+			.HAlign(HAlign_Center)
+			.VAlign(VAlign_Center)
+			[
+				SNew(SRCSignatureActionIcon)
+				.ActionIcon(InActionType->Icon)
+			]
+			+ SHorizontalBox::Slot()
+			.FillWidth(1.f)
+			.Padding(5.f, 0.f)
+			.HAlign(HAlign_Left)
+			.VAlign(VAlign_Center)
+			[
+				SNew(STextBlock)
+				.Text(InActionType->Title)
+				.Justification(ETextJustify::Left)
+			]
+		];
+}
+
+void SRCSignatureActionBox::OnActionTypeSelected(TSharedPtr<FRCSignatureActionType> InActionType, ESelectInfo::Type InSelectType)
+{
+	if (!InActionType.IsValid() || !ensure(InActionType->Type))
+	{
+		return;
+	}
+
+	TSharedPtr<FRCSignatureTreeItemBase> Item = ItemWeak.Pin();
+	if (!Item.IsValid())
+	{
+		return;
+	}
+
+	TSharedPtr<FRCSignatureTreeFieldItem> FieldItem = Item->MutableCast<FRCSignatureTreeFieldItem>();
+	if (!FieldItem.IsValid())
+	{
+		return;
+	}
+
+	FieldItem->AddAction(InActionType->Type);
+	ActionTypesComboBox->ClearSelection();
+	Refresh();
+}
diff --git a/Engine/Plugins/VirtualProduction/RemoteControl/Source/RemoteControlUI/Private/UI/Signature/Actions/SRCSignatureActionBox.h b/Engine/Plugins/VirtualProduction/RemoteControl/Source/RemoteControlUI/Private/UI/Signature/Actions/SRCSignatureActionBox.h
new file mode 100644
--- /dev/null
+++ b/Engine/Plugins/VirtualProduction/RemoteControl/Source/RemoteControlUI/Private/UI/Signature/Actions/SRCSignatureActionBox.h
@@ -0,0 +1,45 @@
+// Copyright Epic Games, Inc. All Rights Reserved.
+
+#pragma once
+
+#include "Widgets/SCompoundWidget.h"
+#include "Widgets/Views/IItemsSource.h"
+
+class FRCSignatureTreeItemBase;
+class SRCSignatureRow;
+class SScrollBox;
+struct FRCSignatureActionType;
+template<typename OptionType> class SComboBox;
+
+class SRCSignatureActionBox : public SCompoundWidget
+{
+public:
+	SLATE_BEGIN_ARGS(SRCSignatureActionBox) {}
+		SLATE_ATTRIBUTE(bool, LiveMode)
+		SLATE_ITEMS_SOURCE_ARGUMENT(TSharedPtr<FRCSignatureActionType>, ActionTypesSource)
+		SLATE_EVENT(TDelegate<void()>, OnActionTypesComboBoxOpening)
+	SLATE_END_ARGS()
+
+	void Construct(const FArguments& InArgs, const TSharedRef<FRCSignatureTreeItemBase>& InItem, const TSharedRef<SRCSignatureRow>& InRow);
+
+private:
+	void Refresh();
+
+	bool CanAddAction() const;
+
+	EVisibility GetAddActionVisibility() const;
+
+	TSharedRef<SWidget> GenerateActionTypeWidget(TSharedPtr<FRCSignatureActionType> InActionType);
+
+	void OnActionTypeSelected(TSharedPtr<FRCSignatureActionType> InActionType, ESelectInfo::Type InSelectType);
+
+	TSharedPtr<SComboBox<TSharedPtr<FRCSignatureActionType>>> ActionTypesComboBox;
+
+	TWeakPtr<FRCSignatureTreeItemBase> ItemWeak;
+
+	TSharedPtr<SScrollBox> ActionListBox;
+
+	TAttribute<bool> IsHovered;
+	TAttribute<bool> IsSelected;
+	TAttribute<bool> LiveMode;
+};
diff --git a/Engine/Plugins/VirtualProduction/RemoteControl/Source/RemoteControlUI/Private/UI/Signature/Actions/SRCSignatureActionIcon.cpp b/Engine/Plugins/VirtualProduction/RemoteControl/Source/RemoteControlUI/Private/UI/Signature/Actions/SRCSignatureActionIcon.cpp
new file mode 100644
--- /dev/null
+++ b/Engine/Plugins/VirtualProduction/RemoteControl/Source/RemoteControlUI/Private/UI/Signature/Actions/SRCSignatureActionIcon.cpp
@@ -0,0 +1,47 @@
+// Copyright Epic Games, Inc. All Rights Reserved.
+
+#include "SRCSignatureActionIcon.h"
+#include "RCSignatureAction.h"
+#include "Textures/SlateIcon.h"
+#include "Widgets/Images/SLayeredImage.h"
+
+namespace UE::RemoteControlUI::Private
+{
+	FSlateIcon GetActionIcon(const FRCSignatureActionIcon& In)
+	{
+		return FSlateIcon(In.StyleSetName, In.StyleName, NAME_None, In.OverlayStyleName);
+	}
+}
+
+void SRCSignatureActionIcon::Construct(const FArguments& InArgs)
+{
+	ChildSlot
+		.HAlign(HAlign_Center)
+		.VAlign(VAlign_Center)
+		[
+			SAssignNew(LayeredImage, SLayeredImage)
+			.DesiredSizeOverride(FVector2D(SRCSignatureActionIcon::IconSize))
+		];
+
+	SetActionIcon(InArgs._ActionIcon);
+}
+
+void SRCSignatureActionIcon::SetActionIcon(const FRCSignatureActionIcon& InActionIcon)
+{
+	const FSlateIcon Icon = UE::RemoteControlUI::Private::GetActionIcon(InActionIcon);
+
+	LayeredImage->RemoveAllLayers();
+
+	// Set Layer 0 (Base) image & color
+	LayeredImage->SetImage(Icon.GetIcon());
+	LayeredImage->SetLayerColor(0, InActionIcon.BaseColor);
+
+	// Set Layer 1 (Overlay) image & color
+	if (const FSlateBrush* OverlayIcon = Icon.GetOverlayIcon())
+	{
+		LayeredImage->AddLayer(OverlayIcon);
+
+		// Overlay color defaults to IconColor if not set
+		LayeredImage->SetLayerColor(1, InActionIcon.OverlayColor.Get(InActionIcon.BaseColor));
+	}
+}
diff --git a/Engine/Plugins/VirtualProduction/RemoteControl/Source/RemoteControlUI/Private/UI/Signature/Actions/SRCSignatureActionIcon.h b/Engine/Plugins/VirtualProduction/RemoteControl/Source/RemoteControlUI/Private/UI/Signature/Actions/SRCSignatureActionIcon.h
new file mode 100644
--- /dev/null
+++ b/Engine/Plugins/VirtualProduction/RemoteControl/Source/RemoteControlUI/Private/UI/Signature/Actions/SRCSignatureActionIcon.h
@@ -0,0 +1,25 @@
+// Copyright Epic Games, Inc. All Rights Reserved.
+
+#pragma once
+
+#include "RCSignatureAction.h"
+#include "Widgets/SCompoundWidget.h"
+
+class SLayeredImage;
+
+class SRCSignatureActionIcon : public SCompoundWidget
+{
+public:
+	SLATE_BEGIN_ARGS(SRCSignatureActionIcon) {}
+		SLATE_ARGUMENT(FRCSignatureActionIcon, ActionIcon)
+	SLATE_END_ARGS()
+
+	static constexpr float IconSize = 16.f;
+
+	void Construct(const FArguments& InArgs);
+
+	void SetActionIcon(const FRCSignatureActionIcon& InActionIcon);
+
+private:
+	TSharedPtr<SLayeredImage> LayeredImage;
+};
diff --git a/Engine/Plugins/VirtualProduction/RemoteControl/Source/RemoteControlUI/Private/UI/Signature/Details/SRCSignatureDetails.cpp b/Engine/Plugins/VirtualProduction/RemoteControl/Source/RemoteControlUI/Private/UI/Signature/Details/SRCSignatureDetails.cpp
new file mode 100644
--- /dev/null
+++ b/Engine/Plugins/VirtualProduction/RemoteControl/Source/RemoteControlUI/Private/UI/Signature/Details/SRCSignatureDetails.cpp
@@ -0,0 +1,104 @@
+// Copyright Epic Games, Inc. All Rights Reserved.
+
+#include "SRCSignatureDetails.h"
+#include "DetailsViewArgs.h"
+#include "IStructureDataProvider.h"
+#include "IStructureDetailsView.h"
+#include "Modules/ModuleManager.h"
+#include "PropertyEditorModule.h"
+#include "RemoteControlSignatureRegistry.h"
+#include "ScopedTransaction.h"
+#include "UI/Signature/Items/RCSignatureTreeRootItem.h"
+#include "UI/Signature/RCSignatureTreeItemSelection.h"
+
+#define LOCTEXT_NAMESPACE "SRCSignatureDetails"
+
+void SRCSignatureDetails::Construct(const FArguments& InArgs, URemoteControlSignatureRegistry* InSignatureRegistry, const TSharedRef<FRCSignatureTreeItemSelection>& InSelection)
+{
+	SignatureRegistryWeak = InSignatureRegistry;
+	SelectionWeak = InSelection;
+
+	InSelection->OnSelectionChanged().AddSP(this, &SRCSignatureDetails::Refresh);
+
+	FDetailsViewArgs DetailsViewArgs;
+	DetailsViewArgs.NameAreaSettings = FDetailsViewArgs::HideNameArea;
+	DetailsViewArgs.NotifyHook = this;
+
+	FPropertyEditorModule& PropertyEditorModule = FModuleManager::GetModuleChecked<FPropertyEditorModule>("PropertyEditor");
+	StructDetailsView = PropertyEditorModule.CreateStructureDetailView(DetailsViewArgs, FStructureDetailsViewArgs(), /*StructOnScope*/nullptr);
+	StructDetailsView->GetOnFinishedChangingPropertiesDelegate().AddSP(this, &SRCSignatureDetails::OnFinishedChangingProperties);
+
+	ChildSlot
+	[
+		StructDetailsView->GetWidget().ToSharedRef()
+	];
+
+	Refresh();
+}
+
+SRCSignatureDetails::~SRCSignatureDetails()
+{
+	if (TSharedPtr<FRCSignatureTreeItemSelection> Selection = SelectionWeak.Pin())
+	{
+		Selection->OnSelectionChanged().RemoveAll(this);
+	}
+}
+
+void SRCSignatureDetails::Refresh()
+{
+	TSharedRef<FStructOnScopeStructureDataProvider> StructProvider = MakeShared<FStructOnScopeStructureDataProvider>();
+	StructProvider->SetStructData(GatherStructOnScopes());
+	StructDetailsView->SetStructureProvider(StructProvider);
+}
+
+TArray<TSharedPtr<FStructOnScope>> SRCSignatureDetails::GatherStructOnScopes() const
+{
+	TArray<TSharedPtr<FStructOnScope>> StructOnScopes;
+
+	TSharedPtr<FRCSignatureTreeItemSelection> Selection = SelectionWeak.Pin();
+	if (!Selection.IsValid())
+	{
+		return StructOnScopes;
+	}
+
+	TConstArrayView<TWeakPtr<FRCSignatureTreeItemBase>> SelectedItems = Selection->GetSelectedItemsView();
+	StructOnScopes.Reserve(SelectedItems.Num());
+
+	for (const TWeakPtr<FRCSignatureTreeItemBase>& SelectedItemWeak : SelectedItems)
+	{
+		TSharedPtr<FRCSignatureTreeItemBase> SelectedItem = SelectedItemWeak.Pin();
+		if (!SelectedItem.IsValid())
+		{
+			continue;
+		}
+
+		TSharedPtr<FStructOnScope> StructOnScope = SelectedItem->MakeSelectionStruct();
+		if (!StructOnScope.IsValid())
+		{
+			continue;
+		}
+
+		StructOnScopes.Add(MoveTemp(StructOnScope));
+	}
+
+	return StructOnScopes;
+}
+
+void SRCSignatureDetails::OnFinishedChangingProperties(const FPropertyChangedEvent& InChangeEvent)
+{
+	// End Transaction
+	CurrentTransaction.Reset();
+}
+
+void SRCSignatureDetails::NotifyPreChange(FEditPropertyChain* InPropertyAboutToChange)
+{
+	if (URemoteControlSignatureRegistry* SignatureRegistry = SignatureRegistryWeak.Get())
+	{
+		// Begin Transaction
+		ensureAlways(!CurrentTransaction.IsValid());
+		CurrentTransaction = MakeShared<FScopedTransaction>(LOCTEXT("EditSignature", "Edit Signature"));
+		SignatureRegistry->Modify();
+	}
+}
+
+#undef LOCTEXT_NAMESPACE 
diff --git a/Engine/Plugins/VirtualProduction/RemoteControl/Source/RemoteControlUI/Private/UI/Signature/Details/SRCSignatureDetails.h b/Engine/Plugins/VirtualProduction/RemoteControl/Source/RemoteControlUI/Private/UI/Signature/Details/SRCSignatureDetails.h
new file mode 100644
--- /dev/null
+++ b/Engine/Plugins/VirtualProduction/RemoteControl/Source/RemoteControlUI/Private/UI/Signature/Details/SRCSignatureDetails.h
@@ -0,0 +1,45 @@
+// Copyright Epic Games, Inc. All Rights Reserved.
+
+#pragma once
+
+#include "Containers/ContainersFwd.h"
+#include "Misc/NotifyHook.h"
+#include "UObject/WeakObjectPtr.h"
+#include "Widgets/SCompoundWidget.h"
+
+class FRCSignatureTreeItemSelection;
+class FScopedTransaction;
+class FStructOnScope;
+class IStructureDetailsView;
+class URemoteControlSignatureRegistry;
+struct FPropertyChangedEvent;
+
+class SRCSignatureDetails : public SCompoundWidget, public FNotifyHook
+{
+public:
+	SLATE_BEGIN_ARGS(SRCSignatureDetails) {}
+	SLATE_END_ARGS()
+
+	void Construct(const FArguments& InArgs, URemoteControlSignatureRegistry* InSignatureRegistry, const TSharedRef<FRCSignatureTreeItemSelection>& InSelection);
+
+	virtual ~SRCSignatureDetails() override;
+
+	void Refresh();
+
+private:
+	TArray<TSharedPtr<FStructOnScope>> GatherStructOnScopes() const;
+
+	void OnFinishedChangingProperties(const FPropertyChangedEvent& InChangeEvent);
+
+	//~ Begin FNotifyHook
+	virtual void NotifyPreChange(FEditPropertyChain* InPropertyAboutToChange) override;
+	//~ End FNotifyHook
+
+	TSharedPtr<IStructureDetailsView> StructDetailsView;
+
+	TSharedPtr<FScopedTransaction> CurrentTransaction;
+
+	TWeakObjectPtr<URemoteControlSignatureRegistry> SignatureRegistryWeak;
+
+	TWeakPtr<FRCSignatureTreeItemSelection> SelectionWeak;
+};
diff --git a/Engine/Plugins/VirtualProduction/RemoteControl/Source/RemoteControlUI/Private/UI/Signature/Items/RCSignatureTreeActionItem.cpp b/Engine/Plugins/VirtualProduction/RemoteControl/Source/RemoteControlUI/Private/UI/Signature/Items/RCSignatureTreeActionItem.cpp
new file mode 100644
--- /dev/null
+++ b/Engine/Plugins/VirtualProduction/RemoteControl/Source/RemoteControlUI/Private/UI/Signature/Items/RCSignatureTreeActionItem.cpp
@@ -0,0 +1,129 @@
+// Copyright Epic Games, Inc. All Rights Reserved.
+
+#include "RCSignatureTreeActionItem.h"
+#include "RCSignatureAction.h"
+#include "RCSignatureTreeFieldItem.h"
+#include "RCSignatureTreeRootItem.h"
+#include "RemoteControlSignature.h"
+#include "RemoteControlSignatureRegistry.h"
+#include "ScopedTransaction.h"
+#include "UI/Signature/RCSignatureTreeItemSelection.h"
+
+#define LOCTEXT_NAMESPACE "RCSignatureTreeActionItem"
+
+FRCSignatureTreeActionItem::FRCSignatureTreeActionItem(int32 InActionIndex, const TSharedPtr<SRCSignatureTree>& InSignatureTree)
+	: FRCSignatureTreeItemBase(InSignatureTree)
+	, ActionIndex(InActionIndex)
+{
+	// Action Items are hidden in Tree View and instead shown in a Horizontal List next to its Parent
+	AddTreeViewFlags(ERCSignatureTreeItemViewFlags::Hidden);
+}
+
+const FRCSignatureActionDefinition* FRCSignatureTreeActionItem::FindActionDefinition() const
+{
+	TSharedPtr<FRCSignatureTreeFieldItem> FieldItem = GetParentFieldItem();
+	if (!FieldItem.IsValid())
+	{
+		return nullptr;
+	}
+
+	const FRCSignatureField* Field = FieldItem->FindField();
+	if (!Field || !Field->ActionDefinitions.IsValidIndex(ActionIndex))
+	{
+		return nullptr;
+	}
+
+	return &Field->ActionDefinitions[ActionIndex];
+}
+
+FRCSignatureActionDefinition* FRCSignatureTreeActionItem::FindActionDefinitionMutable()
+{
+	FRCSignatureField* Field = FindParentFieldMutable();
+	if (!Field || !Field->ActionDefinitions.IsValidIndex(ActionIndex))
+	{
+		return nullptr;
+	}
+
+	return &Field->ActionDefinitions[ActionIndex];
+}
+
+FRCSignatureActionIcon FRCSignatureTreeActionItem::GetIcon() const
+{
+	if (const FRCSignatureActionDefinition* ActionDefinition = FindActionDefinition())
+	{
+		if (const FRCSignatureAction* Action = ActionDefinition->GetAction())
+		{
+			return Action->GetIcon();
+		}
+	}
+	return FRCSignatureActionIcon();
+}
+
+void FRCSignatureTreeActionItem::BuildPathSegment(FStringBuilderBase& InStringBuilder) const
+{
+	InStringBuilder << ActionIndex;
+}
+
+TOptional<bool> FRCSignatureTreeActionItem::IsEnabled() const
+{
+	return true;
+}
+
+int32 FRCSignatureTreeActionItem::RemoveFromRegistry()
+{
+	URemoteControlSignatureRegistry* Registry;
+
+	FRCSignatureField* Field = FindParentFieldMutable(&Registry);
+	if (!Field || !Field->ActionDefinitions.IsValidIndex(ActionIndex))
+	{
+		return 0;
+	}
+
+	check(Registry);
+
+	FScopedTransaction Transaction(LOCTEXT("RemoveAction", "Remove Action"));
+	Registry->Modify();
+	Field->ActionDefinitions.RemoveAt(ActionIndex);
+	return 1;
+}
+
+FText FRCSignatureTreeActionItem::GetDisplayNameText() const
+{
+	// Unused, as Action Items are hidden in Tree View
+	return FText::GetEmpty();
+}
+
+FText FRCSignatureTreeActionItem::GetDescription() const
+{
+	// Unused, as Action Items are hidden in Tree View
+	return FText::GetEmpty();
+}
+
+TSharedPtr<FStructOnScope> FRCSignatureTreeActionItem::MakeSelectionStruct()
+{
+	if (FRCSignatureActionDefinition* ActionDefinition = FindActionDefinitionMutable())
+	{
+		return ActionDefinition->MakeStructOnScope();
+	}
+	return nullptr;
+}
+
+TSharedPtr<FRCSignatureTreeFieldItem> FRCSignatureTreeActionItem::GetParentFieldItem() const
+{
+	if (TSharedPtr<FRCSignatureTreeItemBase> Parent = GetParent())
+	{
+		return Parent->MutableCast<FRCSignatureTreeFieldItem>();
+	}
+	return nullptr;
+}
+
+FRCSignatureField* FRCSignatureTreeActionItem::FindParentFieldMutable(URemoteControlSignatureRegistry** OutRegistry)
+{
+	if (TSharedPtr<FRCSignatureTreeFieldItem> ParentFieldItem = GetParentFieldItem())
+	{
+		return ParentFieldItem->FindFieldMutable(OutRegistry);
+	}
+	return nullptr;
+}
+
+#undef LOCTEXT_NAMESPACE
diff --git a/Engine/Plugins/VirtualProduction/RemoteControl/Source/RemoteControlUI/Private/UI/Signature/Items/RCSignatureTreeActionItem.h b/Engine/Plugins/VirtualProduction/RemoteControl/Source/RemoteControlUI/Private/UI/Signature/Items/RCSignatureTreeActionItem.h
new file mode 100644
--- /dev/null
+++ b/Engine/Plugins/VirtualProduction/RemoteControl/Source/RemoteControlUI/Private/UI/Signature/Items/RCSignatureTreeActionItem.h
@@ -0,0 +1,44 @@
+// Copyright Epic Games, Inc. All Rights Reserved.
+
+#pragma once
+
+#include "RCSignatureTreeItemBase.h"
+
+class FRCSignatureTreeFieldItem;
+class URemoteControlSignatureRegistry;
+struct FRCSignatureActionDefinition;
+struct FRCSignatureActionIcon;
+struct FRCSignatureField;
+
+/** Item class representing a Field owned by a Signature */
+class FRCSignatureTreeActionItem : public FRCSignatureTreeItemBase
+{
+public:
+	static constexpr ERCSignatureTreeItemType StaticItemType = ERCSignatureTreeItemType::Action;
+
+	explicit FRCSignatureTreeActionItem(int32 InActionIndex, const TSharedPtr<SRCSignatureTree>& InSignatureTree);
+
+	const FRCSignatureActionDefinition* FindActionDefinition() const;
+
+	FRCSignatureActionDefinition* FindActionDefinitionMutable();
+
+	FRCSignatureActionIcon GetIcon() const;
+
+protected:
+	//~ Begin FRCSignatureTreeItemBase
+	virtual void BuildPathSegment(FStringBuilderBase& InStringBuilder) const override;
+	virtual TOptional<bool> IsEnabled() const override;
+	virtual int32 RemoveFromRegistry() override;
+	virtual FText GetDisplayNameText() const override;
+	virtual FText GetDescription() const override;
+	virtual TSharedPtr<FStructOnScope> MakeSelectionStruct() override;
+	virtual ERCSignatureTreeItemType GetItemType() const override { return StaticItemType; }
+	//~ End FRCSignatureTreeItemBase
+
+private:
+	TSharedPtr<FRCSignatureTreeFieldItem> GetParentFieldItem() const;
+
+	FRCSignatureField* FindParentFieldMutable(URemoteControlSignatureRegistry** OutRegistry = nullptr);
+
+	int32 ActionIndex;
+};
diff --git a/Engine/Plugins/VirtualProduction/RemoteControl/Source/RemoteControlUI/Private/UI/Signature/Items/RCSignatureTreeFieldItem.cpp b/Engine/Plugins/VirtualProduction/RemoteControl/Source/RemoteControlUI/Private/UI/Signature/Items/RCSignatureTreeFieldItem.cpp
--- a/Engine/Plugins/VirtualProduction/RemoteControl/Source/RemoteControlUI/Private/UI/Signature/Items/RCSignatureTreeFieldItem.cpp
+++ b/Engine/Plugins/VirtualProduction/RemoteControl/Source/RemoteControlUI/Private/UI/Signature/Items/RCSignatureTreeFieldItem.cpp
@@ -1,6 +1,7 @@
 // Copyright Epic Games, Inc. All Rights Reserved.
 
 #include "RCSignatureTreeFieldItem.h"
+#include "RCSignatureTreeActionItem.h"
 #include "RCSignatureTreeSignatureItem.h"
 #include "RemoteControlSignature.h"
 #include "RemoteControlSignatureRegistry.h"
@@ -14,6 +15,49 @@ FRCSignatureTreeFieldItem::FRCSignatureTreeFieldItem(int32 InFieldIndex, const T
 {
 }
 
+const FRCSignatureField* FRCSignatureTreeFieldItem::FindField() const
+{
+	TSharedPtr<FRCSignatureTreeSignatureItem> ParentSignatureItem = GetParentSignatureItem();
+	if (!ParentSignatureItem.IsValid())
+	{
+		return nullptr;
+	}
+
+	const FRCSignature* Signature = ParentSignatureItem->FindSignature();
+	if (!Signature || !Signature->Fields.IsValidIndex(FieldIndex))
+	{
+		return nullptr;
+	}
+
+	return &Signature->Fields[FieldIndex];
+}
+
+FRCSignatureField* FRCSignatureTreeFieldItem::FindFieldMutable(URemoteControlSignatureRegistry** OutRegistry)
+{
+	FRCSignature* Signature = FindParentSignature(OutRegistry);
+	if (!Signature || !Signature->Fields.IsValidIndex(FieldIndex))
+	{
+		return nullptr;
+	}
+	return &Signature->Fields[FieldIndex];
+}
+
+void FRCSignatureTreeFieldItem::AddAction(const UScriptStruct* InActionType)
+{
+	URemoteControlSignatureRegistry* Registry;
+
+	FRCSignatureField* Field = FindFieldMutable(&Registry);
+	if (!Field)
+	{
+		return;
+	}
+
+	FScopedTransaction Transaction(LOCTEXT("AddSignatureAction", "Add Signature Action"));
+	check(Registry);
+	Registry->Modify();
+	Field->ActionDefinitions.Emplace(InActionType, *Field);
+}
+
 void FRCSignatureTreeFieldItem::BuildPathSegment(FStringBuilderBase& InBuilder) const
 {
 	InBuilder << FieldIndex;
@@ -90,46 +134,36 @@ int32 FRCSignatureTreeFieldItem::RemoveFromRegistry()
 	return 1;
 }
 
-FRCSignatureTreeSignatureItem* FRCSignatureTreeFieldItem::GetParentSignatureItem() const
+void FRCSignatureTreeFieldItem::GenerateChildren(TArray<TSharedPtr<FRCSignatureTreeItemBase>>& OutChildren) const
 {
-	if (TSharedPtr<FRCSignatureTreeItemBase> Parent = GetParent())
+	const FRCSignatureField* Field = FindField();
+	if (!Field)
 	{
-		return Parent->AsSignatureItem();
+		return;
 	}
-	return nullptr;
-}
 
-const FRCSignatureField* FRCSignatureTreeFieldItem::FindField() const
-{
-	FRCSignatureTreeSignatureItem* ParentSignatureItem = GetParentSignatureItem();
-	if (!ParentSignatureItem)
-	{
-		return nullptr;
-	}
+	const TSharedPtr<SRCSignatureTree> SignatureTree = GetSignatureTree();
 
-	const FRCSignature* Signature = ParentSignatureItem->FindSignature();
-	if (!Signature || !Signature->Fields.IsValidIndex(FieldIndex))
+	OutChildren.Reserve(OutChildren.Num() + Field->ActionDefinitions.Num());
+	for (int32 ActionIndex = 0; ActionIndex < Field->ActionDefinitions.Num(); ++ActionIndex)
 	{
-		return nullptr;
+		OutChildren.Add(MakeShared<FRCSignatureTreeActionItem>(ActionIndex, SignatureTree));
 	}
-
-	return &Signature->Fields[FieldIndex];
 }
 
-FRCSignatureField* FRCSignatureTreeFieldItem::FindFieldMutable(URemoteControlSignatureRegistry** OutRegistry)
+TSharedPtr<FRCSignatureTreeSignatureItem> FRCSignatureTreeFieldItem::GetParentSignatureItem() const
 {
-	FRCSignature* Signature = FindParentSignature(OutRegistry);
-	if (!Signature || !Signature->Fields.IsValidIndex(FieldIndex))
+	if (TSharedPtr<FRCSignatureTreeItemBase> Parent = GetParent())
 	{
-		return nullptr;
+		return Parent->MutableCast<FRCSignatureTreeSignatureItem>();
 	}
-	return &Signature->Fields[FieldIndex];
+	return nullptr;
 }
 
 FRCSignature* FRCSignatureTreeFieldItem::FindParentSignature(URemoteControlSignatureRegistry** OutRegistry)
 {
-	FRCSignatureTreeSignatureItem* ParentSignatureItem = GetParentSignatureItem();
-	if (!ParentSignatureItem)
+	TSharedPtr<FRCSignatureTreeSignatureItem> ParentSignatureItem = GetParentSignatureItem();
+	if (!ParentSignatureItem.IsValid())
 	{
 		return nullptr;
 	}
diff --git a/Engine/Plugins/VirtualProduction/RemoteControl/Source/RemoteControlUI/Private/UI/Signature/Items/RCSignatureTreeFieldItem.h b/Engine/Plugins/VirtualProduction/RemoteControl/Source/RemoteControlUI/Private/UI/Signature/Items/RCSignatureTreeFieldItem.h
--- a/Engine/Plugins/VirtualProduction/RemoteControl/Source/RemoteControlUI/Private/UI/Signature/Items/RCSignatureTreeFieldItem.h
+++ b/Engine/Plugins/VirtualProduction/RemoteControl/Source/RemoteControlUI/Private/UI/Signature/Items/RCSignatureTreeFieldItem.h
@@ -4,6 +4,7 @@
 
 #include "RCSignatureTreeItemBase.h"
 
+class FRCSignatureTreeSignatureItem;
 class URemoteControlSignatureRegistry;
 struct FRCSignature;
 struct FRCSignatureField;
@@ -12,8 +13,16 @@ struct FRCSignatureField;
 class FRCSignatureTreeFieldItem : public FRCSignatureTreeItemBase
 {
 public:
+	static constexpr ERCSignatureTreeItemType StaticItemType = ERCSignatureTreeItemType::Field;
+
 	explicit FRCSignatureTreeFieldItem(int32 InFieldIndex, const TSharedPtr<SRCSignatureTree>& InSignatureTree);
 
+	const FRCSignatureField* FindField() const;
+
+	FRCSignatureField* FindFieldMutable(URemoteControlSignatureRegistry** OutRegistry);
+
+	void AddAction(const UScriptStruct* InActionType);
+
 protected:
 	//~ Begin FRCSignatureTreeItemBase
 	virtual void BuildPathSegment(FStringBuilderBase& InStringBuilder) const override;
@@ -22,14 +31,12 @@ protected:
 	virtual FText GetDisplayNameText() const override;
 	virtual FText GetDescription() const override;
 	virtual int32 RemoveFromRegistry() override;
+	virtual ERCSignatureTreeItemType GetItemType() const override { return StaticItemType; }
+	virtual void GenerateChildren(TArray<TSharedPtr<FRCSignatureTreeItemBase>>& OutChildren) const override;
 	//~ End FRCSignatureTreeItemBase
 
 private:
-	FRCSignatureTreeSignatureItem* GetParentSignatureItem() const;
-
-	const FRCSignatureField* FindField() const;
-
-	FRCSignatureField* FindFieldMutable(URemoteControlSignatureRegistry** OutRegistry);
+	TSharedPtr<FRCSignatureTreeSignatureItem> GetParentSignatureItem() const;
 
 	FRCSignature* FindParentSignature(URemoteControlSignatureRegistry** OutRegistry);
 
diff --git a/Engine/Plugins/VirtualProduction/RemoteControl/Source/RemoteControlUI/Private/UI/Signature/Items/RCSignatureTreeItemBase.cpp b/Engine/Plugins/VirtualProduction/RemoteControl/Source/RemoteControlUI/Private/UI/Signature/Items/RCSignatureTreeItemBase.cpp
--- a/Engine/Plugins/VirtualProduction/RemoteControl/Source/RemoteControlUI/Private/UI/Signature/Items/RCSignatureTreeItemBase.cpp
+++ b/Engine/Plugins/VirtualProduction/RemoteControl/Source/RemoteControlUI/Private/UI/Signature/Items/RCSignatureTreeItemBase.cpp
@@ -1,6 +1,8 @@
 // Copyright Epic Games, Inc. All Rights Reserved.
 
 #include "RCSignatureTreeItemBase.h"
+#include "RCSignatureTreeRootItem.h"
+#include "UI/Signature/RCSignatureTreeItemSelection.h"
 #include "UI/Signature/SRCSignatureTree.h"
 
 FRCSignatureTreeItemBase::FRCSignatureTreeItemBase(const TSharedPtr<SRCSignatureTree>& InSignatureTree)
@@ -9,19 +11,36 @@ FRCSignatureTreeItemBase::FRCSignatureTreeItemBase(const TSharedPtr<SRCSignature
 {
 }
 
-void FRCSignatureTreeItemBase::AddFlags(ERCSignatureTreeItemFlags InFlags)
+void FRCSignatureTreeItemBase::AddTreeViewFlags(ERCSignatureTreeItemViewFlags InFlags)
 {
-	EnumAddFlags(Flags, InFlags);
+	EnumAddFlags(TreeViewFlags, InFlags);
 }
 
-void FRCSignatureTreeItemBase::RemoveFlags(ERCSignatureTreeItemFlags InFlags)
+void FRCSignatureTreeItemBase::RemoveTreeViewFlags(ERCSignatureTreeItemViewFlags InFlags)
 {
-	EnumRemoveFlags(Flags, InFlags);
+	EnumRemoveFlags(TreeViewFlags, InFlags);
 }
 
-bool FRCSignatureTreeItemBase::HasAnyFlags(ERCSignatureTreeItemFlags InFlags) const
+bool FRCSignatureTreeItemBase::HasAnyTreeViewFlags(ERCSignatureTreeItemViewFlags InFlags) const
 {
-	return EnumHasAnyFlags(Flags, InFlags);
+	return EnumHasAnyFlags(TreeViewFlags, InFlags);
+}
+
+void FRCSignatureTreeItemBase::SetSelected(bool bInSelected, bool bInIsMultiSelection)
+{
+	if (TSharedPtr<FRCSignatureTreeItemSelection> Selection = GetRootSelection())
+	{
+		Selection->SetSelected(SharedThis(this), bInSelected, bInIsMultiSelection);
+	}
+}
+
+bool FRCSignatureTreeItemBase::IsSelected() const
+{
+	if (TSharedPtr<FRCSignatureTreeItemSelection> Selection = GetRootSelection())
+	{
+		return Selection->IsSelected(SharedThis(this));
+	}
+	return false;
 }
 
 void FRCSignatureTreeItemBase::RebuildChildren()
@@ -31,7 +50,7 @@ void FRCSignatureTreeItemBase::RebuildChildren()
 	OldChildren.Reserve(Children.Num());
 	for (const TSharedPtr<FRCSignatureTreeItemBase>& Child : Children)
 	{
-		OldChildren.Add(Child->Path, Child);
+		OldChildren.Add(Child->GetPathId(), Child);
 	}
 
 	Children.Reset();
@@ -41,12 +60,14 @@ void FRCSignatureTreeItemBase::RebuildChildren()
 	for (const TSharedPtr<FRCSignatureTreeItemBase>& Child : Children)
 	{
 		Child->Initialize(This);
-		if (TSharedPtr<FRCSignatureTreeItemBase>* OldChild = OldChildren.Find(Child->Path))
+		if (TSharedPtr<FRCSignatureTreeItemBase>* OldChild = OldChildren.Find(Child->GetPathId()))
 		{
 			Child->RestoreFrom(*OldChild);
 		}
 		Child->RebuildChildren();
 	}
+
+	PostChildrenRebuild();
 }
 
 void FRCSignatureTreeItemBase::VisitChildren(TFunctionRef<bool(const TSharedPtr<FRCSignatureTreeItemBase>&)> InCallable, bool bInRecursive)
@@ -68,6 +89,7 @@ void FRCSignatureTreeItemBase::VisitChildren(TFunctionRef<bool(const TSharedPtr<
 void FRCSignatureTreeItemBase::Initialize(const TSharedPtr<FRCSignatureTreeItemBase>& InParent)
 {
 	ParentWeak = InParent;
+	SelectionWeak = GetRootSelection();
 	Path = BuildPath();
 }
 
@@ -77,11 +99,41 @@ void FRCSignatureTreeItemBase::RestoreFrom(const TSharedPtr<FRCSignatureTreeItem
 	{
 		// The only important things in restoration are flags & children.
 		// Children are only restored so they can restore their flags
-		Flags = InOldItem->GetFlags();
+		TreeViewFlags = InOldItem->GetTreeViewFlags();
 		Children = InOldItem->GetChildren();
 	}
 }
 
+TSharedPtr<FRCSignatureTreeItemSelection> FRCSignatureTreeItemBase::GetRootSelection() const
+{
+	TSharedPtr<FRCSignatureTreeItemSelection> Selection = SelectionWeak.Pin();
+	if (Selection.IsValid())
+	{
+		return Selection;
+	}
+
+	TSharedPtr<const FRCSignatureTreeItemBase> Parent = GetParent();
+	if (!Parent.IsValid())
+	{
+		return nullptr;
+	}
+
+	while (Parent->GetParent().IsValid())
+	{
+		Parent = Parent->GetParent();
+	}
+
+	TSharedPtr<const FRCSignatureTreeRootItem> RootItem = Parent->Cast<FRCSignatureTreeRootItem>();
+	if (!RootItem.IsValid())
+	{
+		return nullptr;
+	}
+
+	Selection = RootItem->GetSelection();
+	SelectionWeak = Selection;
+	return Selection;
+}
+
 FName FRCSignatureTreeItemBase::BuildPath() const
 {
 	// For now, Signatures would have at most 2 items here, and Fields 3.
@@ -98,8 +150,9 @@ FName FRCSignatureTreeItemBase::BuildPath() const
 	}
 
 	// Expectation is that signature items will generate a path from Guid (32 characters)
-	// with a dot delimiter to the field path which is an index (1-2 characters for the numbers in most use cases)
-	// the expected size of the string is 34 characters... rounded up to a multiple of 8 => 40
+	// with a first dot to the field path (index : 1-2 characters for the numbers in most use cases)
+	// and a second dot with the action path (index : 1-2 characters)
+	// the expected size of the string is 36 characters... rounded up to a multiple of 8 => 40
 	TStringBuilder<40> PathBuilder;
 
 	// Build Path starting from the farthest Ancestor.
diff --git a/Engine/Plugins/VirtualProduction/RemoteControl/Source/RemoteControlUI/Private/UI/Signature/Items/RCSignatureTreeItemBase.h b/Engine/Plugins/VirtualProduction/RemoteControl/Source/RemoteControlUI/Private/UI/Signature/Items/RCSignatureTreeItemBase.h
--- a/Engine/Plugins/VirtualProduction/RemoteControl/Source/RemoteControlUI/Private/UI/Signature/Items/RCSignatureTreeItemBase.h
+++ b/Engine/Plugins/VirtualProduction/RemoteControl/Source/RemoteControlUI/Private/UI/Signature/Items/RCSignatureTreeItemBase.h
@@ -4,15 +4,26 @@
 
 #include "UI/BaseLogicUI/RCLogicModeBase.h"
 
+class FRCSignatureTreeItemSelection;
+class FStructOnScope;
 class SRCSignatureTree;
 
-enum ERCSignatureTreeItemFlags : uint8
+enum class ERCSignatureTreeItemViewFlags : uint8
 {
 	None = 0,
 	Expanded = 1 << 0,
-	Selected = 1 << 1,
+	Hidden   = 1 << 1,
+};
+ENUM_CLASS_FLAGS(ERCSignatureTreeItemViewFlags)
+
+enum class ERCSignatureTreeItemType : uint8
+{
+	Undefined,
+	Root,
+	Signature,
+	Field,
+	Action,
 };
-ENUM_CLASS_FLAGS(ERCSignatureTreeItemFlags)
 
 /** Base class for any Item represented in the Signature Tree */
 class FRCSignatureTreeItemBase : public FRCLogicModeBase
@@ -20,6 +31,18 @@ class FRCSignatureTreeItemBase : public FRCLogicModeBase
 	friend class FRCSignatureTreeRootItem;
 
 public:
+	/**
+	 * Creates an Item and Initializes it with the given Parent
+	 * NOTE: this does not build/make child items. RebuildChildren would need to be called separately if necessary.
+	 */
+	template<typename T, typename... ArgTypes UE_REQUIRES(std::is_base_of_v<FRCSignatureTreeItemBase, T>)>
+	static TSharedRef<T> Create(const TSharedPtr<FRCSignatureTreeItemBase>& InParent, ArgTypes&&... InArgs)
+	{
+		TSharedRef<T> Item = MakeShared<T>(Forward<ArgTypes>(InArgs)...);
+		Item->Initialize(InParent);
+		return Item;
+	}
+
 	explicit FRCSignatureTreeItemBase(const TSharedPtr<SRCSignatureTree>& InSignatureTree);
 
 	virtual TOptional<bool> IsEnabled() const
@@ -31,18 +54,26 @@ public:
 	{
 	}
 
-	ERCSignatureTreeItemFlags GetFlags() const
+	FName GetPathId() const
+	{
+		return Path;
+	}
+
+	ERCSignatureTreeItemViewFlags GetTreeViewFlags() const
 	{
-		return Flags;
+		return TreeViewFlags;
 	}
 
-	void AddFlags(ERCSignatureTreeItemFlags InFlags);
+	void AddTreeViewFlags(ERCSignatureTreeItemViewFlags InFlags);
 
-	void RemoveFlags(ERCSignatureTreeItemFlags InFlags);
+	void RemoveTreeViewFlags(ERCSignatureTreeItemViewFlags InFlags);
 
-	bool HasAnyFlags(ERCSignatureTreeItemFlags InFlags) const;
+	bool HasAnyTreeViewFlags(ERCSignatureTreeItemViewFlags InFlags) const;
 
-	virtual FText GetDisplayNameText() const = 0;
+	virtual FText GetDisplayNameText() const
+	{
+		return FText::GetEmpty();
+	}
 
 	virtual bool CanEditDisplayNameText() const
 	{
@@ -53,15 +84,51 @@ public:
 	{
 	}
 
-	virtual FText GetDescription() const = 0;
+	virtual FText GetDescription() const
+	{
+		return FText::GetEmpty();
+	}
 
-	virtual int32 RemoveFromRegistry() = 0;
+	virtual int32 RemoveFromRegistry()
+	{
+		return 0;
+	}
+
+	void SetSelected(bool bInSelected, bool bInIsMultiSelection);
 
-	virtual class FRCSignatureTreeSignatureItem* AsSignatureItem()
+	bool IsSelected() const;
+
+	/** Optional: makes a struct on scope for selection visualization in Details Panel */
+	virtual TSharedPtr<FStructOnScope> MakeSelectionStruct()
 	{
 		return nullptr;
 	}
 
+	virtual ERCSignatureTreeItemType GetItemType() const
+	{
+		return ERCSignatureTreeItemType::Undefined;
+	}
+
+	template<typename InItemType>
+	TSharedPtr<InItemType> MutableCast()
+	{
+		if (GetItemType() == InItemType::StaticItemType)
+		{
+			return StaticCastSharedRef<InItemType>(SharedThis(this));
+		}
+		return nullptr;
+	}
+
+	template<typename InItemType>
+	TSharedPtr<const InItemType> Cast() const
+	{
+		if (GetItemType() == InItemType::StaticItemType)
+		{
+			return StaticCastSharedRef<const InItemType>(SharedThis(this));
+		}
+		return nullptr;
+	}
+
 	TConstArrayView<TSharedPtr<FRCSignatureTreeItemBase>> GetChildren() const
 	{
 		return Children;
@@ -82,17 +149,25 @@ public:
 	void VisitChildren(TFunctionRef<bool(const TSharedPtr<FRCSignatureTreeItemBase>&)> InCallable, bool bInRecursive);
 
 protected:
-	virtual void BuildPathSegment(FStringBuilderBase& InBuilder) const = 0;
+	virtual void BuildPathSegment(FStringBuilderBase& InBuilder) const
+	{
+	}
 
 	virtual void GenerateChildren(TArray<TSharedPtr<FRCSignatureTreeItemBase>>& OutChildren) const
 	{
 	}
 
+	virtual void PostChildrenRebuild()
+	{
+	}
+
 private:
 	void Initialize(const TSharedPtr<FRCSignatureTreeItemBase>& InParent);
 
 	void RestoreFrom(const TSharedPtr<FRCSignatureTreeItemBase>& InOldItem);
 
+	TSharedPtr<FRCSignatureTreeItemSelection> GetRootSelection() const;
+
 	/** Builds the path from the root to the item. Each item will be its own segment delimited by a dot */
 	FName BuildPath() const;
 
@@ -105,5 +180,8 @@ private:
 
 	TWeakPtr<SRCSignatureTree> SignatureTreeWeak;
 
-	ERCSignatureTreeItemFlags Flags = ERCSignatureTreeItemFlags::Expanded;
+	ERCSignatureTreeItemViewFlags TreeViewFlags = ERCSignatureTreeItemViewFlags::Expanded;
+
+	/** Cached weak pointer to the selection object to avoid having to go up the hierarchy to find it */
+	mutable TWeakPtr<FRCSignatureTreeItemSelection> SelectionWeak;
 };
diff --git a/Engine/Plugins/VirtualProduction/RemoteControl/Source/RemoteControlUI/Private/UI/Signature/Items/RCSignatureTreeRootItem.cpp b/Engine/Plugins/VirtualProduction/RemoteControl/Source/RemoteControlUI/Private/UI/Signature/Items/RCSignatureTreeRootItem.cpp
--- a/Engine/Plugins/VirtualProduction/RemoteControl/Source/RemoteControlUI/Private/UI/Signature/Items/RCSignatureTreeRootItem.cpp
+++ b/Engine/Plugins/VirtualProduction/RemoteControl/Source/RemoteControlUI/Private/UI/Signature/Items/RCSignatureTreeRootItem.cpp
@@ -4,11 +4,14 @@
 #include "RCSignatureTreeSignatureItem.h"
 #include "RemoteControlSignature.h"
 #include "RemoteControlSignatureRegistry.h"
+#include "UI/Signature/RCSignatureTreeItemSelection.h"
 #include "UI/Signature/SRCSignatureTree.h"
 
 FRCSignatureTreeRootItem::FRCSignatureTreeRootItem(const TSharedPtr<SRCSignatureTree>& InSignatureTree)
 	: FRCSignatureTreeItemBase(InSignatureTree)
+	, Selection(MakeShared<FRCSignatureTreeItemSelection>())
 {
+	SelectionWeak = Selection;
 }
 
 void FRCSignatureTreeRootItem::GenerateChildren(TArray<TSharedPtr<FRCSignatureTreeItemBase>>& OutChildren) const
@@ -34,20 +37,7 @@ void FRCSignatureTreeRootItem::GenerateChildren(TArray<TSharedPtr<FRCSignatureTr
 	}
 }
 
-FText FRCSignatureTreeRootItem::GetDisplayNameText() const
+void FRCSignatureTreeRootItem::PostChildrenRebuild()
 {
-	checkNoEntry();
-	return FText::GetEmpty();
-}
-
-FText FRCSignatureTreeRootItem::GetDescription() const
-{
-	checkNoEntry();
-	return FText::GetEmpty();
-}
-
-int32 FRCSignatureTreeRootItem::RemoveFromRegistry()
-{
-	checkNoEntry();
-	return 0;
+	Selection->RecacheSelectedItems(SharedThis(this));
 }
diff --git a/Engine/Plugins/VirtualProduction/RemoteControl/Source/RemoteControlUI/Private/UI/Signature/Items/RCSignatureTreeRootItem.h b/Engine/Plugins/VirtualProduction/RemoteControl/Source/RemoteControlUI/Private/UI/Signature/Items/RCSignatureTreeRootItem.h
--- a/Engine/Plugins/VirtualProduction/RemoteControl/Source/RemoteControlUI/Private/UI/Signature/Items/RCSignatureTreeRootItem.h
+++ b/Engine/Plugins/VirtualProduction/RemoteControl/Source/RemoteControlUI/Private/UI/Signature/Items/RCSignatureTreeRootItem.h
@@ -4,6 +4,8 @@
 
 #include "RCSignatureTreeItemBase.h"
 
+class FRCSignatureTreeItemSelection;
+
 /**
  * Item class for the Root of all items
  * This is primarily so that the Top level items can have the same set of functionalities as the rest of the items in the tree 
@@ -11,6 +13,8 @@
 class FRCSignatureTreeRootItem : public FRCSignatureTreeItemBase
 {
 public:
+	static constexpr ERCSignatureTreeItemType StaticItemType = ERCSignatureTreeItemType::Root;
+
 	explicit FRCSignatureTreeRootItem(const TSharedPtr<SRCSignatureTree>& InSignatureTree);
 
 	TArray<TSharedPtr<FRCSignatureTreeItemBase>>& GetChildrenMutable()
@@ -18,12 +22,18 @@ public:
 		return Children;
 	}
 
+	const TSharedRef<FRCSignatureTreeItemSelection>& GetSelection() const
+	{
+		return Selection;
+	}
+
 protected:
 	//~ Begin FRCSignatureTreeItemBase
 	virtual void BuildPathSegment(FStringBuilderBase& InBuilder) const override {}
+	virtual ERCSignatureTreeItemType GetItemType() const { return StaticItemType; }
 	virtual void GenerateChildren(TArray<TSharedPtr<FRCSignatureTreeItemBase>>& OutChildren) const override;
-	virtual FText GetDisplayNameText() const override;
-	virtual FText GetDescription() const override;
-	virtual int32 RemoveFromRegistry() override;
+	virtual void PostChildrenRebuild() override;
 	//~ End FRCSignatureTreeItemBase
+
+	TSharedRef<FRCSignatureTreeItemSelection> Selection;
 };
diff --git a/Engine/Plugins/VirtualProduction/RemoteControl/Source/RemoteControlUI/Private/UI/Signature/Items/RCSignatureTreeSignatureItem.cpp b/Engine/Plugins/VirtualProduction/RemoteControl/Source/RemoteControlUI/Private/UI/Signature/Items/RCSignatureTreeSignatureItem.cpp
--- a/Engine/Plugins/VirtualProduction/RemoteControl/Source/RemoteControlUI/Private/UI/Signature/Items/RCSignatureTreeSignatureItem.cpp
+++ b/Engine/Plugins/VirtualProduction/RemoteControl/Source/RemoteControlUI/Private/UI/Signature/Items/RCSignatureTreeSignatureItem.cpp
@@ -6,6 +6,7 @@
 #include "RemoteControlPreset.h"
 #include "RemoteControlSignatureRegistry.h"
 #include "ScopedTransaction.h"
+#include "UI/SRemoteControlPanel.h"
 #include "UI/Signature/SRCSignatureTree.h"
 
 #define LOCTEXT_NAMESPACE "RCSignatureTreeSignatureItem"
@@ -48,7 +49,7 @@ FRCSignature* FRCSignatureTreeSignatureItem::FindSignatureMutable(URemoteControl
 	return nullptr;
 }
 
-bool FRCSignatureTreeSignatureItem::AddField(URemoteControlSignatureRegistry* InRegistry, const FRCExposesPropertyArgs& InPropertyArgs)
+bool FRCSignatureTreeSignatureItem::AddField(URemoteControlSignatureRegistry* InRegistry, const TSharedRef<IPropertyHandle>& InPropertyHandle)
 {
 	FRCSignature* Signature = FindSignatureMutable(InRegistry);
 	if (!Signature)
@@ -58,8 +59,17 @@ bool FRCSignatureTreeSignatureItem::AddField(URemoteControlSignatureRegistry* In
 
 	TArray<FRCSignatureField, TInlineAllocator<1>> Fields;
 	{
-		FRCFieldPathInfo PathInfo(InPropertyArgs.PropertyHandle->GeneratePathToProperty(), /*bSkipDuplicates*/true);
-		Fields.Emplace(FRCSignatureField:: CreateField(MoveTemp(PathInfo), InPropertyArgs.OwnerObject.Get(), InPropertyArgs.GetProperty()));
+		TArray<UObject*> OuterObjects;
+		InPropertyHandle->GetOuterObjects(OuterObjects);
+		Fields.Reserve(OuterObjects.Num());
+
+		const FRCFieldPathInfo PathInfo(InPropertyHandle->GeneratePathToProperty(), /*bSkipDuplicates*/true);
+		const FProperty* Property = InPropertyHandle->GetProperty();
+
+		for (UObject* OuterObject : OuterObjects)
+		{
+			Fields.Emplace(FRCSignatureField::CreateField(PathInfo, OuterObject, Property));
+		}
 	}
 
 	return Signature->AddFields(Fields) > 0;
@@ -175,11 +185,6 @@ int32 FRCSignatureTreeSignatureItem::RemoveFromRegistry()
 	return SignatureRegistry->RemoveSignature(SignatureId);
 }
 
-FRCSignatureTreeSignatureItem* FRCSignatureTreeSignatureItem::AsSignatureItem()
-{
-	return this;
-}
-
 void FRCSignatureTreeSignatureItem::GenerateChildren(TArray<TSharedPtr<FRCSignatureTreeItemBase>>& OutChildren) const
 {
 	const FRCSignature* Signature = FindSignature();
diff --git a/Engine/Plugins/VirtualProduction/RemoteControl/Source/RemoteControlUI/Private/UI/Signature/Items/RCSignatureTreeSignatureItem.h b/Engine/Plugins/VirtualProduction/RemoteControl/Source/RemoteControlUI/Private/UI/Signature/Items/RCSignatureTreeSignatureItem.h
--- a/Engine/Plugins/VirtualProduction/RemoteControl/Source/RemoteControlUI/Private/UI/Signature/Items/RCSignatureTreeSignatureItem.h
+++ b/Engine/Plugins/VirtualProduction/RemoteControl/Source/RemoteControlUI/Private/UI/Signature/Items/RCSignatureTreeSignatureItem.h
@@ -8,14 +8,16 @@
 #include "UObject/WeakObjectPtr.h"
 #include "UObject/WeakObjectPtrTemplates.h"
 
+class IPropertyHandle;
 class URemoteControlSignatureRegistry;
-struct FRCExposesPropertyArgs;
 struct FRCSignature;
 
 /** Item class representing an RC Signature */
 class FRCSignatureTreeSignatureItem : public FRCSignatureTreeItemBase, public IRCSignatureItem
 {
 public:
+	static constexpr ERCSignatureTreeItemType StaticItemType = ERCSignatureTreeItemType::Signature;
+
 	explicit FRCSignatureTreeSignatureItem(const FRCSignature& InSignature, const TSharedPtr<SRCSignatureTree>& InSignatureTree);
 
 	const FGuid& GetSignatureId() const;
@@ -26,7 +28,7 @@ public:
 
 	FRCSignature* FindSignatureMutable(URemoteControlSignatureRegistry* InRegistry);
 
-	bool AddField(URemoteControlSignatureRegistry* InRegistry, const FRCExposesPropertyArgs& InPropertyArgs);
+	bool AddField(URemoteControlSignatureRegistry* InRegistry, const TSharedRef<IPropertyHandle>& InPropertyHandle);
 
 	//~ Begin IRCSignatureItem
 	virtual void ApplySignature(TConstArrayView<TWeakObjectPtr<UObject>> InObjects) override;
@@ -42,7 +44,7 @@ protected:
 	virtual void SetDisplayNameText(const FText& InText) override;
 	virtual FText GetDescription() const override;
 	virtual int32 RemoveFromRegistry() override;
-	virtual FRCSignatureTreeSignatureItem* AsSignatureItem() override;
+	virtual ERCSignatureTreeItemType GetItemType() const override { return StaticItemType; }
 	virtual void GenerateChildren(TArray<TSharedPtr<FRCSignatureTreeItemBase>>& OutChildren) const override;
 	//~ End FRCSignatureTreeItem
 
diff --git a/Engine/Plugins/VirtualProduction/RemoteControl/Source/RemoteControlUI/Private/UI/Signature/RCSignatureTreeItemSelection.cpp b/Engine/Plugins/VirtualProduction/RemoteControl/Source/RemoteControlUI/Private/UI/Signature/RCSignatureTreeItemSelection.cpp
new file mode 100644
--- /dev/null
+++ b/Engine/Plugins/VirtualProduction/RemoteControl/Source/RemoteControlUI/Private/UI/Signature/RCSignatureTreeItemSelection.cpp
@@ -0,0 +1,100 @@
+// Copyright Epic Games, Inc. All Rights Reserved.
+
+#include "RCSignatureTreeItemSelection.h"
+#include "Items/RCSignatureTreeItemBase.h"
+
+FRCSignatureTreeItemSelection::FSelectionScope::FSelectionScope(FRCSignatureTreeItemSelection* InSelection)
+	: Selection(InSelection)
+{
+	check(Selection);
+	OldValue = Selection->bAllowNotifications;
+	Selection->bAllowNotifications = false;
+}
+
+FRCSignatureTreeItemSelection::FSelectionScope::~FSelectionScope()
+{
+	Selection->bAllowNotifications = OldValue;
+	Selection->NotifySelectionChanged();
+}
+
+FRCSignatureTreeItemSelection::FSelectionScope FRCSignatureTreeItemSelection::CreateSelectionScope()
+{
+	return FSelectionScope(this);
+}
+
+void FRCSignatureTreeItemSelection::SetSelected(const TSharedRef<FRCSignatureTreeItemBase>& InItem, bool bInSelected, bool bInMultiSelection)
+{
+	FSelectionScope SelectionScope = CreateSelectionScope();
+
+	if (bInSelected)
+	{
+		if (!bInMultiSelection)
+		{
+			ClearSelection();
+		}
+
+		SelectedPaths.Add(InItem->GetPathId());
+		CachedSelectedItems.AddUnique(InItem);
+	}
+	else
+	{
+		SelectedPaths.Remove(InItem->GetPathId());
+		CachedSelectedItems.Remove(InItem);
+	}
+}
+
+bool FRCSignatureTreeItemSelection::IsSelected(const TSharedPtr<const FRCSignatureTreeItemBase>& InItem) const
+{
+	return InItem.IsValid() && SelectedPaths.Contains(InItem->GetPathId());
+}
+
+void FRCSignatureTreeItemSelection::ClearSelection()
+{
+	FSelectionScope SelectionScope = CreateSelectionScope();
+
+	CachedSelectedItems.Empty();
+	SelectedPaths.Empty();
+}
+
+void FRCSignatureTreeItemSelection::NotifySelectionChanged()
+{
+	if (bAllowNotifications)
+	{
+		OnSelectionChangedDelegate.Broadcast();
+	}
+}
+
+TArray<TSharedPtr<FRCSignatureTreeItemBase>> FRCSignatureTreeItemSelection::GetSelectedItems() const
+{
+	TArray<TSharedPtr<FRCSignatureTreeItemBase>> SelectedItems;
+	SelectedItems.Reserve(CachedSelectedItems.Num());
+
+	for (const TWeakPtr<FRCSignatureTreeItemBase>& SelectedItemWeak : CachedSelectedItems)
+	{
+		if (TSharedPtr<FRCSignatureTreeItemBase> SelectedItem = SelectedItemWeak.Pin())
+		{
+			SelectedItems.Add(MoveTemp(SelectedItem));
+		}
+	}
+
+	return SelectedItems;
+}
+
+void FRCSignatureTreeItemSelection::RecacheSelectedItems(const TSharedRef<FRCSignatureTreeItemBase>& InRoot)
+{
+	// No Selection Scope is set here because this is not changing the Selected Paths
+	// which is what truly determines whether an item is selected.
+	CachedSelectedItems.Empty();
+
+	auto ProcessItem = [&SelectedItems = CachedSelectedItems, &SelectedPaths = SelectedPaths](const TSharedPtr<FRCSignatureTreeItemBase>& InItem)->bool
+	{
+		if (SelectedPaths.Contains(InItem->GetPathId()))
+		{
+			SelectedItems.Add(InItem);
+		}
+		return true;
+	};
+
+	ProcessItem(InRoot);
+	InRoot->VisitChildren(ProcessItem, /*bRecursive*/true);
+}
diff --git a/Engine/Plugins/VirtualProduction/RemoteControl/Source/RemoteControlUI/Private/UI/Signature/RCSignatureTreeItemSelection.h b/Engine/Plugins/VirtualProduction/RemoteControl/Source/RemoteControlUI/Private/UI/Signature/RCSignatureTreeItemSelection.h
new file mode 100644
--- /dev/null
+++ b/Engine/Plugins/VirtualProduction/RemoteControl/Source/RemoteControlUI/Private/UI/Signature/RCSignatureTreeItemSelection.h
@@ -0,0 +1,61 @@
+// Copyright Epic Games, Inc. All Rights Reserved.
+
+#pragma once
+
+#include "Containers/Array.h"
+#include "Delegates/Delegate.h"
+#include "Templates/SharedPointer.h"
+
+class FRCSignatureTreeItemBase;
+
+class FRCSignatureTreeItemSelection : public TSharedFromThis<FRCSignatureTreeItemSelection>
+{
+public:
+	/** Scope class to ensure Notify Selection Changed gets called at most once for the duration of the scope */
+	struct FSelectionScope
+	{
+		explicit FSelectionScope(FRCSignatureTreeItemSelection* InSelection);
+		~FSelectionScope();
+
+	private:
+		FRCSignatureTreeItemSelection* Selection;
+		bool OldValue = false;
+	};
+
+	FSelectionScope CreateSelectionScope();
+
+	TMulticastDelegateRegistration<void()>& OnSelectionChanged() const
+	{
+		return OnSelectionChangedDelegate;
+	}
+
+	void SetSelected(const TSharedRef<FRCSignatureTreeItemBase>& InItem, bool bInSelected, bool bInMultiSelection);
+
+	bool IsSelected(const TSharedPtr<const FRCSignatureTreeItemBase>& InItem) const;
+
+	void ClearSelection();
+
+	void NotifySelectionChanged();
+
+	TConstArrayView<TWeakPtr<FRCSignatureTreeItemBase>> GetSelectedItemsView() const
+	{
+		return CachedSelectedItems;
+	}
+
+	TArray<TSharedPtr<FRCSignatureTreeItemBase>> GetSelectedItems() const;
+
+	/**
+	 * Clears the current selection items (not paths!) and rebuilds a new selection list of all the items that match one of the selected paths.
+	 * @param InRoot the root item to iterate
+	 */
+	void RecacheSelectedItems(const TSharedRef<FRCSignatureTreeItemBase>& InRoot);
+
+private:
+	TSet<FName> SelectedPaths;
+
+	TArray<TWeakPtr<FRCSignatureTreeItemBase>> CachedSelectedItems;
+
+	mutable TMulticastDelegate<void()> OnSelectionChangedDelegate;
+
+	bool bAllowNotifications = true;
+};
diff --git a/Engine/Plugins/VirtualProduction/RemoteControl/Source/RemoteControlUI/Private/UI/Signature/SRCSignaturePanel.cpp b/Engine/Plugins/VirtualProduction/RemoteControl/Source/RemoteControlUI/Private/UI/Signature/SRCSignaturePanel.cpp
--- a/Engine/Plugins/VirtualProduction/RemoteControl/Source/RemoteControlUI/Private/UI/Signature/SRCSignaturePanel.cpp
+++ b/Engine/Plugins/VirtualProduction/RemoteControl/Source/RemoteControlUI/Private/UI/Signature/SRCSignaturePanel.cpp
@@ -1,7 +1,10 @@
 // Copyright Epic Games, Inc. All Rights Reserved.
 
 #include "SRCSignaturePanel.h"
+#include "Actions/RCSignatureActionColumn.h"
 #include "Description/RCSignatureDescriptionColumn.h"
+#include "Details/SRCSignatureDetails.h"
+#include "Items/RCSignatureTreeRootItem.h"
 #include "Items/RCSignatureTreeSignatureItem.h"
 #include "Label/RCSignatureLabelColumn.h"
 #include "Misc/MessageDialog.h"
@@ -20,10 +23,10 @@ void SRCSignaturePanel::Construct(const FArguments& InArgs, const TSharedRef<SRe
 {
 	SRCLogicPanelBase::Construct(SRCLogicPanelBase::FArguments(), InPanel);
 
-	const FRCPanelStyle* RCPanelStyle = &FRemoteControlPanelStyle::Get()->GetWidgetStyle<FRCPanelStyle>("RemoteControlPanel.MinorPanel");
+	const FRCPanelStyle& RCPanelStyle = FRemoteControlPanelStyle::Get()->GetWidgetStyle<FRCPanelStyle>("RemoteControlPanel.MinorPanel");
 
-	// Signature Dock Panel
-	TSharedRef<SRCMinorPanel> SignatureDockPanel = SNew(SRCMinorPanel)
+	// Signature Tree Panel
+	TSharedRef<SRCMinorPanel> SignaturesTreePanel = SNew(SRCMinorPanel)
 		.HeaderLabel(LOCTEXT("SignaturesLabel", "Signatures"))
 		.EnableFooter(false)
 		[
@@ -32,24 +35,37 @@ void SRCSignaturePanel::Construct(const FArguments& InArgs, const TSharedRef<SRe
 				{
 					MakeShared<FRCSignatureLabelColumn>(),
 					MakeShared<FRCSignatureDescriptionColumn>(),
+					MakeShared<FRCSignatureActionColumn>(InArgs._LiveMode),
 				})
 		];
 
+	// Details Panel
+	TSharedRef<SRCMinorPanel> SignaturesDetailsPanel = SNew(SRCMinorPanel)
+		.HeaderLabel(LOCTEXT("DetailsLabel", "Details"))
+		.EnableFooter(false)
+		[
+			SAssignNew(SignatureDetails, SRCSignatureDetails
+				, GetSignatureRegistry()
+				, SignatureTreeView->GetRootItem()->GetSelection())
+		];
+
+	constexpr float ContentPaddingY = 2.f;
+
 	// Add New Signature Button
 	TSharedRef<SButton> AddSignatureButton = SNew(SButton)
 		.AddMetaData<FTagMetaData>(FTagMetaData(TEXT("Add Signature")))
 		.IsEnabled_Lambda([LiveMode = InArgs._LiveMode]() { return !LiveMode.Get(); })
 		.OnClicked(this, &SRCSignaturePanel::OnAddButtonClicked)
 		.ForegroundColor(FSlateColor::UseForeground())
-		.ButtonStyle(&RCPanelStyle->FlatButtonStyle)
-		.ContentPadding(FMargin(4.f, 2.f))
+		.ButtonStyle(&RCPanelStyle.FlatButtonStyle)
+		.ContentPadding(FMargin(4.f, ContentPaddingY))
 		.HAlign(HAlign_Center)
 		.VAlign(VAlign_Center)
 		.Content()
 		[
 			SNew(SBox)
-			.WidthOverride(RCPanelStyle->IconSize.X)
-			.HeightOverride(RCPanelStyle->IconSize.Y)
+			.WidthOverride(RCPanelStyle.IconSize.X)
+			.HeightOverride(RCPanelStyle.IconSize.Y)
 			[
 				SNew(SImage)
 				.ColorAndOpacity(FSlateColor::UseForeground())
@@ -57,13 +73,36 @@ void SRCSignaturePanel::Construct(const FArguments& InArgs, const TSharedRef<SRe
 			]
 		];
 
-	SignatureDockPanel->AddHeaderToolbarItem(EToolbar::Left, AddSignatureButton);
+	SignaturesTreePanel->AddHeaderToolbarItem(EToolbar::Left, AddSignatureButton);
+
+	// Header Toolbar Heights are not fixed and vary depending on what is inside it
+	// Since the Details Panel header has nothing but text in it (for now),
+	// a box of a calculated height is added to compensate and match the height of the Signatures Panel
+	{
+		const float PanelStyleIconSize = RCPanelStyle.IconSize.Y;
+		const float ButtonStylePadding = RCPanelStyle.FlatButtonStyle.NormalPadding.Bottom + RCPanelStyle.FlatButtonStyle.NormalPadding.Top;
+		constexpr float ContentPadding = 2.f * ContentPaddingY;
+
+		SignaturesDetailsPanel->AddHeaderToolbarItem(EToolbar::Left, SNew(SBox)
+			.HeightOverride(PanelStyleIconSize + ButtonStylePadding + ContentPaddingY));
+	}
 
 	ChildSlot
-	.Padding(RCPanelStyle->PanelPadding)
-	[
-		SignatureDockPanel
-	];
+		.Padding(RCPanelStyle.PanelPadding)
+		[
+			SNew(SSplitter)
+			.Orientation(Orient_Horizontal)
+			+ SSplitter::Slot()
+			.Value(0.6f)
+			[
+				SignaturesTreePanel
+			]
+			+ SSplitter::Slot()
+			.Value(0.4f)
+			[
+				SignaturesDetailsPanel
+			]
+		];
 }
 
 URemoteControlSignatureRegistry* SRCSignaturePanel::GetSignatureRegistry() const
@@ -88,6 +127,13 @@ void SRCSignaturePanel::AddToSignature(const FRCExposesPropertyArgs& InPropertyA
 		return;
 	}
 
+	if (!InPropertyArgs.PropertyHandle.IsValid())
+	{
+		return;
+	}
+
+	TSharedRef<IPropertyHandle> PropertyHandle = InPropertyArgs.PropertyHandle.ToSharedRef();
+
 	FScopedTransaction Transaction(LOCTEXT("AddToSignatureTransaction", "Add to Signature"));
 	Registry->Modify();
 
@@ -96,10 +142,16 @@ void SRCSignaturePanel::AddToSignature(const FRCExposesPropertyArgs& InPropertyA
 	TArray<TSharedPtr<FRCSignatureTreeItemBase>> SelectedItems = SignatureTreeView->GetSelectedItems();
 	if (SelectedItems.IsEmpty())
 	{
+		TSharedRef<FRCSignatureTreeRootItem> RootItem = SignatureTreeView->GetRootItem();
+
+		FRCSignature& Signature = Registry->AddSignature();
+
 		// Make a new signature with temp view model item to add the field
-		TSharedRef<FRCSignatureTreeSignatureItem> SignatureItem = MakeShared<FRCSignatureTreeSignatureItem>(Registry->AddSignature(), SignatureTreeView);
-		if (SignatureItem->AddField(Registry, InPropertyArgs))
+		TSharedRef<FRCSignatureTreeSignatureItem> SignatureItem = FRCSignatureTreeItemBase::Create<FRCSignatureTreeSignatureItem>(RootItem, Signature, SignatureTreeView);
+		if (SignatureItem->AddField(Registry, PropertyHandle))
 		{
+			// Select the Signature Item
+			SignatureItem->SetSelected(/*bSelected*/true, /*bMultiSelection*/false);
 			bFieldsAdded = true;
 		}
 	}
@@ -107,9 +159,9 @@ void SRCSignaturePanel::AddToSignature(const FRCExposesPropertyArgs& InPropertyA
 	{
 		for (const TSharedPtr<FRCSignatureTreeItemBase>& Item : SelectedItems)
 		{
-			if (FRCSignatureTreeSignatureItem* SignatureItem = Item->AsSignatureItem())
+			if (TSharedPtr<FRCSignatureTreeSignatureItem> SignatureItem = Item->MutableCast<FRCSignatureTreeSignatureItem>())
 			{
-				if (SignatureItem->AddField(Registry, InPropertyArgs))
+				if (SignatureItem->AddField(Registry, PropertyHandle))
 				{
 					bFieldsAdded = true;
 				}
@@ -119,7 +171,7 @@ void SRCSignaturePanel::AddToSignature(const FRCExposesPropertyArgs& InPropertyA
 
 	if (bFieldsAdded)
 	{
-		SignatureTreeView->Refresh();
+		Refresh();
 	}
 	else
 	{
@@ -143,19 +195,7 @@ TArray<TSharedPtr<FRCLogicModeBase>> SRCSignaturePanel::GetSelectedLogicItems()
 
 FReply SRCSignaturePanel::RequestDeleteSelectedItem()
 {
-	if (!SignatureTreeView.IsValid())
-	{
-		return FReply::Unhandled();
-	}
-
-	const EAppReturnType::Type UserResponse = FMessageDialog::Open(EAppMsgType::YesNo
-		, LOCTEXT("DeleteSelectedWarning", "Delete the Selected Signatures?"));
-
-	if (UserResponse == EAppReturnType::Yes)
-	{
-		DeleteSelectedPanelItems();
-	}
-
+	DeleteSelectedPanelItems();
 	return FReply::Handled();
 }
 
@@ -194,6 +234,30 @@ void SRCSignaturePanel::DeleteSelectedPanelItems()
 	if (SignatureTreeView.IsValid())
 	{
 		SignatureTreeView->DeleteSelectedPanelItems();
+		Refresh();
+	}
+}
+
+void SRCSignaturePanel::PostUndo(bool bInSuccess)
+{
+	Refresh();
+}
+
+void SRCSignaturePanel::PostRedo(bool bInSuccess)
+{
+	Refresh();
+}
+
+void SRCSignaturePanel::Refresh()
+{
+	if (SignatureTreeView.IsValid())
+	{
+		SignatureTreeView->Refresh();
+	}
+
+	if (SignatureDetails.IsValid())
+	{
+		SignatureDetails->Refresh();
 	}
 }
 
@@ -212,11 +276,7 @@ FReply SRCSignaturePanel::OnAddButtonClicked()
 		SignatureRegistry->AddSignature();
 	}
 
-	if (SignatureTreeView.IsValid())
-	{
-		SignatureTreeView->Refresh();
-	}
-
+	Refresh();
 	return FReply::Handled();
 }
 
@@ -238,11 +298,7 @@ FReply SRCSignaturePanel::DeleteAllItems()
 	SignatureRegistry->Modify();
 	SignatureRegistry->EmptySignatures();
 
-	if (SignatureTreeView.IsValid())
-	{
-		SignatureTreeView->Refresh();
-	}
-
+	Refresh();
 	return FReply::Handled();
 }
 
diff --git a/Engine/Plugins/VirtualProduction/RemoteControl/Source/RemoteControlUI/Private/UI/Signature/SRCSignaturePanel.h b/Engine/Plugins/VirtualProduction/RemoteControl/Source/RemoteControlUI/Private/UI/Signature/SRCSignaturePanel.h
--- a/Engine/Plugins/VirtualProduction/RemoteControl/Source/RemoteControlUI/Private/UI/Signature/SRCSignaturePanel.h
+++ b/Engine/Plugins/VirtualProduction/RemoteControl/Source/RemoteControlUI/Private/UI/Signature/SRCSignaturePanel.h
@@ -4,13 +4,12 @@
 
 #include "UI/BaseLogicUI/SRCLogicPanelBase.h"
 
-class FRCSignatureTreeItemBase;
-class IRCSignatureColumn;
+class SRCSignatureDetails;
 class SRCSignatureTree;
 class URemoteControlSignatureRegistry;
 struct FRCExposesPropertyArgs;
 
-class SRCSignaturePanel : public SRCLogicPanelBase
+class SRCSignaturePanel : public SRCLogicPanelBase, public FSelfRegisteringEditorUndoClient
 {
 public:
 	SLATE_BEGIN_ARGS(SRCSignaturePanel) {}
@@ -36,10 +35,20 @@ public:
 	virtual void DeleteSelectedPanelItems() override;
 	//~ End SRCLogicPanelBase
 
+protected:
+	//~ Begin FEditorUndoClient
+	virtual void PostUndo(bool bInSuccess) override;
+	virtual void PostRedo(bool bInSuccess) override;
+	//~ End FEditorUndoClient
+
 private:
+	void Refresh();
+
 	FReply OnAddButtonClicked();
 
 	FReply DeleteAllItems();
 
 	TSharedPtr<SRCSignatureTree> SignatureTreeView;
+
+	TSharedPtr<SRCSignatureDetails> SignatureDetails;
 };
diff --git a/Engine/Plugins/VirtualProduction/RemoteControl/Source/RemoteControlUI/Private/UI/Signature/SRCSignatureRow.cpp b/Engine/Plugins/VirtualProduction/RemoteControl/Source/RemoteControlUI/Private/UI/Signature/SRCSignatureRow.cpp
--- a/Engine/Plugins/VirtualProduction/RemoteControl/Source/RemoteControlUI/Private/UI/Signature/SRCSignatureRow.cpp
+++ b/Engine/Plugins/VirtualProduction/RemoteControl/Source/RemoteControlUI/Private/UI/Signature/SRCSignatureRow.cpp
@@ -58,7 +58,7 @@ TOptional<EItemDropZone> SRCSignatureRow::OnRowCanAcceptDrop(const FDragDropEven
 		return TOptional<EItemDropZone>();
 	}
 
-	if (FRCSignatureTreeSignatureItem* SignatureItem = InItem->AsSignatureItem())
+	if (FRCSignatureTreeSignatureItem* SignatureItem = InItem->MutableCast<FRCSignatureTreeSignatureItem>().Get())
 	{
 		// External Customizations to handle drag drops first
 		for (TSharedRef<IRCSignatureCustomization> Customization : FRemoteControlUIModule::Get().GetSignatureCustomizations())
@@ -114,7 +114,7 @@ FReply SRCSignatureRow::OnRowAcceptDrop(const FDragDropEvent& InDragDropEvent
 		return FReply::Unhandled();
 	}
 
-	if (FRCSignatureTreeSignatureItem* SignatureItem = InItem->AsSignatureItem())
+	if (FRCSignatureTreeSignatureItem* SignatureItem = InItem->MutableCast<FRCSignatureTreeSignatureItem>().Get())
 	{
 		// External Customizations to handle drops first
 		for (TSharedRef<IRCSignatureCustomization> Customization : FRemoteControlUIModule::Get().GetSignatureCustomizations())
diff --git a/Engine/Plugins/VirtualProduction/RemoteControl/Source/RemoteControlUI/Private/UI/Signature/SRCSignatureTree.cpp b/Engine/Plugins/VirtualProduction/RemoteControl/Source/RemoteControlUI/Private/UI/Signature/SRCSignatureTree.cpp
--- a/Engine/Plugins/VirtualProduction/RemoteControl/Source/RemoteControlUI/Private/UI/Signature/SRCSignatureTree.cpp
+++ b/Engine/Plugins/VirtualProduction/RemoteControl/Source/RemoteControlUI/Private/UI/Signature/SRCSignatureTree.cpp
@@ -5,6 +5,7 @@
 #include "Items/RCSignatureTreeItemBase.h"
 #include "Items/RCSignatureTreeRootItem.h"
 #include "Items/RCSignatureTreeSignatureItem.h"
+#include "RCSignatureTreeItemSelection.h"
 #include "RemoteControlPreset.h"
 #include "RemoteControlSignatureRegistry.h"
 #include "SRCSignaturePanel.h"
@@ -23,6 +24,7 @@ void SRCSignatureTree::Construct(const FArguments& InArgs, const TSharedRef<SRCS
 	SignaturePanelWeak = InSignaturePanel;
 
 	RootItem = MakeShared<FRCSignatureTreeRootItem>(SharedThis(this));
+	RootItem->GetSelection()->OnSelectionChanged().AddSP(this, &SRCSignatureTree::UpdateTreeViewSelection);
 
 	const FRCPanelStyle* RCPanelStyle = &FRemoteControlPanelStyle::Get()->GetWidgetStyle<FRCPanelStyle>("RemoteControlPanel.LogicControllersPanel");
 
@@ -50,6 +52,11 @@ void SRCSignatureTree::Construct(const FArguments& InArgs, const TSharedRef<SRCS
 	Refresh();
 }
 
+TSharedRef<FRCSignatureTreeRootItem> SRCSignatureTree::GetRootItem() const
+{
+	return RootItem.ToSharedRef();
+}
+
 URemoteControlSignatureRegistry* SRCSignatureTree::GetSignatureRegistry() const
 {
 	TSharedPtr<SRCSignaturePanel> SignaturePanel = SignaturePanelWeak.Pin();
@@ -76,18 +83,13 @@ TSharedPtr<IRCSignatureColumn> SRCSignatureTree::FindColumn(FName InColumnName)
 
 void SRCSignatureTree::Refresh()
 {
-	if (bRefreshing)
-	{
-		return;
-	}
-
-	TGuardValue<bool> RefreshGuard(bRefreshing, true);
+	TSharedRef<FRCSignatureTreeItemSelection> Selection = RootItem->GetSelection();
 
 	RootItem->RebuildChildren();
-	RootItem->VisitChildren([&TreeView = SignatureTreeView](const TSharedPtr<FRCSignatureTreeItemBase>& InItem)->bool
+	RootItem->VisitChildren([&TreeView = SignatureTreeView, &Selection](const TSharedPtr<FRCSignatureTreeItemBase>& InItem)->bool
 		{
-			TreeView->SetItemExpansion(InItem, InItem->HasAnyFlags(ERCSignatureTreeItemFlags::Expanded));
-			TreeView->SetItemSelection(InItem, InItem->HasAnyFlags(ERCSignatureTreeItemFlags::Selected));
+			TreeView->SetItemExpansion(InItem, InItem->HasAnyTreeViewFlags(ERCSignatureTreeItemViewFlags::Expanded));
+			TreeView->SetItemSelection(InItem, Selection->IsSelected(InItem));
 			return true;
 		}
 		, /*bRecursive*/true);
@@ -97,7 +99,7 @@ void SRCSignatureTree::Refresh()
 
 TArray<TSharedPtr<FRCSignatureTreeItemBase>> SRCSignatureTree::GetSelectedItems() const
 {
-	return SignatureTreeView->GetSelectedItems();
+	return RootItem->GetSelection()->GetSelectedItems();
 }
 
 URemoteControlPreset* SRCSignatureTree::GetPreset()
@@ -145,14 +147,14 @@ int32 SRCSignatureTree::RemoveModel(const TSharedPtr<FRCLogicModeBase> InItem)
 
 void SRCSignatureTree::DeleteSelectedPanelItems()
 {
-	const TArray<TSharedPtr<FRCSignatureTreeItemBase>> SelectedSignatures = SignatureTreeView->GetSelectedItems();
-	if (SelectedSignatures.IsEmpty())
+	const TArray<TSharedPtr<FRCSignatureTreeItemBase>> SelectedItems = RootItem->GetSelection()->GetSelectedItems();
+	if (SelectedItems.IsEmpty())
 	{
 		return;
 	}
 
-	FScopedTransaction Transaction(LOCTEXT("RemoveSelectedSignatures", "Remove Selected Signatures"));
-	DeleteItemsFromLogicPanel<FRCSignatureTreeItemBase>(RootItem->GetChildrenMutable(), SelectedSignatures);
+	FScopedTransaction Transaction(LOCTEXT("RemoveSelectedItems", "Remove Selected Items"));
+	DeleteItemsFromLogicPanel<FRCSignatureTreeItemBase>(RootItem->GetChildrenMutable(), SelectedItems);
 }
 
 void SRCSignatureTree::Reset()
@@ -184,46 +186,62 @@ void SRCSignatureTree::OnGetChildren(TSharedPtr<FRCSignatureTreeItemBase> InItem
 	if (InItem.IsValid())
 	{
 		OutChildren.Append(InItem->GetChildren());
+
+		// Remove all children that are not meant to be shown in the Tree View
+		OutChildren.RemoveAll([](const TSharedPtr<FRCSignatureTreeItemBase>& InItem)
+			{
+				return InItem->HasAnyTreeViewFlags(ERCSignatureTreeItemViewFlags::Hidden);
+			});
 	}
 }
 
 void SRCSignatureTree::OnItemExpansionChanged(TSharedPtr<FRCSignatureTreeItemBase> InItem, bool bInIsExpanded)
 {
-	if (bRefreshing)
-	{
-		return;
-	}
-
 	if (InItem.IsValid())
 	{
 		if (bInIsExpanded)
 		{
-			InItem->AddFlags(ERCSignatureTreeItemFlags::Expanded);
+			InItem->AddTreeViewFlags(ERCSignatureTreeItemViewFlags::Expanded);
 		}
 		else
 		{
-			InItem->RemoveFlags(ERCSignatureTreeItemFlags::Expanded);
+			InItem->RemoveTreeViewFlags(ERCSignatureTreeItemViewFlags::Expanded);
 		}
 	}
 }
 
+void SRCSignatureTree::UpdateTreeViewSelection()
+{
+	// Skip if already Syncing Selection to the Selection object
+	if (bSyncingSelection)
+	{
+		return;
+	}
+
+	TArray<TSharedPtr<FRCSignatureTreeItemBase>> Items = RootItem->GetSelection()->GetSelectedItems();
+	SignatureTreeView->ClearSelection();
+	SignatureTreeView->SetItemSelection(Items, true);
+}
+
 void SRCSignatureTree::OnItemSelectionChanged(TSharedPtr<FRCSignatureTreeItemBase> InItem, ESelectInfo::Type InSelectionType)
 {
-	if (bRefreshing)
+	// Skip if already syncing Selection or if the selection wasn't done by the user, as these direct selections go through the Selection object already
+	if (bSyncingSelection || InSelectionType == ESelectInfo::Direct)
 	{
 		return;
 	}
 
-	RootItem->VisitChildren([](const TSharedPtr<FRCSignatureTreeItemBase>& InItem)->bool
-		{
-			InItem->RemoveFlags(ERCSignatureTreeItemFlags::Selected);
-			return true;
-		}
-		, /*bRecursive*/true);
+	TGuardValue<bool> SyncSelectionGuard(bSyncingSelection, true);
+
+	TSharedRef<FRCSignatureTreeItemSelection> Selection = RootItem->GetSelection();
+
+	FRCSignatureTreeItemSelection::FSelectionScope SelectionScope = Selection->CreateSelectionScope();
+
+	Selection->ClearSelection();
 
 	for (const TSharedPtr<FRCSignatureTreeItemBase>& SelectedItem : SignatureTreeView->GetSelectedItems())
 	{
-		SelectedItem->AddFlags(ERCSignatureTreeItemFlags::Selected);
+		Selection->SetSelected(SelectedItem.ToSharedRef(), /*bSelected*/true, /*bMultiSelection*/true);
 	}
 }
 
diff --git a/Engine/Plugins/VirtualProduction/RemoteControl/Source/RemoteControlUI/Private/UI/Signature/SRCSignatureTree.h b/Engine/Plugins/VirtualProduction/RemoteControl/Source/RemoteControlUI/Private/UI/Signature/SRCSignatureTree.h
--- a/Engine/Plugins/VirtualProduction/RemoteControl/Source/RemoteControlUI/Private/UI/Signature/SRCSignatureTree.h
+++ b/Engine/Plugins/VirtualProduction/RemoteControl/Source/RemoteControlUI/Private/UI/Signature/SRCSignatureTree.h
@@ -21,6 +21,8 @@ public:
 
 	void Construct(const FArguments& InArgs, const TSharedRef<SRCSignaturePanel>& InSignaturePanel, const TSharedRef<SRemoteControlPanel>& InRCPanel);
 
+	TSharedRef<FRCSignatureTreeRootItem> GetRootItem() const;
+
 	URemoteControlSignatureRegistry* GetSignatureRegistry() const;
 
 	TSharedPtr<IRCSignatureColumn> FindColumn(FName InColumnName) const;
@@ -51,6 +53,8 @@ private:
 
 	void OnItemExpansionChanged(TSharedPtr<FRCSignatureTreeItemBase> InItem, bool bInIsExpanded);
 
+	void UpdateTreeViewSelection();
+
 	void OnItemSelectionChanged(TSharedPtr<FRCSignatureTreeItemBase> InItem, ESelectInfo::Type InSelectionType);
 
 	TSharedPtr<FRCSignatureTreeRootItem> RootItem;
@@ -63,5 +67,5 @@ private:
 
 	TWeakPtr<SRCSignaturePanel> SignaturePanelWeak;
 
-	bool bRefreshing = false;
+	bool bSyncingSelection = false;
 };
