Remote Control: Added Bind Controller Signature Action

diff --git a/Engine/Plugins/VirtualProduction/RemoteControl/Source/RemoteControlLogic/Private/Signature/RCSignatureControllerAction.cpp b/Engine/Plugins/VirtualProduction/RemoteControl/Source/RemoteControlLogic/Private/Signature/RCSignatureControllerAction.cpp
new file mode 100644
--- /dev/null
+++ b/Engine/Plugins/VirtualProduction/RemoteControl/Source/RemoteControlLogic/Private/Signature/RCSignatureControllerAction.cpp
@@ -0,0 +1,93 @@
+// Copyright Epic Games, Inc. All Rights Reserved.
+
+#include "RCSignatureControllerAction.h"
+#include "Behaviour/Builtin/Bind/RCBehaviourBind.h"
+#include "Controller/RCController.h"
+#include "Controller/RCCustomControllerUtilities.h"
+#include "RemoteControlPreset.h"
+#include "RemoteControlSignature.h"
+
+void FRCSignatureControllerAction::Initialize(const FRCSignatureField& InFieldOwner)
+{
+	ControllerName = InFieldOwner.FieldPath.GetFieldName();
+}
+
+bool FRCSignatureControllerAction::Execute(const FRCSignatureActionContext& InContext) const
+{
+	bool bControllerAdded;
+	URCController* Controller = FindOrAddController(InContext, &bControllerAdded);
+	if (!Controller)
+	{
+		return false;
+	}
+
+	bool bControllerModified = ExecuteControllerAction(InContext, Controller);
+
+	// if controller was added or modified, trigger refresh
+	if (bControllerAdded || bControllerModified)
+	{
+		InContext.Preset->OnVirtualPropertyContainerModified().Broadcast();
+		return true;
+	}
+
+	return false;
+}
+
+#if WITH_EDITOR
+FRCSignatureActionIcon FRCSignatureControllerAction::GetIcon() const
+{
+	FRCSignatureActionIcon Icon;
+	Icon.StyleSetName = TEXT("EditorStyle");
+	Icon.StyleName = TEXT("GraphEditor.StateMachine_16x");
+	return Icon;
+}
+#endif
+
+URCController* FRCSignatureControllerAction::FindOrAddController(const FRCSignatureActionContext& InContext, bool* bOutControllerAdded) const
+{
+	bool bControllerAdded = false;
+
+	FName ControllerNameToUse = ControllerName;
+
+	if (ControllerNameToUse.IsNone())
+	{
+		ControllerNameToUse = InContext.Property->FieldName;
+	}
+
+	URCController* Controller =  Cast<URCController>(InContext.Preset->GetController(ControllerNameToUse));
+
+	if (Controller && !IsControllerCompatible(InContext, Controller))
+	{
+		// Unset controller if it's not compatible
+		Controller = nullptr;
+	}
+
+	// Create a controller if it wasn't found or one was found but deemed incompatible
+	if (!Controller && bCreateControllerIfNotFound)
+	{
+		FPropertyBagPropertyDesc Desc;
+
+		if (MakeControllerDesc(InContext, Desc))
+		{
+			Controller = CastChecked<URCController>(InContext.Preset->AddController(URCController::StaticClass(), Desc.ValueType, const_cast<UObject*>(Desc.ValueTypeObject.Get()), ControllerNameToUse));
+			Controller->DisplayIndex = InContext.Preset->GetNumControllers() - 1;
+
+#if WITH_EDITOR
+			for (const FPropertyBagPropertyDescMetaData& MetaData : Desc.MetaData)
+			{
+				Controller->SetMetadataValue(MetaData.Key, MetaData.Value);
+			}
+#endif
+
+			bControllerAdded = true;
+			OnControllerCreated(InContext, Controller);
+		}
+	}
+
+	if (bOutControllerAdded)
+	{
+		*bOutControllerAdded = bControllerAdded;
+	}
+
+	return Controller;
+}
diff --git a/Engine/Plugins/VirtualProduction/RemoteControl/Source/RemoteControlLogic/Private/Signature/RCSignatureControllerAction.h b/Engine/Plugins/VirtualProduction/RemoteControl/Source/RemoteControlLogic/Private/Signature/RCSignatureControllerAction.h
new file mode 100644
--- /dev/null
+++ b/Engine/Plugins/VirtualProduction/RemoteControl/Source/RemoteControlLogic/Private/Signature/RCSignatureControllerAction.h
@@ -0,0 +1,73 @@
+// Copyright Epic Games, Inc. All Rights Reserved.
+
+#pragma once
+
+#include "RCSignatureAction.h"
+#include "RCSignatureControllerAction.generated.h"
+
+class URCController;
+struct FPropertyBagPropertyDesc;
+
+USTRUCT(meta=(Hidden))
+struct FRCSignatureControllerAction : public FRCSignatureAction
+{
+	GENERATED_BODY()
+
+	//~ Begin FRCSignatureAction
+	virtual void Initialize(const FRCSignatureField& InFieldOwner) override;
+	virtual bool Execute(const FRCSignatureActionContext& InContext) const override;
+#if WITH_EDITOR
+	virtual FRCSignatureActionIcon GetIcon() const override;
+#endif
+	//~ End FRCSignatureAction
+
+	/**
+	 * Determines whether a found Controller is compatible for re-use
+	 * @param InContext context of the action (preset, exposed property, etc)
+	 * @param InController the existing controller to check
+	 * @return true if compatible
+	 */
+	virtual bool IsControllerCompatible(const FRCSignatureActionContext& InContext, URCController* InController) const
+	{
+		return false;
+	}
+
+	/**
+	 * Retrieves the Property Desc to use for creating a new controller
+	 * @param InContext context of the action (preset, exposed property, etc)
+	 * @param OutDesc the property description to use for creating the controller
+	 * @return true if the property desc was filled in
+	 */
+	virtual bool MakeControllerDesc(const FRCSignatureActionContext& InContext, FPropertyBagPropertyDesc& OutDesc) const
+	{
+		return false;
+	}
+
+	/**
+	 * Called after a controller has been created
+	 * @param InContext context of the action (preset, exposed property, etc)
+	 * @param InController the controller created
+	 */
+	virtual void OnControllerCreated(const FRCSignatureActionContext& InContext, URCController* InController) const
+	{
+	}
+
+	/**
+	 * The main logic of the action
+	 * @param InContext context of the action (preset, exposed property, etc)
+	 * @param InController the controller to use for the action
+	 * @return true if the action modified the controller or some other object; false if nothing happened
+	 */
+	virtual bool ExecuteControllerAction(const FRCSignatureActionContext& InContext, URCController* InController) const
+	{
+		return false;
+	}
+
+	URCController* FindOrAddController(const FRCSignatureActionContext& InContext, bool* bOutControllerAdded = nullptr) const;
+
+	UPROPERTY(EditAnywhere, Category = "Controller")
+	FName ControllerName;
+
+	UPROPERTY(EditAnywhere, Category = "Controller")
+	bool bCreateControllerIfNotFound = true;
+};
diff --git a/Engine/Plugins/VirtualProduction/RemoteControl/Source/RemoteControlLogic/Private/Signature/RCSignatureControllerBindAction.cpp b/Engine/Plugins/VirtualProduction/RemoteControl/Source/RemoteControlLogic/Private/Signature/RCSignatureControllerBindAction.cpp
new file mode 100644
--- /dev/null
+++ b/Engine/Plugins/VirtualProduction/RemoteControl/Source/RemoteControlLogic/Private/Signature/RCSignatureControllerBindAction.cpp
@@ -0,0 +1,113 @@
+// Copyright Epic Games, Inc. All Rights Reserved.
+
+#include "RCSignatureControllerBindAction.h"
+#include "Action/Bind/RCPropertyBindAction.h"
+#include "Behaviour/Builtin/Bind/RCBehaviourBind.h"
+#include "Behaviour/Builtin/Bind/RCBehaviourBindNode.h"
+#include "Controller/RCController.h"
+#include "Controller/RCCustomControllerUtilities.h"
+#include "RemoteControlField.h"
+
+#if WITH_EDITOR
+#include "Engine/Texture2D.h"
+#endif
+
+bool FRCSignatureControllerBindAction::IsControllerCompatible(const FRCSignatureActionContext& InContext, URCController* InController) const
+{
+	return URCBehaviourBind::CanHaveActionForField(InController, InContext.Property.ToSharedRef(), /*bAllowNumericInputAsStrings*/true);
+}
+
+bool FRCSignatureControllerBindAction::MakeControllerDesc(const FRCSignatureActionContext& InContext, FPropertyBagPropertyDesc& OutDesc) const
+{
+	EPropertyBagPropertyType ValueType = EPropertyBagPropertyType::None;
+	UObject* ValueTypeObject = nullptr;
+
+	if (!URCBehaviourBind::GetPropertyBagTypeFromFieldProperty(InContext.Property->GetProperty(), ValueType, ValueTypeObject))
+	{
+		return false;
+	}
+
+#if WITH_EDITOR
+	FString CustomControllerName;
+
+	// Custom controller.This might require some interfacing for more customizations.
+	// This was taken from SRCControllerPanel::OnAddControllerClicked
+	if (ValueTypeObject == UTexture::StaticClass() || ValueTypeObject == UTexture2D::StaticClass())
+	{
+		if (ValueType == EPropertyBagPropertyType::String)
+		{
+			ValueTypeObject = nullptr;
+			CustomControllerName = UE::RCCustomControllers::CustomTextureControllerName;
+		}
+	}
+#endif
+
+	OutDesc.ValueType = ValueType;
+	OutDesc.ValueTypeObject = ValueTypeObject;
+
+#if WITH_EDITOR
+	// Add metadata to this controller, if this is a custom controller
+	if (!CustomControllerName.IsEmpty())
+	{
+		const TMap<FName, FString>& CustomControllerMetaData = UE::RCCustomControllers::GetCustomControllerMetaData(CustomControllerName);
+		OutDesc.MetaData.Reserve(CustomControllerMetaData.Num());
+
+		for (const TPair<FName, FString>& MetaData : CustomControllerMetaData)
+		{
+			OutDesc.MetaData.Emplace(MetaData.Key, MetaData.Value);
+		}
+	}
+#endif
+
+	return true;
+}
+
+void FRCSignatureControllerBindAction::OnControllerCreated(const FRCSignatureActionContext& InContext, URCController* InController) const
+{
+	// Transfer property value from Exposed Property to the New Controller.
+	// The goal is to keep values synced for a Controller newly created via "Auto Bind"
+	URCBehaviourBind::CopyPropertyValueToController(InController, InContext.Property.ToSharedRef());
+}
+
+bool FRCSignatureControllerBindAction::ExecuteControllerAction(const FRCSignatureActionContext& InContext, URCController* InController) const
+{
+	TSharedRef<FRemoteControlProperty> RCProperty = InContext.Property.ToSharedRef();
+
+	// Controller is already compatible with URCBehaviourBind::CanHaveActionForField when allowing string numeric input
+	// (see FRCSignatureControllerBindAction::IsControllerCompatible)
+	// Numeric conversion is required if a numeric action can't be added 
+	const bool bRequiresNumericConversion = !URCBehaviourBind::CanHaveActionForField(InController, RCProperty, /*bAllowNumericInputAsStrings*/false);
+
+	URCBehaviourBind* BindBehavior = nullptr;
+
+	for (URCBehaviour* Behavior : InController->Behaviours)
+	{
+		URCBehaviourBind* CurrentBindBehavior = Cast<URCBehaviourBind>(Behavior);
+
+		// In case numeric conversion is required we might have multiple Bind behaviours with different settings,
+		// so we do not break in case there is at least one Bind behavior with a matching clause.
+		// If not, we will create a new Bind behavior with the required setting (see below)
+		if (CurrentBindBehavior && (!bRequiresNumericConversion || CurrentBindBehavior->AreNumericInputsAllowedAsStrings()))
+		{
+			BindBehavior = CurrentBindBehavior;
+			break;
+		}
+	}
+
+	if (!BindBehavior)
+	{
+		BindBehavior = Cast<URCBehaviourBind>(InController->AddBehaviour(URCBehaviourBindNode::StaticClass()));
+
+		// If this is a new Bind Behavior and we are attempting to link unrelated (but compatible types), via numeric conversion, then we apply the bAllowNumericInputAsStrings flag
+		BindBehavior->SetAllowNumericInputAsStrings(bRequiresNumericConversion);
+	}
+
+	if (!ensure(BindBehavior))
+	{
+		return false;
+	}
+
+	URCAction* BindAction = BindBehavior->AddPropertyBindAction(RCProperty);
+	BindAction->Execute();
+	return true;
+}
diff --git a/Engine/Plugins/VirtualProduction/RemoteControl/Source/RemoteControlLogic/Private/Signature/RCSignatureControllerBindAction.h b/Engine/Plugins/VirtualProduction/RemoteControl/Source/RemoteControlLogic/Private/Signature/RCSignatureControllerBindAction.h
new file mode 100644
--- /dev/null
+++ b/Engine/Plugins/VirtualProduction/RemoteControl/Source/RemoteControlLogic/Private/Signature/RCSignatureControllerBindAction.h
@@ -0,0 +1,19 @@
+// Copyright Epic Games, Inc. All Rights Reserved.
+
+#pragma once
+
+#include "RCSignatureControllerAction.h"
+#include "RCSignatureControllerBindAction.generated.h"
+
+USTRUCT(DisplayName="Controller Bind Action")
+struct FRCSignatureControllerBindAction : public FRCSignatureControllerAction
+{
+	GENERATED_BODY()
+
+	//~ Begin FRCSignatureControllerAction
+	virtual bool IsControllerCompatible(const FRCSignatureActionContext& InContext, URCController* InController) const override;
+	virtual bool MakeControllerDesc(const FRCSignatureActionContext& InContext, FPropertyBagPropertyDesc& OutDesc) const override;
+	virtual void OnControllerCreated(const FRCSignatureActionContext& InContext, URCController* InController) const override;
+	virtual bool ExecuteControllerAction(const FRCSignatureActionContext& InContext, URCController* InController) const override;
+	//~ End FRCSignatureControllerAction
+};
