Remote Control: Rename RemoteControlSignatureRegistry to RCSignatureRegistry for readability (+ other file/classes named the same way).
Renamed ActionDefinition to ActionInstance (better naming to the struct's purpose)

diff --git a/Engine/Plugins/VirtualProduction/RemoteControl/Source/RemoteControl/Private/RCSignature.cpp b/Engine/Plugins/VirtualProduction/RemoteControl/Source/RemoteControl/Private/RCSignature.cpp
new file mode 100644
--- /dev/null
+++ b/Engine/Plugins/VirtualProduction/RemoteControl/Source/RemoteControl/Private/RCSignature.cpp
@@ -0,0 +1,164 @@
+// Copyright Epic Games, Inc. All Rights Reserved.
+
+#include "RCSignature.h"
+#include "GameFramework/Actor.h"
+#include "RCSignatureAction.h"
+#include "RCSignatureActionInstance.h"
+#include "RemoteControlField.h"
+#include "RemoteControlPreset.h"
+
+namespace UE::RemoteControl::Private
+{
+	TArray<UObject*> ResolveObjects(TConstArrayView<TWeakObjectPtr<UObject>> InObjects)
+	{
+		TArray<UObject*> ResolvedObjects;
+		ResolvedObjects.Reserve(InObjects.Num());
+		for (const TWeakObjectPtr<UObject>& ObjectWeak : InObjects)
+		{
+			if (UObject* Object = ObjectWeak.Get())
+			{
+				ResolvedObjects.Add(Object);
+			}
+		}
+		return ResolvedObjects;
+	}
+
+	UObject* FindContext(const FRCSignatureField& InField, UObject* InOuterObject, UClass* InSupportedClass)
+	{
+		UObject* Context;
+
+		if (InField.ObjectRelativePath.IsEmpty())
+		{
+			Context = InOuterObject;
+		}
+		else
+		{
+			UClass* FindClass = InSupportedClass ? InSupportedClass : UObject::StaticClass();
+			Context = StaticFindObject(FindClass, InOuterObject, *InField.ObjectRelativePath);
+
+			// Slow Path: if Subobject Path did not find the object, try to find the first sub-object of the object matching the class.
+			if (!Context && InSupportedClass)
+			{
+				ForEachObjectWithOuterBreakable(InOuterObject, [&Context, InSupportedClass](UObject* InSubobject)->bool
+					{
+						UClass* SubobjectClass = InSubobject->GetClass();
+						if (SubobjectClass && SubobjectClass->IsChildOf(InSupportedClass))
+						{
+							Context = InSubobject;
+							return false;
+						}
+						return true;
+					}
+					, /*bIncludeNestedObjects*/true);
+			}
+		}
+
+		return Context;
+	}
+}
+
+FRCSignatureField FRCSignatureField::CreateField(const FRCFieldPathInfo& InFieldPathInfo, const UObject* InOwnerObject, const FProperty* InProperty)
+{
+	const UClass* SupportedClass = nullptr;
+	if (InProperty)
+	{
+		SupportedClass = InProperty->GetOwnerClass();
+	}
+
+	if (!SupportedClass && InFieldPathInfo.GetSegmentCount() > 0 && InFieldPathInfo.GetFieldSegment(0).IsResolved())
+	{
+		SupportedClass = InFieldPathInfo.GetFieldSegment(0).ResolvedData.Field->GetOwnerClass();
+	}
+
+	FRCSignatureField Field;
+	Field.bEnabled = true;
+	Field.FieldPath = InFieldPathInfo;
+	Field.SupportedClass = SupportedClass;
+
+	if (InOwnerObject)
+	{
+		if (AActor* ActorOwner = InOwnerObject->GetTypedOuter<AActor>())
+		{
+			Field.ObjectRelativePath = InOwnerObject->GetPathName(ActorOwner);
+		}
+	}
+
+	return Field;
+}
+
+int32 FRCSignature::AddFields(TConstArrayView<FRCSignatureField> InFields)
+{
+	const int32 PreviousNum = Fields.Num();
+	Fields.Reserve(PreviousNum + InFields.Num());
+
+	for (const FRCSignatureField& Field : InFields)
+	{
+		Fields.AddUnique(Field);
+	}
+
+	return Fields.Num() - PreviousNum;
+}
+
+int32 FRCSignature::ApplySignature(URemoteControlPreset* InPreset, TConstArrayView<TWeakObjectPtr<UObject>> InObjects) const
+{
+	if (!InPreset || InObjects.IsEmpty())
+	{
+		return 0;
+	}
+
+	TArray<UObject*> ResolvedObjects = UE::RemoteControl::Private::ResolveObjects(InObjects);
+	if (ResolvedObjects.IsEmpty())
+	{
+		return false;
+	}
+
+	int32 ExposeCount = 0;
+
+	FRemoteControlPresetExposeArgs ExposeArgs;
+
+	for (const FRCSignatureField& Field : Fields)
+	{
+		if (!Field.bEnabled)
+		{
+			continue;
+		}
+
+		// Resolve, not load. It should be loaded already if relevant to the Objects to apply this signature to
+		UClass* SupportedClass = Field.SupportedClass.ResolveClass();
+
+		// Copy field path as it's going to be used to resolve (mutable)
+		FRCFieldPathInfo Path = Field.FieldPath;
+
+		for (UObject* Object : ResolvedObjects)
+		{
+			UObject* Context = UE::RemoteControl::Private::FindContext(Field, Object, SupportedClass);
+
+			// Attempt to resolve the path from the given context
+			if (!Context || !Path.Resolve(Context))
+			{
+				continue;
+			}
+
+			// If path resolved, expose the property
+			FRCSignatureActionContext ActionContext;
+			ActionContext.Preset = InPreset;
+			ActionContext.Object = Object;
+			ActionContext.Property = InPreset->ExposeProperty(Context, Field.FieldPath, ExposeArgs).Pin();
+			if (!ActionContext.Property.IsValid())
+			{
+				continue;
+			}
+
+			// Property was exposed successfully, increase expose count
+			++ExposeCount;
+
+			// Execute the actions for the newly exposed property
+			for (const FRCSignatureActionInstance& Action : Field.Actions)
+			{
+				Action.Execute(ActionContext);
+			}
+		}
+	}
+
+	return ExposeCount;
+}
diff --git a/Engine/Plugins/VirtualProduction/RemoteControl/Source/RemoteControl/Private/RCSignatureActionInstance.cpp b/Engine/Plugins/VirtualProduction/RemoteControl/Source/RemoteControl/Private/RCSignatureActionInstance.cpp
new file mode 100644
--- /dev/null
+++ b/Engine/Plugins/VirtualProduction/RemoteControl/Source/RemoteControl/Private/RCSignatureActionInstance.cpp
@@ -0,0 +1,34 @@
+// Copyright Epic Games, Inc. All Rights Reserved.
+
+#include "RCSignatureActionInstance.h"
+#include "RCSignatureAction.h"
+#include "UObject/StructOnScope.h"
+
+FRCSignatureActionInstance::FRCSignatureActionInstance(const UScriptStruct* InScriptStruct, const FRCSignatureField& InFieldOwner)
+{
+	ActionInstance.InitializeAsScriptStruct(InScriptStruct, /*StructMemory*/nullptr);
+
+	if (FRCSignatureAction* Action = ActionInstance.GetMutablePtr())
+	{
+		Action->Initialize(InFieldOwner);
+	}
+}
+
+const FRCSignatureAction* FRCSignatureActionInstance::GetAction() const
+{
+	return ActionInstance.GetPtr();
+}
+
+TSharedRef<FStructOnScope> FRCSignatureActionInstance::MakeStructOnScope()
+{
+	return MakeShared<FStructOnScope>(ActionInstance.GetScriptStruct(), ActionInstance.GetMutableMemory());
+}
+
+bool FRCSignatureActionInstance::Execute(const FRCSignatureActionContext& InContext) const
+{
+	if (const FRCSignatureAction* Action = ActionInstance.GetPtr())
+	{
+		return Action->Execute(InContext);
+	}
+	return false;
+}
diff --git a/Engine/Plugins/VirtualProduction/RemoteControl/Source/RemoteControl/Private/RCSignatureRegistry.cpp b/Engine/Plugins/VirtualProduction/RemoteControl/Source/RemoteControl/Private/RCSignatureRegistry.cpp
new file mode 100644
--- /dev/null
+++ b/Engine/Plugins/VirtualProduction/RemoteControl/Source/RemoteControl/Private/RCSignatureRegistry.cpp
@@ -0,0 +1,38 @@
+// Copyright Epic Games, Inc. All Rights Reserved.
+
+#include "RCSignatureRegistry.h"
+
+#define LOCTEXT_NAMESPACE "RCSignatureRegistry"
+
+const FRCSignature* URCSignatureRegistry::FindSignature(const FGuid& InSignatureId) const
+{
+	return Signatures.FindByKey(InSignatureId);
+}
+
+FRCSignature* URCSignatureRegistry::FindSignatureMutable(const FGuid& InSignatureId)
+{
+	return Signatures.FindByKey(InSignatureId);
+}
+
+FRCSignature& URCSignatureRegistry::AddSignature()
+{
+	FRCSignature& Signature = Signatures.AddDefaulted_GetRef();
+	Signature.DisplayName = LOCTEXT("NewSignatureDisplayName", "New Signature");
+	Signature.Id = FGuid::NewGuid();
+	return Signature;
+}
+
+int32 URCSignatureRegistry::RemoveSignature(const FGuid& InSignatureId)
+{
+	return Signatures.RemoveAll([&InSignatureId](const FRCSignature& InSignature)
+		{
+			return InSignature.Id == InSignatureId;
+		});
+}
+
+void URCSignatureRegistry::EmptySignatures()
+{
+	Signatures.Empty();
+}
+
+#undef LOCTEXT_NAMESPACE
diff --git a/Engine/Plugins/VirtualProduction/RemoteControl/Source/RemoteControl/Private/RemoteControlPreset.cpp b/Engine/Plugins/VirtualProduction/RemoteControl/Source/RemoteControl/Private/RemoteControlPreset.cpp
--- a/Engine/Plugins/VirtualProduction/RemoteControl/Source/RemoteControl/Private/RemoteControlPreset.cpp
+++ b/Engine/Plugins/VirtualProduction/RemoteControl/Source/RemoteControl/Private/RemoteControlPreset.cpp
@@ -14,21 +14,21 @@
 #include "GameFramework/Actor.h"
 #include "HAL/IConsoleManager.h"
 #include "IRemoteControlModule.h"
-#include "RCVirtualPropertyContainer.h"
-#include "RCVirtualProperty.h"
 #include "Misc/CoreDelegates.h"
-#include "Misc/TransactionObjectEvent.h"
 #include "Misc/Optional.h"
+#include "Misc/TransactionObjectEvent.h"
+#include "RCSignatureRegistry.h"
+#include "RCVirtualProperty.h"
+#include "RCVirtualPropertyContainer.h"
 #include "RemoteControlActor.h"
 #include "RemoteControlBinding.h"
 #include "RemoteControlEntityFactory.h"
 #include "RemoteControlExposeRegistry.h"
 #include "RemoteControlFieldPath.h"
-#include "RemoteControlPropertyIdRegistry.h"
 #include "RemoteControlLogger.h"
 #include "RemoteControlObjectVersion.h"
 #include "RemoteControlPresetRebindingManager.h"
-#include "RemoteControlSignatureRegistry.h"
+#include "RemoteControlPropertyIdRegistry.h"
 #include "RemoteControlTransactionListenerHelper.h"
 
 #include "UObject/Object.h"
@@ -674,7 +674,7 @@ URemoteControlPreset::URemoteControlPreset()
 {
 	Registry = CreateDefaultSubobject<URemoteControlExposeRegistry>(FName("ExposeRegistry"));
 
-	SignatureRegistry = CreateDefaultSubobject<URemoteControlSignatureRegistry>(TEXT("SignatureRegistry"));
+	SignatureRegistry = CreateDefaultSubobject<URCSignatureRegistry>(TEXT("SignatureRegistry"));
 
 	PropertyIdRegistry = CreateDefaultSubobject<URemoteControlPropertyIdRegistry>(FName("PropertyIdRegistry"));
 	PropertyIdRegistry->Initialize();
diff --git a/Engine/Plugins/VirtualProduction/RemoteControl/Source/RemoteControl/Private/RemoteControlSignature.cpp b/Engine/Plugins/VirtualProduction/RemoteControl/Source/RemoteControl/Private/RemoteControlSignature.cpp
deleted file mode 100644
diff --git a/Engine/Plugins/VirtualProduction/RemoteControl/Source/RemoteControl/Private/RemoteControlSignatureActionDefinition.cpp b/Engine/Plugins/VirtualProduction/RemoteControl/Source/RemoteControl/Private/RemoteControlSignatureActionDefinition.cpp
deleted file mode 100644
diff --git a/Engine/Plugins/VirtualProduction/RemoteControl/Source/RemoteControl/Private/RemoteControlSignatureRegistry.cpp b/Engine/Plugins/VirtualProduction/RemoteControl/Source/RemoteControl/Private/RemoteControlSignatureRegistry.cpp
deleted file mode 100644
diff --git a/Engine/Plugins/VirtualProduction/RemoteControl/Source/RemoteControl/Public/RCSignature.h b/Engine/Plugins/VirtualProduction/RemoteControl/Source/RemoteControl/Public/RCSignature.h
new file mode 100644
--- /dev/null
+++ b/Engine/Plugins/VirtualProduction/RemoteControl/Source/RemoteControl/Public/RCSignature.h
@@ -0,0 +1,89 @@
+// Copyright Epic Games, Inc. All Rights Reserved.
+
+#pragma once
+
+#include "RCSignatureActionInstance.h"
+#include "RemoteControlFieldPath.h"
+#include "RCSignature.generated.h"
+
+enum class EPropertyBagPropertyType : uint8;
+class URemoteControlPreset;
+
+/** Representation class of a Field (e.g. Property) in an Object */
+USTRUCT()
+struct FRCSignatureField
+{
+	GENERATED_BODY()
+
+	FRCSignatureField() = default;
+
+	REMOTECONTROL_API static FRCSignatureField CreateField(const FRCFieldPathInfo& InFieldPathInfo, const UObject* InOwnerObject, const FProperty* InProperty);
+
+	bool operator==(const FRCSignatureField& InOtherField) const
+	{
+		return FieldPath == InOtherField.FieldPath
+			&& ObjectRelativePath == InOtherField.ObjectRelativePath
+			&& SupportedClass == InOtherField.SupportedClass;
+	}
+
+	/** Path Info for this Field */
+	UPROPERTY()
+	FRCFieldPathInfo FieldPath;
+
+	/** Optional: Relative path from an owner (e.g. Actor) to the object owning the property (e.g. an Actor Component) */
+	UPROPERTY()
+	FString ObjectRelativePath;
+
+	/** Object class holding the property */
+	UPROPERTY()
+	FSoftClassPath SupportedClass;
+
+	/** Container holding the action instances for the field */
+	UPROPERTY()
+	TArray<FRCSignatureActionInstance> Actions;
+
+	/** Whether to consider this field when applying a Signature */
+	UPROPERTY()
+	bool bEnabled = true;
+};
+
+USTRUCT()
+struct FRCSignature
+{
+	GENERATED_BODY()
+
+	bool operator==(const FGuid& InSignatureId) const
+	{
+		return Id == InSignatureId;
+	}
+
+	/**
+	 * Adds the given fields to this signature and ensuring no repeated field is present
+	 * @return the number of new fields added
+	 */
+	REMOTECONTROL_API int32 AddFields(TConstArrayView<FRCSignatureField> InFields);
+
+	/**
+	 * Applies this Signature to the given Objects by exposing all this Signature's fields to the given preset
+	 * @param InPreset the preset where these properties will be exposed to
+	 * @param InObjects the objects whose properties or properties of its subojects to expose
+	 * @return the number of properties affected in total
+	 */
+	REMOTECONTROL_API int32 ApplySignature(URemoteControlPreset* InPreset, TConstArrayView<TWeakObjectPtr<UObject>> InObjects) const;
+
+	/** User facing friendly name. Used as the Label when exposing */
+	UPROPERTY()
+	FText DisplayName;
+
+	/** Unique Id identifying this Signature */
+	UPROPERTY(meta=(IgnoreForMemberInitializationTest))
+	FGuid Id;
+
+	/** The fields owned by this Signature */
+	UPROPERTY()
+	TArray<FRCSignatureField> Fields;
+
+	/** Whether this Signature can be applied */
+	UPROPERTY()
+	bool bEnabled = true;
+};
\ No newline at end of file
diff --git a/Engine/Plugins/VirtualProduction/RemoteControl/Source/RemoteControl/Public/RCSignatureActionDefinition.h b/Engine/Plugins/VirtualProduction/RemoteControl/Source/RemoteControl/Public/RCSignatureActionDefinition.h
deleted file mode 100644
diff --git a/Engine/Plugins/VirtualProduction/RemoteControl/Source/RemoteControl/Public/RCSignatureActionInstance.h b/Engine/Plugins/VirtualProduction/RemoteControl/Source/RemoteControl/Public/RCSignatureActionInstance.h
new file mode 100644
--- /dev/null
+++ b/Engine/Plugins/VirtualProduction/RemoteControl/Source/RemoteControl/Public/RCSignatureActionInstance.h
@@ -0,0 +1,32 @@
+// Copyright Epic Games, Inc. All Rights Reserved.
+
+#pragma once
+
+#include "StructUtils/InstancedStruct.h"
+#include "RCSignatureActionInstance.generated.h"
+
+class FStructOnScope;
+struct FRCSignatureAction;
+struct FRCSignatureActionContext;
+struct FRCSignatureField;
+
+/** Struct containing an Action Instance and handling its Execution */
+USTRUCT()
+struct FRCSignatureActionInstance
+{
+	GENERATED_BODY()
+
+	FRCSignatureActionInstance() = default;
+
+	REMOTECONTROL_API explicit FRCSignatureActionInstance(const UScriptStruct* InScriptStruct, const FRCSignatureField& InFieldOwner);
+
+	REMOTECONTROL_API const FRCSignatureAction* GetAction() const;
+
+	REMOTECONTROL_API TSharedRef<FStructOnScope> MakeStructOnScope();
+
+	bool Execute(const FRCSignatureActionContext& InContext) const;
+
+private:
+	UPROPERTY()
+	TInstancedStruct<FRCSignatureAction> ActionInstance;
+};
diff --git a/Engine/Plugins/VirtualProduction/RemoteControl/Source/RemoteControl/Public/RCSignatureRegistry.h b/Engine/Plugins/VirtualProduction/RemoteControl/Source/RemoteControl/Public/RCSignatureRegistry.h
new file mode 100644
--- /dev/null
+++ b/Engine/Plugins/VirtualProduction/RemoteControl/Source/RemoteControl/Public/RCSignatureRegistry.h
@@ -0,0 +1,51 @@
+// Copyright Epic Games, Inc. All Rights Reserved.
+
+#pragma once
+
+#include "RCSignature.h"
+#include "UObject/Object.h"
+#include "RCSignatureRegistry.generated.h"
+
+UCLASS(MinimalAPI)
+class URCSignatureRegistry : public UObject
+{
+	GENERATED_BODY()
+
+public:
+	TConstArrayView<FRCSignature> GetSignatures() const
+	{
+		return Signatures;
+	}
+
+	/**
+	 * Find an existing Signature matching the given id
+	 * @param InSignatureId the id of the signature to find
+	 * @return const pointer to the signature if found, nullptr otherwise 
+	 */
+	REMOTECONTROL_API const FRCSignature* FindSignature(const FGuid& InSignatureId) const;
+
+	/**
+	 * Find an existing Signature matching the given id
+	 * @param InSignatureId the id of the signature to find
+	 * @return pointer to the signature if found, nullptr otherwise 
+	 */
+	REMOTECONTROL_API FRCSignature* FindSignatureMutable(const FGuid& InSignatureId);
+
+	/** Adds a new Signature to this registry */
+	REMOTECONTROL_API FRCSignature& AddSignature();
+
+	/**
+	 * Removes the Signature with the matching Signature Id
+	 * @param InSignatureId the id of the signature to remove
+	 * @return the number of signatures removed
+	 */
+	REMOTECONTROL_API int32 RemoveSignature(const FGuid& InSignatureId);
+
+	/** Removes all the Signature in this registry */
+	REMOTECONTROL_API void EmptySignatures();
+
+private:
+	/** Holds all the Signatures in this Registry */
+	UPROPERTY()
+	TArray<FRCSignature> Signatures;
+};
diff --git a/Engine/Plugins/VirtualProduction/RemoteControl/Source/RemoteControl/Public/RemoteControlPreset.h b/Engine/Plugins/VirtualProduction/RemoteControl/Source/RemoteControl/Public/RemoteControlPreset.h
--- a/Engine/Plugins/VirtualProduction/RemoteControl/Source/RemoteControl/Public/RemoteControlPreset.h
+++ b/Engine/Plugins/VirtualProduction/RemoteControl/Source/RemoteControl/Public/RemoteControlPreset.h
@@ -29,7 +29,7 @@ class URCVirtualPropertyContainerBase;
 class URCVirtualPropertyInContainer;
 class URCVirtualPropertySelfContainer;
 class URemoteControlExposeRegistry;
-class URemoteControlSignatureRegistry;
+class URCSignatureRegistry;
 class URemoteControlBinding;
 class URemoteControlPreset;
 
@@ -707,7 +707,7 @@ public:
 	FName GenerateUniqueLabel(const FName InDesiredName) const;
 
 	/** Returns the preset's signature registry */
-	URemoteControlSignatureRegistry* GetSignatureRegistry() const
+	URCSignatureRegistry* GetSignatureRegistry() const
 	{
 		return SignatureRegistry;
 	}
@@ -945,7 +945,7 @@ private:
 
 	/** Holds all signatures on the preset. */
 	UPROPERTY(Instanced)
-	TObjectPtr<URemoteControlSignatureRegistry> SignatureRegistry = nullptr;
+	TObjectPtr<URCSignatureRegistry> SignatureRegistry = nullptr;
 
 	/** Holds identities of exposed entities on the preset. */
 	UPROPERTY(Transient)
diff --git a/Engine/Plugins/VirtualProduction/RemoteControl/Source/RemoteControl/Public/RemoteControlSignature.h b/Engine/Plugins/VirtualProduction/RemoteControl/Source/RemoteControl/Public/RemoteControlSignature.h
deleted file mode 100644
diff --git a/Engine/Plugins/VirtualProduction/RemoteControl/Source/RemoteControl/Public/RemoteControlSignatureRegistry.h b/Engine/Plugins/VirtualProduction/RemoteControl/Source/RemoteControl/Public/RemoteControlSignatureRegistry.h
deleted file mode 100644
diff --git a/Engine/Plugins/VirtualProduction/RemoteControl/Source/RemoteControlLogic/Private/Signature/RCSignatureControllerAction.cpp b/Engine/Plugins/VirtualProduction/RemoteControl/Source/RemoteControlLogic/Private/Signature/RCSignatureControllerAction.cpp
--- a/Engine/Plugins/VirtualProduction/RemoteControl/Source/RemoteControlLogic/Private/Signature/RCSignatureControllerAction.cpp
+++ b/Engine/Plugins/VirtualProduction/RemoteControl/Source/RemoteControlLogic/Private/Signature/RCSignatureControllerAction.cpp
@@ -4,8 +4,8 @@
 #include "Behaviour/Builtin/Bind/RCBehaviourBind.h"
 #include "Controller/RCController.h"
 #include "Controller/RCCustomControllerUtilities.h"
+#include "RCSignature.h"
 #include "RemoteControlPreset.h"
-#include "RemoteControlSignature.h"
 
 void FRCSignatureControllerAction::Initialize(const FRCSignatureField& InFieldOwner)
 {
diff --git a/Engine/Plugins/VirtualProduction/RemoteControl/Source/RemoteControlUI/Private/UI/Signature/Actions/SRCSignatureActionBox.cpp b/Engine/Plugins/VirtualProduction/RemoteControl/Source/RemoteControlUI/Private/UI/Signature/Actions/SRCSignatureActionBox.cpp
--- a/Engine/Plugins/VirtualProduction/RemoteControl/Source/RemoteControlUI/Private/UI/Signature/Actions/SRCSignatureActionBox.cpp
+++ b/Engine/Plugins/VirtualProduction/RemoteControl/Source/RemoteControlUI/Private/UI/Signature/Actions/SRCSignatureActionBox.cpp
@@ -2,8 +2,8 @@
 
 #include "SRCSignatureActionBox.h"
 #include "DetailLayoutBuilder.h"
+#include "RCSignature.h"
 #include "RCSignatureActionType.h"
-#include "RemoteControlSignature.h"
 #include "SRCSignatureAction.h"
 #include "SRCSignatureActionIcon.h"
 #include "Styling/RemoteControlStyles.h"
diff --git a/Engine/Plugins/VirtualProduction/RemoteControl/Source/RemoteControlUI/Private/UI/Signature/Details/SRCSignatureDetails.cpp b/Engine/Plugins/VirtualProduction/RemoteControl/Source/RemoteControlUI/Private/UI/Signature/Details/SRCSignatureDetails.cpp
--- a/Engine/Plugins/VirtualProduction/RemoteControl/Source/RemoteControlUI/Private/UI/Signature/Details/SRCSignatureDetails.cpp
+++ b/Engine/Plugins/VirtualProduction/RemoteControl/Source/RemoteControlUI/Private/UI/Signature/Details/SRCSignatureDetails.cpp
@@ -6,14 +6,14 @@
 #include "IStructureDetailsView.h"
 #include "Modules/ModuleManager.h"
 #include "PropertyEditorModule.h"
-#include "RemoteControlSignatureRegistry.h"
+#include "RCSignatureRegistry.h"
 #include "ScopedTransaction.h"
 #include "UI/Signature/Items/RCSignatureTreeRootItem.h"
 #include "UI/Signature/RCSignatureTreeItemSelection.h"
 
 #define LOCTEXT_NAMESPACE "SRCSignatureDetails"
 
-void SRCSignatureDetails::Construct(const FArguments& InArgs, URemoteControlSignatureRegistry* InSignatureRegistry, const TSharedRef<FRCSignatureTreeItemSelection>& InSelection)
+void SRCSignatureDetails::Construct(const FArguments& InArgs, URCSignatureRegistry* InSignatureRegistry, const TSharedRef<FRCSignatureTreeItemSelection>& InSelection)
 {
 	SignatureRegistryWeak = InSignatureRegistry;
 	SelectionWeak = InSelection;
@@ -92,7 +92,7 @@ void SRCSignatureDetails::OnFinishedChangingProperties(const FPropertyChangedEve
 
 void SRCSignatureDetails::NotifyPreChange(FEditPropertyChain* InPropertyAboutToChange)
 {
-	if (URemoteControlSignatureRegistry* SignatureRegistry = SignatureRegistryWeak.Get())
+	if (URCSignatureRegistry* SignatureRegistry = SignatureRegistryWeak.Get())
 	{
 		// Begin Transaction
 		ensureAlways(!CurrentTransaction.IsValid());
diff --git a/Engine/Plugins/VirtualProduction/RemoteControl/Source/RemoteControlUI/Private/UI/Signature/Details/SRCSignatureDetails.h b/Engine/Plugins/VirtualProduction/RemoteControl/Source/RemoteControlUI/Private/UI/Signature/Details/SRCSignatureDetails.h
--- a/Engine/Plugins/VirtualProduction/RemoteControl/Source/RemoteControlUI/Private/UI/Signature/Details/SRCSignatureDetails.h
+++ b/Engine/Plugins/VirtualProduction/RemoteControl/Source/RemoteControlUI/Private/UI/Signature/Details/SRCSignatureDetails.h
@@ -11,7 +11,7 @@ class FRCSignatureTreeItemSelection;
 class FScopedTransaction;
 class FStructOnScope;
 class IStructureDetailsView;
-class URemoteControlSignatureRegistry;
+class URCSignatureRegistry;
 struct FPropertyChangedEvent;
 
 class SRCSignatureDetails : public SCompoundWidget, public FNotifyHook
@@ -20,7 +20,7 @@ public:
 	SLATE_BEGIN_ARGS(SRCSignatureDetails) {}
 	SLATE_END_ARGS()
 
-	void Construct(const FArguments& InArgs, URemoteControlSignatureRegistry* InSignatureRegistry, const TSharedRef<FRCSignatureTreeItemSelection>& InSelection);
+	void Construct(const FArguments& InArgs, URCSignatureRegistry* InSignatureRegistry, const TSharedRef<FRCSignatureTreeItemSelection>& InSelection);
 
 	virtual ~SRCSignatureDetails() override;
 
@@ -39,7 +39,7 @@ private:
 
 	TSharedPtr<FScopedTransaction> CurrentTransaction;
 
-	TWeakObjectPtr<URemoteControlSignatureRegistry> SignatureRegistryWeak;
+	TWeakObjectPtr<URCSignatureRegistry> SignatureRegistryWeak;
 
 	TWeakPtr<FRCSignatureTreeItemSelection> SelectionWeak;
 };
diff --git a/Engine/Plugins/VirtualProduction/RemoteControl/Source/RemoteControlUI/Private/UI/Signature/Items/RCSignatureTreeActionItem.cpp b/Engine/Plugins/VirtualProduction/RemoteControl/Source/RemoteControlUI/Private/UI/Signature/Items/RCSignatureTreeActionItem.cpp
--- a/Engine/Plugins/VirtualProduction/RemoteControl/Source/RemoteControlUI/Private/UI/Signature/Items/RCSignatureTreeActionItem.cpp
+++ b/Engine/Plugins/VirtualProduction/RemoteControl/Source/RemoteControlUI/Private/UI/Signature/Items/RCSignatureTreeActionItem.cpp
@@ -1,11 +1,11 @@
 // Copyright Epic Games, Inc. All Rights Reserved.
 
 #include "RCSignatureTreeActionItem.h"
+#include "RCSignature.h"
 #include "RCSignatureAction.h"
+#include "RCSignatureRegistry.h"
 #include "RCSignatureTreeFieldItem.h"
 #include "RCSignatureTreeRootItem.h"
-#include "RemoteControlSignature.h"
-#include "RemoteControlSignatureRegistry.h"
 #include "ScopedTransaction.h"
 #include "UI/Signature/RCSignatureTreeItemSelection.h"
 
@@ -19,7 +19,7 @@ FRCSignatureTreeActionItem::FRCSignatureTreeActionItem(int32 InActionIndex, cons
 	AddTreeViewFlags(ERCSignatureTreeItemViewFlags::Hidden);
 }
 
-const FRCSignatureActionDefinition* FRCSignatureTreeActionItem::FindActionDefinition() const
+const FRCSignatureActionInstance* FRCSignatureTreeActionItem::FindActionInstance() const
 {
 	TSharedPtr<FRCSignatureTreeFieldItem> FieldItem = GetParentFieldItem();
 	if (!FieldItem.IsValid())
@@ -28,30 +28,30 @@ const FRCSignatureActionDefinition* FRCSignatureTreeActionItem::FindActionDefini
 	}
 
 	const FRCSignatureField* Field = FieldItem->FindField();
-	if (!Field || !Field->ActionDefinitions.IsValidIndex(ActionIndex))
+	if (!Field || !Field->Actions.IsValidIndex(ActionIndex))
 	{
 		return nullptr;
 	}
 
-	return &Field->ActionDefinitions[ActionIndex];
+	return &Field->Actions[ActionIndex];
 }
 
-FRCSignatureActionDefinition* FRCSignatureTreeActionItem::FindActionDefinitionMutable()
+FRCSignatureActionInstance* FRCSignatureTreeActionItem::FindActionInstanceMutable()
 {
 	FRCSignatureField* Field = FindParentFieldMutable();
-	if (!Field || !Field->ActionDefinitions.IsValidIndex(ActionIndex))
+	if (!Field || !Field->Actions.IsValidIndex(ActionIndex))
 	{
 		return nullptr;
 	}
 
-	return &Field->ActionDefinitions[ActionIndex];
+	return &Field->Actions[ActionIndex];
 }
 
 FRCSignatureActionIcon FRCSignatureTreeActionItem::GetIcon() const
 {
-	if (const FRCSignatureActionDefinition* ActionDefinition = FindActionDefinition())
+	if (const FRCSignatureActionInstance* ActionInstance = FindActionInstance())
 	{
-		if (const FRCSignatureAction* Action = ActionDefinition->GetAction())
+		if (const FRCSignatureAction* Action = ActionInstance->GetAction())
 		{
 			return Action->GetIcon();
 		}
@@ -71,10 +71,10 @@ TOptional<bool> FRCSignatureTreeActionItem::IsEnabled() const
 
 int32 FRCSignatureTreeActionItem::RemoveFromRegistry()
 {
-	URemoteControlSignatureRegistry* Registry;
+	URCSignatureRegistry* Registry;
 
 	FRCSignatureField* Field = FindParentFieldMutable(&Registry);
-	if (!Field || !Field->ActionDefinitions.IsValidIndex(ActionIndex))
+	if (!Field || !Field->Actions.IsValidIndex(ActionIndex))
 	{
 		return 0;
 	}
@@ -83,7 +83,7 @@ int32 FRCSignatureTreeActionItem::RemoveFromRegistry()
 
 	FScopedTransaction Transaction(LOCTEXT("RemoveAction", "Remove Action"));
 	Registry->Modify();
-	Field->ActionDefinitions.RemoveAt(ActionIndex);
+	Field->Actions.RemoveAt(ActionIndex);
 	return 1;
 }
 
@@ -101,9 +101,9 @@ FText FRCSignatureTreeActionItem::GetDescription() const
 
 TSharedPtr<FStructOnScope> FRCSignatureTreeActionItem::MakeSelectionStruct()
 {
-	if (FRCSignatureActionDefinition* ActionDefinition = FindActionDefinitionMutable())
+	if (FRCSignatureActionInstance* ActionInstance = FindActionInstanceMutable())
 	{
-		return ActionDefinition->MakeStructOnScope();
+		return ActionInstance->MakeStructOnScope();
 	}
 	return nullptr;
 }
@@ -117,7 +117,7 @@ TSharedPtr<FRCSignatureTreeFieldItem> FRCSignatureTreeActionItem::GetParentField
 	return nullptr;
 }
 
-FRCSignatureField* FRCSignatureTreeActionItem::FindParentFieldMutable(URemoteControlSignatureRegistry** OutRegistry)
+FRCSignatureField* FRCSignatureTreeActionItem::FindParentFieldMutable(URCSignatureRegistry** OutRegistry)
 {
 	if (TSharedPtr<FRCSignatureTreeFieldItem> ParentFieldItem = GetParentFieldItem())
 	{
diff --git a/Engine/Plugins/VirtualProduction/RemoteControl/Source/RemoteControlUI/Private/UI/Signature/Items/RCSignatureTreeActionItem.h b/Engine/Plugins/VirtualProduction/RemoteControl/Source/RemoteControlUI/Private/UI/Signature/Items/RCSignatureTreeActionItem.h
--- a/Engine/Plugins/VirtualProduction/RemoteControl/Source/RemoteControlUI/Private/UI/Signature/Items/RCSignatureTreeActionItem.h
+++ b/Engine/Plugins/VirtualProduction/RemoteControl/Source/RemoteControlUI/Private/UI/Signature/Items/RCSignatureTreeActionItem.h
@@ -5,8 +5,8 @@
 #include "RCSignatureTreeItemBase.h"
 
 class FRCSignatureTreeFieldItem;
-class URemoteControlSignatureRegistry;
-struct FRCSignatureActionDefinition;
+class URCSignatureRegistry;
+struct FRCSignatureActionInstance;
 struct FRCSignatureActionIcon;
 struct FRCSignatureField;
 
@@ -18,9 +18,9 @@ public:
 
 	explicit FRCSignatureTreeActionItem(int32 InActionIndex, const TSharedPtr<SRCSignatureTree>& InSignatureTree);
 
-	const FRCSignatureActionDefinition* FindActionDefinition() const;
+	const FRCSignatureActionInstance* FindActionInstance() const;
 
-	FRCSignatureActionDefinition* FindActionDefinitionMutable();
+	FRCSignatureActionInstance* FindActionInstanceMutable();
 
 	FRCSignatureActionIcon GetIcon() const;
 
@@ -38,7 +38,7 @@ protected:
 private:
 	TSharedPtr<FRCSignatureTreeFieldItem> GetParentFieldItem() const;
 
-	FRCSignatureField* FindParentFieldMutable(URemoteControlSignatureRegistry** OutRegistry = nullptr);
+	FRCSignatureField* FindParentFieldMutable(URCSignatureRegistry** OutRegistry = nullptr);
 
 	int32 ActionIndex;
 };
diff --git a/Engine/Plugins/VirtualProduction/RemoteControl/Source/RemoteControlUI/Private/UI/Signature/Items/RCSignatureTreeFieldItem.cpp b/Engine/Plugins/VirtualProduction/RemoteControl/Source/RemoteControlUI/Private/UI/Signature/Items/RCSignatureTreeFieldItem.cpp
--- a/Engine/Plugins/VirtualProduction/RemoteControl/Source/RemoteControlUI/Private/UI/Signature/Items/RCSignatureTreeFieldItem.cpp
+++ b/Engine/Plugins/VirtualProduction/RemoteControl/Source/RemoteControlUI/Private/UI/Signature/Items/RCSignatureTreeFieldItem.cpp
@@ -1,10 +1,10 @@
 // Copyright Epic Games, Inc. All Rights Reserved.
 
 #include "RCSignatureTreeFieldItem.h"
+#include "RCSignature.h"
+#include "RCSignatureRegistry.h"
 #include "RCSignatureTreeActionItem.h"
 #include "RCSignatureTreeSignatureItem.h"
-#include "RemoteControlSignature.h"
-#include "RemoteControlSignatureRegistry.h"
 #include "ScopedTransaction.h"
 
 #define LOCTEXT_NAMESPACE "RCSignatureTreeFieldItem"
@@ -32,7 +32,7 @@ const FRCSignatureField* FRCSignatureTreeFieldItem::FindField() const
 	return &Signature->Fields[FieldIndex];
 }
 
-FRCSignatureField* FRCSignatureTreeFieldItem::FindFieldMutable(URemoteControlSignatureRegistry** OutRegistry)
+FRCSignatureField* FRCSignatureTreeFieldItem::FindFieldMutable(URCSignatureRegistry** OutRegistry)
 {
 	FRCSignature* Signature = FindParentSignature(OutRegistry);
 	if (!Signature || !Signature->Fields.IsValidIndex(FieldIndex))
@@ -44,7 +44,7 @@ FRCSignatureField* FRCSignatureTreeFieldItem::FindFieldMutable(URemoteControlSig
 
 void FRCSignatureTreeFieldItem::AddAction(const UScriptStruct* InActionType)
 {
-	URemoteControlSignatureRegistry* Registry;
+	URCSignatureRegistry* Registry;
 
 	FRCSignatureField* Field = FindFieldMutable(&Registry);
 	if (!Field)
@@ -55,7 +55,7 @@ void FRCSignatureTreeFieldItem::AddAction(const UScriptStruct* InActionType)
 	FScopedTransaction Transaction(LOCTEXT("AddSignatureAction", "Add Signature Action"));
 	check(Registry);
 	Registry->Modify();
-	Field->ActionDefinitions.Emplace(InActionType, *Field);
+	Field->Actions.Emplace(InActionType, *Field);
 }
 
 void FRCSignatureTreeFieldItem::BuildPathSegment(FStringBuilderBase& InBuilder) const
@@ -74,7 +74,7 @@ TOptional<bool> FRCSignatureTreeFieldItem::IsEnabled() const
 
 void FRCSignatureTreeFieldItem::SetEnabled(bool bInEnabled)
 {
-	URemoteControlSignatureRegistry* Registry;
+	URCSignatureRegistry* Registry;
 
 	FRCSignatureField* Field = FindFieldMutable(&Registry);
 	if (!Field || Field->bEnabled == bInEnabled)
@@ -118,7 +118,7 @@ FText FRCSignatureTreeFieldItem::GetDescription() const
 
 int32 FRCSignatureTreeFieldItem::RemoveFromRegistry()
 {
-	URemoteControlSignatureRegistry* Registry;
+	URCSignatureRegistry* Registry;
 
 	FRCSignature* Signature = FindParentSignature(&Registry);
 	if (!Signature || !Signature->Fields.IsValidIndex(FieldIndex))
@@ -144,8 +144,8 @@ void FRCSignatureTreeFieldItem::GenerateChildren(TArray<TSharedPtr<FRCSignatureT
 
 	const TSharedPtr<SRCSignatureTree> SignatureTree = GetSignatureTree();
 
-	OutChildren.Reserve(OutChildren.Num() + Field->ActionDefinitions.Num());
-	for (int32 ActionIndex = 0; ActionIndex < Field->ActionDefinitions.Num(); ++ActionIndex)
+	OutChildren.Reserve(OutChildren.Num() + Field->Actions.Num());
+	for (int32 ActionIndex = 0; ActionIndex < Field->Actions.Num(); ++ActionIndex)
 	{
 		OutChildren.Add(MakeShared<FRCSignatureTreeActionItem>(ActionIndex, SignatureTree));
 	}
@@ -160,7 +160,7 @@ TSharedPtr<FRCSignatureTreeSignatureItem> FRCSignatureTreeFieldItem::GetParentSi
 	return nullptr;
 }
 
-FRCSignature* FRCSignatureTreeFieldItem::FindParentSignature(URemoteControlSignatureRegistry** OutRegistry)
+FRCSignature* FRCSignatureTreeFieldItem::FindParentSignature(URCSignatureRegistry** OutRegistry)
 {
 	TSharedPtr<FRCSignatureTreeSignatureItem> ParentSignatureItem = GetParentSignatureItem();
 	if (!ParentSignatureItem.IsValid())
@@ -168,7 +168,7 @@ FRCSignature* FRCSignatureTreeFieldItem::FindParentSignature(URemoteControlSigna
 		return nullptr;
 	}
 
-	URemoteControlSignatureRegistry* Registry = ParentSignatureItem->GetRegistry();
+	URCSignatureRegistry* Registry = ParentSignatureItem->GetRegistry();
 
 	FRCSignature* Signature = ParentSignatureItem->FindSignatureMutable(Registry);
 	if (!Signature)
diff --git a/Engine/Plugins/VirtualProduction/RemoteControl/Source/RemoteControlUI/Private/UI/Signature/Items/RCSignatureTreeFieldItem.h b/Engine/Plugins/VirtualProduction/RemoteControl/Source/RemoteControlUI/Private/UI/Signature/Items/RCSignatureTreeFieldItem.h
--- a/Engine/Plugins/VirtualProduction/RemoteControl/Source/RemoteControlUI/Private/UI/Signature/Items/RCSignatureTreeFieldItem.h
+++ b/Engine/Plugins/VirtualProduction/RemoteControl/Source/RemoteControlUI/Private/UI/Signature/Items/RCSignatureTreeFieldItem.h
@@ -5,7 +5,7 @@
 #include "RCSignatureTreeItemBase.h"
 
 class FRCSignatureTreeSignatureItem;
-class URemoteControlSignatureRegistry;
+class URCSignatureRegistry;
 struct FRCSignature;
 struct FRCSignatureField;
 
@@ -19,7 +19,7 @@ public:
 
 	const FRCSignatureField* FindField() const;
 
-	FRCSignatureField* FindFieldMutable(URemoteControlSignatureRegistry** OutRegistry);
+	FRCSignatureField* FindFieldMutable(URCSignatureRegistry** OutRegistry);
 
 	void AddAction(const UScriptStruct* InActionType);
 
@@ -38,7 +38,7 @@ protected:
 private:
 	TSharedPtr<FRCSignatureTreeSignatureItem> GetParentSignatureItem() const;
 
-	FRCSignature* FindParentSignature(URemoteControlSignatureRegistry** OutRegistry);
+	FRCSignature* FindParentSignature(URCSignatureRegistry** OutRegistry);
 
 	int32 FieldIndex;
 };
diff --git a/Engine/Plugins/VirtualProduction/RemoteControl/Source/RemoteControlUI/Private/UI/Signature/Items/RCSignatureTreeRootItem.cpp b/Engine/Plugins/VirtualProduction/RemoteControl/Source/RemoteControlUI/Private/UI/Signature/Items/RCSignatureTreeRootItem.cpp
--- a/Engine/Plugins/VirtualProduction/RemoteControl/Source/RemoteControlUI/Private/UI/Signature/Items/RCSignatureTreeRootItem.cpp
+++ b/Engine/Plugins/VirtualProduction/RemoteControl/Source/RemoteControlUI/Private/UI/Signature/Items/RCSignatureTreeRootItem.cpp
@@ -1,9 +1,9 @@
 // Copyright Epic Games, Inc. All Rights Reserved.
 
 #include "RCSignatureTreeRootItem.h"
+#include "RCSignature.h"
+#include "RCSignatureRegistry.h"
 #include "RCSignatureTreeSignatureItem.h"
-#include "RemoteControlSignature.h"
-#include "RemoteControlSignatureRegistry.h"
 #include "UI/Signature/RCSignatureTreeItemSelection.h"
 #include "UI/Signature/SRCSignatureTree.h"
 
@@ -22,7 +22,7 @@ void FRCSignatureTreeRootItem::GenerateChildren(TArray<TSharedPtr<FRCSignatureTr
 		return;
 	}
 
-	URemoteControlSignatureRegistry* SignatureRegistry = SignatureTree->GetSignatureRegistry();
+	URCSignatureRegistry* SignatureRegistry = SignatureTree->GetSignatureRegistry();
 	if (!SignatureRegistry)
 	{
 		return;
diff --git a/Engine/Plugins/VirtualProduction/RemoteControl/Source/RemoteControlUI/Private/UI/Signature/Items/RCSignatureTreeSignatureItem.cpp b/Engine/Plugins/VirtualProduction/RemoteControl/Source/RemoteControlUI/Private/UI/Signature/Items/RCSignatureTreeSignatureItem.cpp
--- a/Engine/Plugins/VirtualProduction/RemoteControl/Source/RemoteControlUI/Private/UI/Signature/Items/RCSignatureTreeSignatureItem.cpp
+++ b/Engine/Plugins/VirtualProduction/RemoteControl/Source/RemoteControlUI/Private/UI/Signature/Items/RCSignatureTreeSignatureItem.cpp
@@ -2,9 +2,9 @@
 
 #include "RCSignatureTreeSignatureItem.h"
 #include "IRemoteControlUIModule.h"
+#include "RCSignatureRegistry.h"
 #include "RCSignatureTreeFieldItem.h"
 #include "RemoteControlPreset.h"
-#include "RemoteControlSignatureRegistry.h"
 #include "ScopedTransaction.h"
 #include "UI/SRemoteControlPanel.h"
 #include "UI/Signature/SRCSignatureTree.h"
@@ -26,21 +26,21 @@ const FGuid& FRCSignatureTreeSignatureItem::GetSignatureId() const
 	return SignatureId;
 }
 
-URemoteControlSignatureRegistry* FRCSignatureTreeSignatureItem::GetRegistry() const
+URCSignatureRegistry* FRCSignatureTreeSignatureItem::GetRegistry() const
 {
 	return RegistryWeak.Get();
 }
 
 const FRCSignature* FRCSignatureTreeSignatureItem::FindSignature() const
 {
-	if (URemoteControlSignatureRegistry* SignatureRegistry = GetRegistry())
+	if (URCSignatureRegistry* SignatureRegistry = GetRegistry())
 	{
 		return SignatureRegistry->FindSignature(SignatureId);
 	}
 	return nullptr;
 }
 
-FRCSignature* FRCSignatureTreeSignatureItem::FindSignatureMutable(URemoteControlSignatureRegistry* InRegistry)
+FRCSignature* FRCSignatureTreeSignatureItem::FindSignatureMutable(URCSignatureRegistry* InRegistry)
 {
 	if (InRegistry)
 	{
@@ -49,7 +49,7 @@ FRCSignature* FRCSignatureTreeSignatureItem::FindSignatureMutable(URemoteControl
 	return nullptr;
 }
 
-bool FRCSignatureTreeSignatureItem::AddField(URemoteControlSignatureRegistry* InRegistry, const TSharedRef<IPropertyHandle>& InPropertyHandle)
+bool FRCSignatureTreeSignatureItem::AddField(URCSignatureRegistry* InRegistry, const TSharedRef<IPropertyHandle>& InPropertyHandle)
 {
 	FRCSignature* Signature = FindSignatureMutable(InRegistry);
 	if (!Signature)
@@ -121,7 +121,7 @@ TOptional<bool> FRCSignatureTreeSignatureItem::IsEnabled() const
 
 void FRCSignatureTreeSignatureItem::SetEnabled(bool bInEnabled)
 {
-	URemoteControlSignatureRegistry* SignatureRegistry = GetRegistry();
+	URCSignatureRegistry* SignatureRegistry = GetRegistry();
 
 	FRCSignature* Signature = FindSignatureMutable(SignatureRegistry);
 	if (!Signature || Signature->bEnabled == bInEnabled)
@@ -153,7 +153,7 @@ bool FRCSignatureTreeSignatureItem::CanEditDisplayNameText() const
 
 void FRCSignatureTreeSignatureItem::SetDisplayNameText(const FText& InText)
 {
-	URemoteControlSignatureRegistry* SignatureRegistry = GetRegistry();
+	URCSignatureRegistry* SignatureRegistry = GetRegistry();
 
 	FRCSignature* Signature = FindSignatureMutable(SignatureRegistry);
 	if (!Signature || Signature->DisplayName.EqualTo(InText))
@@ -173,7 +173,7 @@ FText FRCSignatureTreeSignatureItem::GetDescription() const
 
 int32 FRCSignatureTreeSignatureItem::RemoveFromRegistry()
 {
-	URemoteControlSignatureRegistry* SignatureRegistry = GetRegistry();
+	URCSignatureRegistry* SignatureRegistry = GetRegistry();
 	if (!SignatureRegistry)
 	{
 		return 0;
diff --git a/Engine/Plugins/VirtualProduction/RemoteControl/Source/RemoteControlUI/Private/UI/Signature/Items/RCSignatureTreeSignatureItem.h b/Engine/Plugins/VirtualProduction/RemoteControl/Source/RemoteControlUI/Private/UI/Signature/Items/RCSignatureTreeSignatureItem.h
--- a/Engine/Plugins/VirtualProduction/RemoteControl/Source/RemoteControlUI/Private/UI/Signature/Items/RCSignatureTreeSignatureItem.h
+++ b/Engine/Plugins/VirtualProduction/RemoteControl/Source/RemoteControlUI/Private/UI/Signature/Items/RCSignatureTreeSignatureItem.h
@@ -9,7 +9,7 @@
 #include "UObject/WeakObjectPtrTemplates.h"
 
 class IPropertyHandle;
-class URemoteControlSignatureRegistry;
+class URCSignatureRegistry;
 struct FRCSignature;
 
 /** Item class representing an RC Signature */
@@ -22,13 +22,13 @@ public:
 
 	const FGuid& GetSignatureId() const;
 
-	URemoteControlSignatureRegistry* GetRegistry() const;
+	URCSignatureRegistry* GetRegistry() const;
 
 	const FRCSignature* FindSignature() const;
 
-	FRCSignature* FindSignatureMutable(URemoteControlSignatureRegistry* InRegistry);
+	FRCSignature* FindSignatureMutable(URCSignatureRegistry* InRegistry);
 
-	bool AddField(URemoteControlSignatureRegistry* InRegistry, const TSharedRef<IPropertyHandle>& InPropertyHandle);
+	bool AddField(URCSignatureRegistry* InRegistry, const TSharedRef<IPropertyHandle>& InPropertyHandle);
 
 	//~ Begin IRCSignatureItem
 	virtual void ApplySignature(TConstArrayView<TWeakObjectPtr<UObject>> InObjects) override;
@@ -49,7 +49,7 @@ protected:
 	//~ End FRCSignatureTreeItem
 
 private:
-	TWeakObjectPtr<URemoteControlSignatureRegistry> RegistryWeak;
+	TWeakObjectPtr<URCSignatureRegistry> RegistryWeak;
 
 	FGuid SignatureId;
 };
diff --git a/Engine/Plugins/VirtualProduction/RemoteControl/Source/RemoteControlUI/Private/UI/Signature/SRCSignaturePanel.cpp b/Engine/Plugins/VirtualProduction/RemoteControl/Source/RemoteControlUI/Private/UI/Signature/SRCSignaturePanel.cpp
--- a/Engine/Plugins/VirtualProduction/RemoteControl/Source/RemoteControlUI/Private/UI/Signature/SRCSignaturePanel.cpp
+++ b/Engine/Plugins/VirtualProduction/RemoteControl/Source/RemoteControlUI/Private/UI/Signature/SRCSignaturePanel.cpp
@@ -8,8 +8,8 @@
 #include "Items/RCSignatureTreeSignatureItem.h"
 #include "Label/RCSignatureLabelColumn.h"
 #include "Misc/MessageDialog.h"
+#include "RCSignatureRegistry.h"
 #include "RemoteControlPreset.h"
-#include "RemoteControlSignatureRegistry.h"
 #include "SRCSignatureTree.h"
 #include "ScopedTransaction.h"
 #include "Styling/RemoteControlStyles.h"
@@ -105,7 +105,7 @@ void SRCSignaturePanel::Construct(const FArguments& InArgs, const TSharedRef<SRe
 		];
 }
 
-URemoteControlSignatureRegistry* SRCSignaturePanel::GetSignatureRegistry() const
+URCSignatureRegistry* SRCSignaturePanel::GetSignatureRegistry() const
 {
 	if (URemoteControlPreset* Preset = GetPreset())
 	{
@@ -121,7 +121,7 @@ void SRCSignaturePanel::AddToSignature(const FRCExposesPropertyArgs& InPropertyA
 		return;
 	}
 
-	URemoteControlSignatureRegistry* Registry = GetSignatureRegistry();
+	URCSignatureRegistry* Registry = GetSignatureRegistry();
 	if (!Registry)
 	{
 		return;
@@ -263,7 +263,7 @@ void SRCSignaturePanel::Refresh()
 
 FReply SRCSignaturePanel::OnAddButtonClicked()
 {
-	URemoteControlSignatureRegistry* SignatureRegistry = GetSignatureRegistry();
+	URCSignatureRegistry* SignatureRegistry = GetSignatureRegistry();
 	if (!SignatureRegistry)
 	{
 		return FReply::Unhandled();
@@ -282,7 +282,7 @@ FReply SRCSignaturePanel::OnAddButtonClicked()
 
 FReply SRCSignaturePanel::DeleteAllItems()
 {
-	URemoteControlSignatureRegistry* SignatureRegistry = GetSignatureRegistry();
+	URCSignatureRegistry* SignatureRegistry = GetSignatureRegistry();
 	if (!SignatureRegistry)
 	{
 		return FReply::Unhandled();
diff --git a/Engine/Plugins/VirtualProduction/RemoteControl/Source/RemoteControlUI/Private/UI/Signature/SRCSignaturePanel.h b/Engine/Plugins/VirtualProduction/RemoteControl/Source/RemoteControlUI/Private/UI/Signature/SRCSignaturePanel.h
--- a/Engine/Plugins/VirtualProduction/RemoteControl/Source/RemoteControlUI/Private/UI/Signature/SRCSignaturePanel.h
+++ b/Engine/Plugins/VirtualProduction/RemoteControl/Source/RemoteControlUI/Private/UI/Signature/SRCSignaturePanel.h
@@ -6,7 +6,7 @@
 
 class SRCSignatureDetails;
 class SRCSignatureTree;
-class URemoteControlSignatureRegistry;
+class URCSignatureRegistry;
 struct FRCExposesPropertyArgs;
 
 class SRCSignaturePanel : public SRCLogicPanelBase, public FSelfRegisteringEditorUndoClient
@@ -19,7 +19,7 @@ public:
 	void Construct(const FArguments& InArgs, const TSharedRef<SRemoteControlPanel>& InPanel);
 
 	/** Gets the Signature Registry from the RC Preset */
-	URemoteControlSignatureRegistry* GetSignatureRegistry() const;
+	URCSignatureRegistry* GetSignatureRegistry() const;
 
 	void AddToSignature(const FRCExposesPropertyArgs& InPropertyArgs);
 
diff --git a/Engine/Plugins/VirtualProduction/RemoteControl/Source/RemoteControlUI/Private/UI/Signature/SRCSignatureTree.cpp b/Engine/Plugins/VirtualProduction/RemoteControl/Source/RemoteControlUI/Private/UI/Signature/SRCSignatureTree.cpp
--- a/Engine/Plugins/VirtualProduction/RemoteControl/Source/RemoteControlUI/Private/UI/Signature/SRCSignatureTree.cpp
+++ b/Engine/Plugins/VirtualProduction/RemoteControl/Source/RemoteControlUI/Private/UI/Signature/SRCSignatureTree.cpp
@@ -5,9 +5,9 @@
 #include "Items/RCSignatureTreeItemBase.h"
 #include "Items/RCSignatureTreeRootItem.h"
 #include "Items/RCSignatureTreeSignatureItem.h"
+#include "RCSignatureRegistry.h"
 #include "RCSignatureTreeItemSelection.h"
 #include "RemoteControlPreset.h"
-#include "RemoteControlSignatureRegistry.h"
 #include "SRCSignaturePanel.h"
 #include "SRCSignatureRow.h"
 #include "ScopedTransaction.h"
@@ -57,7 +57,7 @@ TSharedRef<FRCSignatureTreeRootItem> SRCSignatureTree::GetRootItem() const
 	return RootItem.ToSharedRef();
 }
 
-URemoteControlSignatureRegistry* SRCSignatureTree::GetSignatureRegistry() const
+URCSignatureRegistry* SRCSignatureTree::GetSignatureRegistry() const
 {
 	TSharedPtr<SRCSignaturePanel> SignaturePanel = SignaturePanelWeak.Pin();
 	if (!SignaturePanel.IsValid())
diff --git a/Engine/Plugins/VirtualProduction/RemoteControl/Source/RemoteControlUI/Private/UI/Signature/SRCSignatureTree.h b/Engine/Plugins/VirtualProduction/RemoteControl/Source/RemoteControlUI/Private/UI/Signature/SRCSignatureTree.h
--- a/Engine/Plugins/VirtualProduction/RemoteControl/Source/RemoteControlUI/Private/UI/Signature/SRCSignatureTree.h
+++ b/Engine/Plugins/VirtualProduction/RemoteControl/Source/RemoteControlUI/Private/UI/Signature/SRCSignatureTree.h
@@ -9,7 +9,7 @@ class FRCSignatureTreeRootItem;
 class IRCSignatureColumn;
 class SHeaderRow;
 class SRCSignaturePanel;
-class URemoteControlSignatureRegistry;
+class URCSignatureRegistry;
 template<typename ItemType> class STreeView;
 
 class SRCSignatureTree : public SRCLogicPanelListBase
@@ -23,7 +23,7 @@ public:
 
 	TSharedRef<FRCSignatureTreeRootItem> GetRootItem() const;
 
-	URemoteControlSignatureRegistry* GetSignatureRegistry() const;
+	URCSignatureRegistry* GetSignatureRegistry() const;
 
 	TSharedPtr<IRCSignatureColumn> FindColumn(FName InColumnName) const;
 
