Remote Control: Signature obj rename

diff --git a/Engine/Plugins/VirtualProduction/RemoteControl/Source/RemoteControl/Private/RemoteControlPreset.cpp b/Engine/Plugins/VirtualProduction/RemoteControl/Source/RemoteControl/Private/RemoteControlPreset.cpp
--- a/Engine/Plugins/VirtualProduction/RemoteControl/Source/RemoteControl/Private/RemoteControlPreset.cpp
+++ b/Engine/Plugins/VirtualProduction/RemoteControl/Source/RemoteControl/Private/RemoteControlPreset.cpp
@@ -673,8 +673,7 @@ URemoteControlPreset::URemoteControlPreset()
 	, RebindingManager(MakePimpl<FRemoteControlPresetRebindingManager>())
 {
 	Registry = CreateDefaultSubobject<URemoteControlExposeRegistry>(FName("ExposeRegistry"));
-
-	SignatureRegistry = CreateDefaultSubobject<URCSignatureRegistry>(TEXT("SignatureRegistry"));
+	Signatures = CreateDefaultSubobject<URCSignatureRegistry>(TEXT("Signatures"));
 
 	PropertyIdRegistry = CreateDefaultSubobject<URemoteControlPropertyIdRegistry>(FName("PropertyIdRegistry"));
 	PropertyIdRegistry->Initialize();
diff --git a/Engine/Plugins/VirtualProduction/RemoteControl/Source/RemoteControl/Public/RemoteControlPreset.h b/Engine/Plugins/VirtualProduction/RemoteControl/Source/RemoteControl/Public/RemoteControlPreset.h
--- a/Engine/Plugins/VirtualProduction/RemoteControl/Source/RemoteControl/Public/RemoteControlPreset.h
+++ b/Engine/Plugins/VirtualProduction/RemoteControl/Source/RemoteControl/Public/RemoteControlPreset.h
@@ -709,7 +709,7 @@ public:
 	/** Returns the preset's signature registry */
 	URCSignatureRegistry* GetSignatureRegistry() const
 	{
-		return SignatureRegistry;
+		return Signatures;
 	}
 
 public:
@@ -945,7 +945,7 @@ private:
 
 	/** Holds all signatures on the preset. */
 	UPROPERTY(Instanced)
-	TObjectPtr<URCSignatureRegistry> SignatureRegistry = nullptr;
+	TObjectPtr<URCSignatureRegistry> Signatures = nullptr;
 
 	/** Holds identities of exposed entities on the preset. */
 	UPROPERTY(Transient)
