Remote Control: Moved the Controller suppert check in UI module to its own utility in Logic module

diff --git a/Engine/Plugins/VirtualProduction/RemoteControl/Source/RemoteControlLogic/Private/Controller/RCControllerUtilities.cpp b/Engine/Plugins/VirtualProduction/RemoteControl/Source/RemoteControlLogic/Private/Controller/RCControllerUtilities.cpp
new file mode 100644
--- /dev/null
+++ b/Engine/Plugins/VirtualProduction/RemoteControl/Source/RemoteControlLogic/Private/Controller/RCControllerUtilities.cpp
@@ -0,0 +1,82 @@
+// Copyright Epic Games, Inc. All Rights Reserved.
+
+#include "Controller/RCControllerUtilities.h"
+#include "Engine/StaticMesh.h"
+#include "Engine/Texture.h"
+#include "Materials/MaterialInterface.h"
+#include "StructUtils/PropertyBag.h"
+#include "UObject/Class.h"
+
+namespace UE::RCControllers
+{
+	namespace Private
+	{
+		/**
+		 * Determines whether a class of a specific type is supported explicitly. 
+		 * Explicit here means that checking against a derived type of a supported type will return false as the checks done here do not include child/derived classes.
+		 */
+		bool CanCreateControllerFromStruct(const UScriptStruct* InStruct)
+		{
+			return InStruct == TBaseStructure<FVector>::Get()
+				|| InStruct == TBaseStructure<FVector2D>::Get()
+				|| InStruct == TBaseStructure<FRotator>::Get()
+				|| InStruct == TBaseStructure<FColor>::Get();
+		}
+
+		/**
+		 * Determines whether a class of a specific type is supported explicitly. 
+		 * Explicit here means that checking against UTexture2D or UMaterial will return false as the checks done here do not include child/derived classes.
+		 */
+		bool CanCreateControllerFromClass(const UClass* InClass)
+		{
+			return InClass == UTexture::StaticClass()
+				|| InClass == UStaticMesh::StaticClass()
+				|| InClass == UMaterialInterface::StaticClass();
+		}
+	} // UE::RCControllers::Private
+
+	bool CanCreateControllerFromPropertyDesc(const FPropertyBagPropertyDesc& InPropertyDesc)
+	{
+		switch (InPropertyDesc.ValueType)
+		{
+		case EPropertyBagPropertyType::Enum:
+			if (const UEnum* Enum = Cast<UEnum>(InPropertyDesc.ValueTypeObject))
+			{
+				const int64 MaxEnumValue = Enum->GetMaxEnumValue();
+				const uint32 NeededBits = FMath::RoundUpToPowerOfTwo(MaxEnumValue);
+
+				// 8 bits enums only
+				return NeededBits <= 256;
+			}
+			break;
+
+		case EPropertyBagPropertyType::Struct:
+			if (const UScriptStruct* ScriptStruct = Cast<UScriptStruct>(InPropertyDesc.ValueTypeObject))
+			{
+				return Private::CanCreateControllerFromStruct(ScriptStruct);
+			}
+			break;
+
+		case EPropertyBagPropertyType::Object:
+		case EPropertyBagPropertyType::SoftObject:
+		case EPropertyBagPropertyType::Class:
+		case EPropertyBagPropertyType::SoftClass:
+			if (const UClass* Class = Cast<UClass>(InPropertyDesc.ValueTypeObject))
+			{
+				return Private::CanCreateControllerFromClass(Class);
+			}
+			break;
+
+		default:
+			// All other types are allowed
+			return true;
+		}
+
+		return false;
+	}
+
+	bool CanCreateControllerFromProperty(const FProperty* InProperty)
+	{
+		return InProperty && CanCreateControllerFromPropertyDesc(FPropertyBagPropertyDesc(NAME_None, InProperty));
+	}
+} // UE::RCControllers
diff --git a/Engine/Plugins/VirtualProduction/RemoteControl/Source/RemoteControlLogic/Public/Controller/RCControllerUtilities.h b/Engine/Plugins/VirtualProduction/RemoteControl/Source/RemoteControlLogic/Public/Controller/RCControllerUtilities.h
new file mode 100644
--- /dev/null
+++ b/Engine/Plugins/VirtualProduction/RemoteControl/Source/RemoteControlLogic/Public/Controller/RCControllerUtilities.h
@@ -0,0 +1,23 @@
+// Copyright Epic Games, Inc. All Rights Reserved.
+
+#pragma once
+
+#include "HAL/Platform.h"
+
+class FProperty;
+struct FPropertyBagPropertyDesc;
+
+namespace UE::RCControllers
+{
+	/**
+	 * Determines whether a controller can be created from a given property desc explicitly
+	 * Explicit meaning that it won't use IsChildOf or recurse up its Super Types for Object/Struct properties
+	 */
+	bool CanCreateControllerFromPropertyDesc(const FPropertyBagPropertyDesc& InPropertyDesc);
+
+	/**
+	 * Determines whether a controller can be created from a given property type explicitly
+	 * Explicit meaning that it won't use IsChildOf or recurse up its Super Types for Object/Struct properties
+	 */
+	REMOTECONTROLLOGIC_API bool CanCreateControllerFromProperty(const FProperty* InProperty);
+};
diff --git a/Engine/Plugins/VirtualProduction/RemoteControl/Source/RemoteControlUI/Private/UI/Controller/SRCControllerPanelList.cpp b/Engine/Plugins/VirtualProduction/RemoteControl/Source/RemoteControlUI/Private/UI/Controller/SRCControllerPanelList.cpp
--- a/Engine/Plugins/VirtualProduction/RemoteControl/Source/RemoteControlUI/Private/UI/Controller/SRCControllerPanelList.cpp
+++ b/Engine/Plugins/VirtualProduction/RemoteControl/Source/RemoteControlUI/Private/UI/Controller/SRCControllerPanelList.cpp
@@ -7,6 +7,7 @@
 #include "Commands/RemoteControlCommands.h"
 #include "Controller/RCController.h"
 #include "Controller/RCControllerContainer.h"
+#include "Controller/RCControllerUtilities.h"
 #include "Controller/RCCustomControllerUtilities.h"
 #include "Engine/Texture.h"
 #include "Engine/Texture2D.h"
@@ -60,49 +61,6 @@ namespace UE::RCControllerPanelList
 		const FName ValueTypeSelection = TEXT("Value Type Selection");
 	}
 
-	const TSet<UScriptStruct*>& GetSupportedStructs()
-	{
-		static const TSet<UScriptStruct*> SupportedStructs = {
-			TBaseStructure<FVector>::Get(),
-			TBaseStructure<FVector2D>::Get(),
-			TBaseStructure<FRotator>::Get(),
-			TBaseStructure<FColor>::Get()
-		};
-		
-		return SupportedStructs;
-	}
-
-	const TSet<UClass*>& GetSupportedObjects()
-	{
-		static const TSet<UClass*> SupportedObjects = {
-			UTexture::StaticClass(),
-			UStaticMesh::StaticClass(),
-			UMaterialInterface::StaticClass(),
-		};
-		
-		return SupportedObjects;
-	}
-	
-	bool IsStructPropertyTypeSupported(const FStructProperty* InStructProperty)
-	{
-		if (InStructProperty)
-		{
-			return GetSupportedStructs().Contains(InStructProperty->Struct);
-		}
-
-		return false;
-	}
-
-	bool IsObjectPropertyTypeSupported(const FObjectProperty* InObjectProperty)
-	{
-		if (InObjectProperty)
-		{
-			return GetSupportedObjects().Contains(InObjectProperty->PropertyClass);
-		}
-
-		return false;
-	}
-
 	class SControllerItemListRow : public SMultiColumnTableRow<TSharedRef<FRCControllerModel>>
 	{
 	public:
@@ -868,29 +826,7 @@ bool SRCControllerPanelList::IsEntitySupported(const FGuid ExposedEntityId)
 
 			if (RemoteControlProperty->FieldType == EExposedFieldType::Property)
 			{
-				const FProperty* Property = RemoteControlProperty->GetProperty();
-				
-				if (const FEnumProperty* EnumProperty = CastField<FEnumProperty>(Property))
-				{
-					if (const UEnum* Enum = EnumProperty->GetEnum())
-					{
-						const int64 MaxEnumValue = Enum->GetMaxEnumValue();
-						const uint32 NeededBits = FMath::RoundUpToPowerOfTwo(MaxEnumValue);
-
-						// 8 bits enums only
-						return NeededBits <= 256;
-					}
-				}
-				else if (const FStructProperty* StructProperty = CastField<FStructProperty>(Property))
-				{
-					return UE::RCControllerPanelList::IsStructPropertyTypeSupported(StructProperty);
-				}
-				else if (const FObjectProperty* ObjectProperty = CastField<FObjectProperty>(Property))
-				{
-					return UE::RCControllerPanelList::IsObjectPropertyTypeSupported(ObjectProperty);
-				}
-
-				return true;
+				return UE::RCControllers::CanCreateControllerFromProperty(RemoteControlProperty->GetProperty());
 			}
 		}
 	}
