Remote Control: added Support check for Signature Action as some fields might not be able to execute certain actions (e.g. a field of unsupported type for Controllers should not give option to create a controller with bind).

diff --git a/Engine/Plugins/VirtualProduction/RemoteControl/Source/RemoteControl/Private/RCSignature.cpp b/Engine/Plugins/VirtualProduction/RemoteControl/Source/RemoteControl/Private/RCSignature.cpp
--- a/Engine/Plugins/VirtualProduction/RemoteControl/Source/RemoteControl/Private/RCSignature.cpp
+++ b/Engine/Plugins/VirtualProduction/RemoteControl/Source/RemoteControl/Private/RCSignature.cpp
@@ -75,6 +75,11 @@ FRCSignatureField FRCSignatureField::CreateField(const FRCFieldPathInfo& InField
 	Field.FieldPath = InFieldPathInfo;
 	Field.SupportedClass = SupportedClass;
 
+	if (InProperty)
+	{
+		Field.PropertyDesc = FPropertyBagPropertyDesc(InProperty->GetFName(), InProperty);
+	}
+
 	if (InOwnerObject)
 	{
 		if (AActor* ActorOwner = InOwnerObject->GetTypedOuter<AActor>())
diff --git a/Engine/Plugins/VirtualProduction/RemoteControl/Source/RemoteControl/Private/RCSignatureActionInstance.cpp b/Engine/Plugins/VirtualProduction/RemoteControl/Source/RemoteControl/Private/RCSignatureActionInstance.cpp
--- a/Engine/Plugins/VirtualProduction/RemoteControl/Source/RemoteControl/Private/RCSignatureActionInstance.cpp
+++ b/Engine/Plugins/VirtualProduction/RemoteControl/Source/RemoteControl/Private/RCSignatureActionInstance.cpp
@@ -10,6 +10,7 @@ FRCSignatureActionInstance::FRCSignatureActionInstance(const UScriptStruct* InSc
 
 	if (FRCSignatureAction* Action = ActionInstance.GetMutablePtr())
 	{
+		checkSlow(Action->IsSupported(InFieldOwner));
 		Action->Initialize(InFieldOwner);
 	}
 }
diff --git a/Engine/Plugins/VirtualProduction/RemoteControl/Source/RemoteControl/Public/RCSignature.h b/Engine/Plugins/VirtualProduction/RemoteControl/Source/RemoteControl/Public/RCSignature.h
--- a/Engine/Plugins/VirtualProduction/RemoteControl/Source/RemoteControl/Public/RCSignature.h
+++ b/Engine/Plugins/VirtualProduction/RemoteControl/Source/RemoteControl/Public/RCSignature.h
@@ -4,6 +4,7 @@
 
 #include "RCSignatureActionInstance.h"
 #include "RemoteControlFieldPath.h"
+#include "StructUtils/PropertyBag.h"
 #include "RCSignature.generated.h"
 
 enum class EPropertyBagPropertyType : uint8;
@@ -38,6 +39,10 @@ struct FRCSignatureField
 	UPROPERTY()
 	FSoftClassPath SupportedClass;
 
+	/** Optional: Property Description of the field (if property) */
+	UPROPERTY()
+	FPropertyBagPropertyDesc PropertyDesc;
+
 	/** Container holding the action instances for the field */
 	UPROPERTY()
 	TArray<FRCSignatureActionInstance> Actions;
diff --git a/Engine/Plugins/VirtualProduction/RemoteControl/Source/RemoteControl/Public/RCSignatureAction.h b/Engine/Plugins/VirtualProduction/RemoteControl/Source/RemoteControl/Public/RCSignatureAction.h
--- a/Engine/Plugins/VirtualProduction/RemoteControl/Source/RemoteControl/Public/RCSignatureAction.h
+++ b/Engine/Plugins/VirtualProduction/RemoteControl/Source/RemoteControl/Public/RCSignatureAction.h
@@ -7,6 +7,7 @@
 #include "RCSignatureAction.generated.h"
 
 class URemoteControlPreset;
+struct FRCSignature;
 struct FRCSignatureField;
 struct FRemoteControlProperty;
 
@@ -57,12 +58,22 @@ struct FRCSignatureAction
 
 	/**
 	 * Called only once when the Signature Action is first added to the action list
-	 * @param InFieldOwner the owner of the action list
+	 * @param InField the owner of the action list
 	 */
-	virtual void Initialize(const FRCSignatureField& InFieldOwner)
+	virtual void Initialize(const FRCSignatureField& InField)
 	{
 	}
 
+	/**
+	 * Determines whether this Action can execute under a given field
+	 * @param InField the field that will contain this action if supported
+	 * @return true if the action is supported
+	 */
+	virtual bool IsSupported(const FRCSignatureField& InField) const
+	{
+		return true;
+	}
+
 	/**
 	 * Executes the Action logic
 	 * @param InContext context of the action (preset, exposed property and uobject being used)
diff --git a/Engine/Plugins/VirtualProduction/RemoteControl/Source/RemoteControlLogic/Private/Signature/RCSignatureControllerAction.cpp b/Engine/Plugins/VirtualProduction/RemoteControl/Source/RemoteControlLogic/Private/Signature/RCSignatureControllerAction.cpp
--- a/Engine/Plugins/VirtualProduction/RemoteControl/Source/RemoteControlLogic/Private/Signature/RCSignatureControllerAction.cpp
+++ b/Engine/Plugins/VirtualProduction/RemoteControl/Source/RemoteControlLogic/Private/Signature/RCSignatureControllerAction.cpp
@@ -7,9 +7,9 @@
 #include "RCSignature.h"
 #include "RemoteControlPreset.h"
 
-void FRCSignatureControllerAction::Initialize(const FRCSignatureField& InFieldOwner)
+void FRCSignatureControllerAction::Initialize(const FRCSignatureField& InField)
 {
-	ControllerName = InFieldOwner.FieldPath.GetFieldName();
+	ControllerName = InField.FieldPath.GetFieldName();
 }
 
 bool FRCSignatureControllerAction::Execute(const FRCSignatureActionContext& InContext) const
diff --git a/Engine/Plugins/VirtualProduction/RemoteControl/Source/RemoteControlLogic/Private/Signature/RCSignatureControllerAction.h b/Engine/Plugins/VirtualProduction/RemoteControl/Source/RemoteControlLogic/Private/Signature/RCSignatureControllerAction.h
--- a/Engine/Plugins/VirtualProduction/RemoteControl/Source/RemoteControlLogic/Private/Signature/RCSignatureControllerAction.h
+++ b/Engine/Plugins/VirtualProduction/RemoteControl/Source/RemoteControlLogic/Private/Signature/RCSignatureControllerAction.h
@@ -14,7 +14,7 @@ struct FRCSignatureControllerAction : public FRCSignatureAction
 	GENERATED_BODY()
 
 	//~ Begin FRCSignatureAction
-	virtual void Initialize(const FRCSignatureField& InFieldOwner) override;
+	virtual void Initialize(const FRCSignatureField& InField) override;
 	virtual bool Execute(const FRCSignatureActionContext& InContext) const override;
 #if WITH_EDITOR
 	virtual FRCSignatureActionIcon GetIcon() const override;
diff --git a/Engine/Plugins/VirtualProduction/RemoteControl/Source/RemoteControlLogic/Private/Signature/RCSignatureControllerBindAction.cpp b/Engine/Plugins/VirtualProduction/RemoteControl/Source/RemoteControlLogic/Private/Signature/RCSignatureControllerBindAction.cpp
--- a/Engine/Plugins/VirtualProduction/RemoteControl/Source/RemoteControlLogic/Private/Signature/RCSignatureControllerBindAction.cpp
+++ b/Engine/Plugins/VirtualProduction/RemoteControl/Source/RemoteControlLogic/Private/Signature/RCSignatureControllerBindAction.cpp
@@ -5,13 +5,20 @@
 #include "Behaviour/Builtin/Bind/RCBehaviourBind.h"
 #include "Behaviour/Builtin/Bind/RCBehaviourBindNode.h"
 #include "Controller/RCController.h"
+#include "Controller/RCControllerUtilities.h"
 #include "Controller/RCCustomControllerUtilities.h"
+#include "RCSignature.h"
 #include "RemoteControlField.h"
 
 #if WITH_EDITOR
 #include "Engine/Texture2D.h"
 #endif
 
+bool FRCSignatureControllerBindAction::IsSupported(const FRCSignatureField& InField) const
+{
+	return UE::RCControllers::CanCreateControllerFromPropertyDesc(InField.PropertyDesc);
+}
+
 bool FRCSignatureControllerBindAction::IsControllerCompatible(const FRCSignatureActionContext& InContext, URCController* InController) const
 {
 	return URCBehaviourBind::CanHaveActionForField(InController, InContext.Property.ToSharedRef(), /*bAllowNumericInputAsStrings*/true);
diff --git a/Engine/Plugins/VirtualProduction/RemoteControl/Source/RemoteControlLogic/Private/Signature/RCSignatureControllerBindAction.h b/Engine/Plugins/VirtualProduction/RemoteControl/Source/RemoteControlLogic/Private/Signature/RCSignatureControllerBindAction.h
--- a/Engine/Plugins/VirtualProduction/RemoteControl/Source/RemoteControlLogic/Private/Signature/RCSignatureControllerBindAction.h
+++ b/Engine/Plugins/VirtualProduction/RemoteControl/Source/RemoteControlLogic/Private/Signature/RCSignatureControllerBindAction.h
@@ -10,6 +10,10 @@ struct FRCSignatureControllerBindAction : public FRCSignatureControllerAction
 {
 	GENERATED_BODY()
 
+	//~ Begin FRCSignatureAction
+	virtual bool IsSupported(const FRCSignatureField& InField) const override;
+	//~ End FRCSignatureAction
+
 	//~ Begin FRCSignatureControllerAction
 	virtual bool IsControllerCompatible(const FRCSignatureActionContext& InContext, URCController* InController) const override;
 	virtual bool MakeControllerDesc(const FRCSignatureActionContext& InContext, FPropertyBagPropertyDesc& OutDesc) const override;
diff --git a/Engine/Plugins/VirtualProduction/RemoteControl/Source/RemoteControlUI/Private/UI/Signature/Actions/RCSignatureActionColumn.cpp b/Engine/Plugins/VirtualProduction/RemoteControl/Source/RemoteControlUI/Private/UI/Signature/Actions/RCSignatureActionColumn.cpp
--- a/Engine/Plugins/VirtualProduction/RemoteControl/Source/RemoteControlUI/Private/UI/Signature/Actions/RCSignatureActionColumn.cpp
+++ b/Engine/Plugins/VirtualProduction/RemoteControl/Source/RemoteControlUI/Private/UI/Signature/Actions/RCSignatureActionColumn.cpp
@@ -5,6 +5,7 @@
 #include "RCSignatureActionType.h"
 #include "SRCSignatureActionBox.h"
 #include "StructUtils/InstancedStruct.h"
+#include "UI/Signature/Items/RCSignatureTreeFieldItem.h"
 #include "UObject/Class.h"
 #include "UObject/UObjectIterator.h"
 
@@ -13,7 +14,6 @@
 FRCSignatureActionColumn::FRCSignatureActionColumn(const TAttribute<bool>& InLiveMode)
 	: LiveMode(InLiveMode)
 {
-	RefreshActionTypes();
 }
 
 FName FRCSignatureActionColumn::GetColumnId() const
@@ -38,10 +38,10 @@ TSharedRef<SWidget> FRCSignatureActionColumn::ConstructRowWidget(TSharedPtr<FRCS
 	return SNew(SRCSignatureActionBox, InItem.ToSharedRef(), InRow)
 		.LiveMode(LiveMode)
 		.ActionTypesSource(&ActionTypes)
-		.OnActionTypesComboBoxOpening(this, &FRCSignatureActionColumn::RefreshActionTypes);
+		.OnRefreshActionTypes(this, &FRCSignatureActionColumn::RefreshActionTypes);
 }
 
-void FRCSignatureActionColumn::RefreshActionTypes()
+void FRCSignatureActionColumn::RefreshActionTypes(const TSharedRef<FRCSignatureTreeFieldItem>& InFieldItem)
 {
 	ActionTypes.Reset();
 
@@ -50,6 +50,12 @@ void FRCSignatureActionColumn::RefreshActionTypes()
 		return;
 	}
 
+	const FRCSignatureField* Field = InFieldItem->FindField();
+	if (!Field)
+	{
+		return;
+	}
+
 	const FName HiddenMetaData(TEXT("Hidden"));
 
 	for (UScriptStruct* ScriptStruct : TObjectRange<UScriptStruct>())
@@ -61,17 +67,22 @@ void FRCSignatureActionColumn::RefreshActionTypes()
 
 		if (ScriptStruct->IsChildOf(TBaseStructure<FRCSignatureAction>::Get()))
 		{
-			TSharedRef<FRCSignatureActionType> ActionType = MakeShared<FRCSignatureActionType>();
-			ActionType->Type = ScriptStruct;
-			ActionType->Title = ScriptStruct->GetDisplayNameText();
+			// Initialize a Temp Instance to determine if it's supported by the Tree Item
+			TInstancedStruct<FRCSignatureAction> Instance;
+			Instance.InitializeAsScriptStruct(ScriptStruct, /*StructMemory*/nullptr);
+
+			const FRCSignatureAction& Action = Instance.Get();
 
-			// Initialize a Temp Instance to get the Icon to use
+			if (!Action.IsSupported(*Field))
 			{
-				TInstancedStruct<FRCSignatureAction> Instance;
-				Instance.InitializeAsScriptStruct(ScriptStruct, /*StructMemory*/nullptr);
-				ActionType->Icon = Instance.Get().GetIcon();
+				continue;
 			}
 
+			TSharedRef<FRCSignatureActionType> ActionType = MakeShared<FRCSignatureActionType>();
+			ActionType->Type = ScriptStruct;
+			ActionType->Title = ScriptStruct->GetDisplayNameText();
+			ActionType->Icon = Action.GetIcon();
+
 			ActionTypes.Add(MoveTemp(ActionType));
 		}
 	}
diff --git a/Engine/Plugins/VirtualProduction/RemoteControl/Source/RemoteControlUI/Private/UI/Signature/Actions/RCSignatureActionColumn.h b/Engine/Plugins/VirtualProduction/RemoteControl/Source/RemoteControlUI/Private/UI/Signature/Actions/RCSignatureActionColumn.h
--- a/Engine/Plugins/VirtualProduction/RemoteControl/Source/RemoteControlUI/Private/UI/Signature/Actions/RCSignatureActionColumn.h
+++ b/Engine/Plugins/VirtualProduction/RemoteControl/Source/RemoteControlUI/Private/UI/Signature/Actions/RCSignatureActionColumn.h
@@ -4,6 +4,7 @@
 
 #include "UI/Signature/IRCSignatureColumn.h"
 
+class FRCSignatureTreeFieldItem;
 struct FRCSignatureActionType;
 
 class FRCSignatureActionColumn : public IRCSignatureColumn
@@ -19,7 +20,7 @@ private:
 	virtual TSharedRef<SWidget> ConstructRowWidget(TSharedPtr<FRCSignatureTreeItemBase> InItem, const TSharedRef<SRCSignatureTree>& InList, const TSharedRef<SRCSignatureRow>& InRow) override;
 	//~ End IRCSignatureColumn
 
-	void RefreshActionTypes();
+	void RefreshActionTypes(const TSharedRef<FRCSignatureTreeFieldItem>& InFieldItem);
 
 	TArray<TSharedPtr<FRCSignatureActionType>> ActionTypes; 
 
diff --git a/Engine/Plugins/VirtualProduction/RemoteControl/Source/RemoteControlUI/Private/UI/Signature/Actions/SRCSignatureActionBox.cpp b/Engine/Plugins/VirtualProduction/RemoteControl/Source/RemoteControlUI/Private/UI/Signature/Actions/SRCSignatureActionBox.cpp
--- a/Engine/Plugins/VirtualProduction/RemoteControl/Source/RemoteControlUI/Private/UI/Signature/Actions/SRCSignatureActionBox.cpp
+++ b/Engine/Plugins/VirtualProduction/RemoteControl/Source/RemoteControlUI/Private/UI/Signature/Actions/SRCSignatureActionBox.cpp
@@ -19,6 +19,8 @@ void SRCSignatureActionBox::Construct(const FArguments& InArgs, const TSharedRef
 {
 	ItemWeak = InItem;
 	LiveMode = InArgs._LiveMode;
+	OnRefreshActionTypes = InArgs._OnRefreshActionTypes;
+
 	IsHovered = TAttribute<bool>(InRow, &SRCSignatureRow::IsHovered);
 	IsSelected = TAttribute<bool>(InRow, &SRCSignatureRow::IsSelected);
 
@@ -27,7 +29,7 @@ void SRCSignatureActionBox::Construct(const FArguments& InArgs, const TSharedRef
 		.Visibility(this, &SRCSignatureActionBox::GetAddActionVisibility)
 		.ComboBoxStyle(&FAppStyle::Get().GetWidgetStyle<FComboBoxStyle>("SimpleComboBox"))
 		.OptionsSource(InArgs.GetActionTypesSource())
-		.OnComboBoxOpening(InArgs._OnActionTypesComboBoxOpening)
+		.OnComboBoxOpening(this, &SRCSignatureActionBox::OnComboBoxOpening)
 		.OnGenerateWidget(this, &SRCSignatureActionBox::GenerateActionTypeWidget)
 		.OnSelectionChanged(this, &SRCSignatureActionBox::OnActionTypeSelected)
 		.ContentPadding(0)
@@ -92,6 +94,28 @@ void SRCSignatureActionBox::Refresh()
 		];
 }
 
+void SRCSignatureActionBox::OnComboBoxOpening()
+{
+	if (!OnRefreshActionTypes.IsBound())
+	{
+		return;
+	}
+
+	TSharedPtr<FRCSignatureTreeItemBase> Item = ItemWeak.Pin();
+	if (!Item.IsValid())
+	{
+		return;	
+	}
+
+	TSharedPtr<FRCSignatureTreeFieldItem> FieldItem = Item->MutableCast<FRCSignatureTreeFieldItem>();
+	if (!FieldItem.IsValid())
+	{
+		return;
+	}
+
+	OnRefreshActionTypes.Execute(FieldItem.ToSharedRef());
+}
+
 bool SRCSignatureActionBox::CanAddAction() const
 {
 	const bool bIsLive = LiveMode.Get(false);
diff --git a/Engine/Plugins/VirtualProduction/RemoteControl/Source/RemoteControlUI/Private/UI/Signature/Actions/SRCSignatureActionBox.h b/Engine/Plugins/VirtualProduction/RemoteControl/Source/RemoteControlUI/Private/UI/Signature/Actions/SRCSignatureActionBox.h
--- a/Engine/Plugins/VirtualProduction/RemoteControl/Source/RemoteControlUI/Private/UI/Signature/Actions/SRCSignatureActionBox.h
+++ b/Engine/Plugins/VirtualProduction/RemoteControl/Source/RemoteControlUI/Private/UI/Signature/Actions/SRCSignatureActionBox.h
@@ -5,6 +5,7 @@
 #include "Widgets/SCompoundWidget.h"
 #include "Widgets/Views/IItemsSource.h"
 
+class FRCSignatureTreeFieldItem;
 class FRCSignatureTreeItemBase;
 class SRCSignatureRow;
 class SScrollBox;
@@ -14,10 +15,12 @@ template<typename OptionType> class SComboBox;
 class SRCSignatureActionBox : public SCompoundWidget
 {
 public:
+	using FOnRefreshActionTypes = TDelegate<void(const TSharedRef<FRCSignatureTreeFieldItem>&)>;
+
 	SLATE_BEGIN_ARGS(SRCSignatureActionBox) {}
 		SLATE_ATTRIBUTE(bool, LiveMode)
 		SLATE_ITEMS_SOURCE_ARGUMENT(TSharedPtr<FRCSignatureActionType>, ActionTypesSource)
-		SLATE_EVENT(TDelegate<void()>, OnActionTypesComboBoxOpening)
+		SLATE_EVENT(FOnRefreshActionTypes, OnRefreshActionTypes)
 	SLATE_END_ARGS()
 
 	void Construct(const FArguments& InArgs, const TSharedRef<FRCSignatureTreeItemBase>& InItem, const TSharedRef<SRCSignatureRow>& InRow);
@@ -25,6 +28,8 @@ public:
 private:
 	void Refresh();
 
+	void OnComboBoxOpening();
+
 	bool CanAddAction() const;
 
 	EVisibility GetAddActionVisibility() const;
@@ -39,6 +44,8 @@ private:
 
 	TSharedPtr<SScrollBox> ActionListBox;
 
+	FOnRefreshActionTypes OnRefreshActionTypes;
+
 	TAttribute<bool> IsHovered;
 	TAttribute<bool> IsSelected;
 	TAttribute<bool> LiveMode;
