Remote Control: Added Signature renaming

diff --git a/Engine/Plugins/VirtualProduction/RemoteControl/Source/RemoteControlUI/Private/UI/SRemoteControlPanel.cpp b/Engine/Plugins/VirtualProduction/RemoteControl/Source/RemoteControlUI/Private/UI/SRemoteControlPanel.cpp
--- a/Engine/Plugins/VirtualProduction/RemoteControl/Source/RemoteControlUI/Private/UI/SRemoteControlPanel.cpp
+++ b/Engine/Plugins/VirtualProduction/RemoteControl/Source/RemoteControlUI/Private/UI/SRemoteControlPanel.cpp
@@ -2180,6 +2180,12 @@ void SRemoteControlPanel::RenameEntity_Execute() const
 		return;
 	}
 
+	if (SignaturePanel->IsListFocused())
+	{
+		SignaturePanel->EnterRenameMode();
+		return;
+	}
+
 	if (LastSelectedEntity->GetRCType() == SRCPanelTreeNode::FieldChild ||
 		LastSelectedEntity->GetRCType() == SRCPanelTreeNode::FieldGroup) // Field Child/Group does not contain entity ID, that is why it should not be processed
 	{
@@ -2201,6 +2207,11 @@ bool SRemoteControlPanel::CanRenameEntity() const
 		return true;
 	}
 
+	if (SignaturePanel->IsListFocused())
+	{
+		return true;
+	}
+
 	if (LastSelectedEntity.IsValid() && Preset.IsValid())
 	{
 		// Do not allow default group to be renamed.
diff --git a/Engine/Plugins/VirtualProduction/RemoteControl/Source/RemoteControlUI/Private/UI/Signature/Items/RCSignatureTreeItemBase.h b/Engine/Plugins/VirtualProduction/RemoteControl/Source/RemoteControlUI/Private/UI/Signature/Items/RCSignatureTreeItemBase.h
--- a/Engine/Plugins/VirtualProduction/RemoteControl/Source/RemoteControlUI/Private/UI/Signature/Items/RCSignatureTreeItemBase.h
+++ b/Engine/Plugins/VirtualProduction/RemoteControl/Source/RemoteControlUI/Private/UI/Signature/Items/RCSignatureTreeItemBase.h
@@ -75,9 +75,17 @@ public:
 		return FText::GetEmpty();
 	}
 
-	virtual bool CanEditDisplayNameText() const
+	virtual void SetRenaming(bool bInRenaming)
 	{
-		return false;
+	}
+
+	/**
+	 * Gets the Rename Delegate
+	 * @return the delegate to when renaming state changes, or nullptr if item does not support renaming
+	 */
+	virtual TMulticastDelegateRegistration<void(bool)>* GetOnRenameStateChanged() const
+	{
+		return nullptr;
 	}
 
 	virtual void SetDisplayNameText(const FText& InText)
diff --git a/Engine/Plugins/VirtualProduction/RemoteControl/Source/RemoteControlUI/Private/UI/Signature/Items/RCSignatureTreeSignatureItem.cpp b/Engine/Plugins/VirtualProduction/RemoteControl/Source/RemoteControlUI/Private/UI/Signature/Items/RCSignatureTreeSignatureItem.cpp
--- a/Engine/Plugins/VirtualProduction/RemoteControl/Source/RemoteControlUI/Private/UI/Signature/Items/RCSignatureTreeSignatureItem.cpp
+++ b/Engine/Plugins/VirtualProduction/RemoteControl/Source/RemoteControlUI/Private/UI/Signature/Items/RCSignatureTreeSignatureItem.cpp
@@ -146,9 +146,18 @@ FText FRCSignatureTreeSignatureItem::GetDisplayNameText() const
 	return FText::GetEmpty();
 }
 
-bool FRCSignatureTreeSignatureItem::CanEditDisplayNameText() const
+void FRCSignatureTreeSignatureItem::SetRenaming(bool bInRenaming)
 {
-	return true;
+	if (bRenaming != bInRenaming)
+	{
+		bRenaming = bInRenaming;
+		OnRenameStateChangedDelegate.Broadcast(bInRenaming);
+	}
+}
+
+TMulticastDelegateRegistration<void(bool)>* FRCSignatureTreeSignatureItem::GetOnRenameStateChanged() const
+{
+	return &OnRenameStateChangedDelegate;
 }
 
 void FRCSignatureTreeSignatureItem::SetDisplayNameText(const FText& InText)
diff --git a/Engine/Plugins/VirtualProduction/RemoteControl/Source/RemoteControlUI/Private/UI/Signature/Items/RCSignatureTreeSignatureItem.h b/Engine/Plugins/VirtualProduction/RemoteControl/Source/RemoteControlUI/Private/UI/Signature/Items/RCSignatureTreeSignatureItem.h
--- a/Engine/Plugins/VirtualProduction/RemoteControl/Source/RemoteControlUI/Private/UI/Signature/Items/RCSignatureTreeSignatureItem.h
+++ b/Engine/Plugins/VirtualProduction/RemoteControl/Source/RemoteControlUI/Private/UI/Signature/Items/RCSignatureTreeSignatureItem.h
@@ -40,7 +40,8 @@ protected:
 	virtual TOptional<bool> IsEnabled() const override;
 	virtual void SetEnabled(bool bInEnabled) override;
 	virtual FText GetDisplayNameText() const override;
-	virtual bool CanEditDisplayNameText() const override;
+	virtual void SetRenaming(bool bInRenaming) override;
+	virtual TMulticastDelegateRegistration<void(bool)>* GetOnRenameStateChanged() const override;
 	virtual void SetDisplayNameText(const FText& InText) override;
 	virtual FText GetDescription() const override;
 	virtual int32 RemoveFromRegistry() override;
@@ -52,4 +53,10 @@ private:
 	TWeakObjectPtr<URCSignatureRegistry> RegistryWeak;
 
 	FGuid SignatureId;
+
+	/** delegate to call when entering/exiting rename mode */
+	mutable TMulticastDelegate<void(bool)> OnRenameStateChangedDelegate;
+
+	/** true if currently in rename mode */
+	bool bRenaming = false;
 };
diff --git a/Engine/Plugins/VirtualProduction/RemoteControl/Source/RemoteControlUI/Private/UI/Signature/Label/SRCSignatureLabel.cpp b/Engine/Plugins/VirtualProduction/RemoteControl/Source/RemoteControlUI/Private/UI/Signature/Label/SRCSignatureLabel.cpp
--- a/Engine/Plugins/VirtualProduction/RemoteControl/Source/RemoteControlUI/Private/UI/Signature/Label/SRCSignatureLabel.cpp
+++ b/Engine/Plugins/VirtualProduction/RemoteControl/Source/RemoteControlUI/Private/UI/Signature/Label/SRCSignatureLabel.cpp
@@ -4,6 +4,7 @@
 #include "Styling/SlateTypes.h"
 #include "UI/Signature/Items/RCSignatureTreeItemBase.h"
 #include "UI/Signature/SRCSignatureRow.h"
+#include "UI/Signature/SRCSignatureTree.h"
 #include "Widgets/Input/SCheckBox.h"
 #include "Widgets/Text/SInlineEditableTextBlock.h"
 #include "Widgets/Views/SExpanderArrow.h"
@@ -16,6 +17,8 @@ void SRCSignatureLabel::Construct(const FArguments& InArgs
 {
 	ItemWeak = InItem;
 
+	SignatureTreeWeak = InItem->GetSignatureTree();
+
 	ChildSlot
 	[
 		SNew(SBox)
@@ -36,21 +39,83 @@ void SRCSignatureLabel::Construct(const FArguments& InArgs
 				SNew(SCheckBox)
 				.IsChecked(this, &SRCSignatureLabel::GetItemEnabledState)
 				.OnCheckStateChanged(this, &SRCSignatureLabel::SetItemEnabledState)
-				.Padding(FMargin(4.0f, 0.0f))
 				.ToolTipText(LOCTEXT("ItemEnableCheckBoxTooltip", "Determines whether the entry is enabled or not"))
 			]
 			+ SHorizontalBox::Slot()
-			.FillWidth(1.0f)
+			.FillWidth(1.f)
+			.Padding(6.f, 0.f)
+			.HAlign(HAlign_Left)
+			.VAlign(VAlign_Center)
 			[
-				SNew(SInlineEditableTextBlock)
-				.Text(this, &SRCSignatureLabel::GetSignatureDisplayName)
-				.OnTextCommitted(this, &SRCSignatureLabel::OnSignatureDisplayNameCommitted)
-				.IsSelected(InRow, &SRCSignatureRow::IsSelected)
+				CreateTextBlock(InItem, InRow)
 			]
 		]
 	];
 }
 
+SRCSignatureLabel::~SRCSignatureLabel()
+{
+	if (TSharedPtr<FRCSignatureTreeItemBase> Item = ItemWeak.Pin())
+	{
+		if (TMulticastDelegateRegistration<void(bool)>* OnRenameStateChanged = Item->GetOnRenameStateChanged())
+		{
+			OnRenameStateChanged->Remove(OnRenameStateChangedHandle);
+			OnRenameStateChangedHandle.Reset();
+		}
+		Item->SetRenaming(false);
+	}
+}
+
+TSharedRef<SWidget> SRCSignatureLabel::CreateTextBlock( const TSharedRef<FRCSignatureTreeItemBase>& InItem, const TSharedRef<SRCSignatureRow>& InRow)
+{
+	// Item can be renamed if it has a valid rename delegate
+	if (TMulticastDelegateRegistration<void(bool)>* OnRenameStateChanged = InItem->GetOnRenameStateChanged())
+	{
+		OnRenameStateChangedHandle = OnRenameStateChanged->AddSP(this, &SRCSignatureLabel::OnItemRenameStateChanged);
+
+		return SAssignNew(EditableTextBlock, SInlineEditableTextBlock)
+			.Text(this, &SRCSignatureLabel::GetSignatureDisplayName)
+			.OnTextCommitted(this, &SRCSignatureLabel::OnSignatureDisplayNameCommitted)
+			.Justification(ETextJustify::Left)
+			.IsSelected(InRow, &SRCSignatureRow::IsSelected)
+			.OnEnterEditingMode(this, &SRCSignatureLabel::SetEditMode, /*bEditMode*/true)
+			.OnExitEditingMode(this, &SRCSignatureLabel::SetEditMode, /*bEditMode*/false);
+	}
+
+	// Else resort to just showing the items display name without option to edit it
+	return SNew(STextBlock)
+		.Text(this, &SRCSignatureLabel::GetSignatureDisplayName)
+		.Justification(ETextJustify::Left);
+}
+
+void SRCSignatureLabel::OnItemRenameStateChanged(bool bInRenaming)
+{
+	// Skip if the Edit Mode is already in the desired state
+	if (bEditMode == bInRenaming || !EditableTextBlock.IsValid())
+	{
+		return;
+	}
+
+	if (bInRenaming)
+	{
+		EditableTextBlock->EnterEditingMode();
+	}
+	else
+	{
+		EditableTextBlock->ExitEditingMode();
+	}
+}
+
+void SRCSignatureLabel::SetEditMode(bool bInEditMode)
+{
+	bEditMode = bInEditMode;
+
+	if (TSharedPtr<FRCSignatureTreeItemBase> Item = ItemWeak.Pin())
+	{
+		Item->SetRenaming(bInEditMode);
+	}
+}
+
 ECheckBoxState SRCSignatureLabel::GetItemEnabledState() const
 {
 	if (!CachedCheckBoxState.IsSet())
@@ -96,6 +161,14 @@ void SRCSignatureLabel::OnSignatureDisplayNameCommitted(const FText& InText, ETe
 		Item->SetDisplayNameText(InText);
 		CachedDisplayName.Reset();
 	}
+
+	if (InCommitType == ETextCommit::OnEnter)
+	{
+		if (TSharedPtr<SRCSignatureTree> SignatureTree = SignatureTreeWeak.Pin())
+		{
+			SignatureTree->ProcessRenameQueue();
+		}
+	}
 }
 
 #undef LOCTEXT_NAMESPACE
diff --git a/Engine/Plugins/VirtualProduction/RemoteControl/Source/RemoteControlUI/Private/UI/Signature/Label/SRCSignatureLabel.h b/Engine/Plugins/VirtualProduction/RemoteControl/Source/RemoteControlUI/Private/UI/Signature/Label/SRCSignatureLabel.h
--- a/Engine/Plugins/VirtualProduction/RemoteControl/Source/RemoteControlUI/Private/UI/Signature/Label/SRCSignatureLabel.h
+++ b/Engine/Plugins/VirtualProduction/RemoteControl/Source/RemoteControlUI/Private/UI/Signature/Label/SRCSignatureLabel.h
@@ -5,7 +5,9 @@
 #include "Widgets/SCompoundWidget.h"
 
 class FRCSignatureTreeItemBase;
+class SInlineEditableTextBlock;
 class SRCSignatureRow;
+class SRCSignatureTree;
 enum class ECheckBoxState : uint8;
 
 class SRCSignatureLabel : public SCompoundWidget
@@ -18,7 +20,16 @@ public:
 		, const TSharedRef<FRCSignatureTreeItemBase>& InItem
 		, const TSharedRef<SRCSignatureRow>& InRow);
 
+	virtual ~SRCSignatureLabel() override;
+
 private:
+	TSharedRef<SWidget> CreateTextBlock(const TSharedRef<FRCSignatureTreeItemBase>& InItem, const TSharedRef<SRCSignatureRow>& InRow);
+
+	/** Called when the Item rename state changes */
+	void OnItemRenameStateChanged(bool bInRenaming);
+
+	void SetEditMode(bool bInEditMode);
+
 	ECheckBoxState GetItemEnabledState() const;
 
 	void SetItemEnabledState(ECheckBoxState InState);
@@ -29,7 +40,15 @@ private:
 
 	TWeakPtr<FRCSignatureTreeItemBase> ItemWeak;
 
+	TWeakPtr<SRCSignatureTree> SignatureTreeWeak;
+
+	TSharedPtr<SInlineEditableTextBlock> EditableTextBlock;
+
 	mutable TOptional<ECheckBoxState> CachedCheckBoxState;
 
 	mutable TOptional<FText> CachedDisplayName;
+
+	FDelegateHandle OnRenameStateChangedHandle;
+
+	bool bEditMode = false;
 };
diff --git a/Engine/Plugins/VirtualProduction/RemoteControl/Source/RemoteControlUI/Private/UI/Signature/SRCSignaturePanel.cpp b/Engine/Plugins/VirtualProduction/RemoteControl/Source/RemoteControlUI/Private/UI/Signature/SRCSignaturePanel.cpp
--- a/Engine/Plugins/VirtualProduction/RemoteControl/Source/RemoteControlUI/Private/UI/Signature/SRCSignaturePanel.cpp
+++ b/Engine/Plugins/VirtualProduction/RemoteControl/Source/RemoteControlUI/Private/UI/Signature/SRCSignaturePanel.cpp
@@ -184,6 +184,14 @@ bool SRCSignaturePanel::IsListFocused() const
 	return SignatureTreeView.IsValid() && SignatureTreeView->IsListFocused();
 }
 
+void SRCSignaturePanel::EnterRenameMode()
+{
+	if (SignatureTreeView.IsValid())
+	{
+		return SignatureTreeView->EnterRenameMode();
+	}
+}
+
 TArray<TSharedPtr<FRCLogicModeBase>> SRCSignaturePanel::GetSelectedLogicItems() const
 {
 	if (SignatureTreeView.IsValid())
diff --git a/Engine/Plugins/VirtualProduction/RemoteControl/Source/RemoteControlUI/Private/UI/Signature/SRCSignaturePanel.h b/Engine/Plugins/VirtualProduction/RemoteControl/Source/RemoteControlUI/Private/UI/Signature/SRCSignaturePanel.h
--- a/Engine/Plugins/VirtualProduction/RemoteControl/Source/RemoteControlUI/Private/UI/Signature/SRCSignaturePanel.h
+++ b/Engine/Plugins/VirtualProduction/RemoteControl/Source/RemoteControlUI/Private/UI/Signature/SRCSignaturePanel.h
@@ -26,6 +26,8 @@ public:
 	/** Whether this widget currently has focus */
 	bool IsListFocused() const;
 
+	void EnterRenameMode();
+
 	//~ Begin SRCLogicPanelBase
 	virtual TArray<TSharedPtr<FRCLogicModeBase>> GetSelectedLogicItems() const override;
 	virtual FReply RequestDeleteSelectedItem() override;
diff --git a/Engine/Plugins/VirtualProduction/RemoteControl/Source/RemoteControlUI/Private/UI/Signature/SRCSignatureTree.cpp b/Engine/Plugins/VirtualProduction/RemoteControl/Source/RemoteControlUI/Private/UI/Signature/SRCSignatureTree.cpp
--- a/Engine/Plugins/VirtualProduction/RemoteControl/Source/RemoteControlUI/Private/UI/Signature/SRCSignatureTree.cpp
+++ b/Engine/Plugins/VirtualProduction/RemoteControl/Source/RemoteControlUI/Private/UI/Signature/SRCSignatureTree.cpp
@@ -1,6 +1,7 @@
 // Copyright Epic Games, Inc. All Rights Reserved.
 
 #include "SRCSignatureTree.h"
+#include "Algo/Reverse.h"
 #include "IRCSignatureColumn.h"
 #include "Items/RCSignatureTreeItemBase.h"
 #include "Items/RCSignatureTreeRootItem.h"
@@ -97,6 +98,57 @@ void SRCSignatureTree::Refresh()
 	SignatureTreeView->RequestTreeRefresh();
 }
 
+void SRCSignatureTree::EnterRenameMode()
+{
+	TArray<TSharedPtr<FRCSignatureTreeItemBase>> SelectedItems = GetSelectedItems();
+
+	// Remove all the items that can't be renamed
+	SelectedItems.RemoveAll([](const TSharedPtr<FRCSignatureTreeItemBase>& InItem)
+		{
+			return !InItem->GetOnRenameStateChanged();
+		});
+
+	RenameQueue.Reset();
+	RenameQueue.Append(MoveTemp(SelectedItems));
+
+	// Reverse as items will be removed from the end (pop) to avoid shifting items every dequeue
+	Algo::Reverse(RenameQueue);
+
+	ProcessRenameQueue();
+}
+
+void SRCSignatureTree::ProcessRenameQueue()
+{
+	if (TSharedPtr<FRCSignatureTreeItemBase> CurrentItemRenaming = CurrentItemRenamingWeak.Pin())
+	{
+		CurrentItemRenaming->SetRenaming(false);
+	}
+
+	CurrentItemRenamingWeak.Reset();
+
+	if (RenameQueue.IsEmpty())
+	{
+		return;
+	}
+
+	// Dequeue until a valid item is gotten
+	// Most likely it will be the first item
+	TSharedPtr<FRCSignatureTreeItemBase> ItemToRename;
+	do
+	{
+		ItemToRename = RenameQueue.Pop().Pin();
+	}
+	while (!RenameQueue.IsEmpty() && !ItemToRename.IsValid());
+
+	if (!ItemToRename.IsValid())
+	{
+		return;
+	}
+
+	CurrentItemRenamingWeak = ItemToRename;
+	ItemToRename->SetRenaming(true);
+}
+
 TArray<TSharedPtr<FRCSignatureTreeItemBase>> SRCSignatureTree::GetSelectedItems() const
 {
 	return RootItem->GetSelection()->GetSelectedItems();
diff --git a/Engine/Plugins/VirtualProduction/RemoteControl/Source/RemoteControlUI/Private/UI/Signature/SRCSignatureTree.h b/Engine/Plugins/VirtualProduction/RemoteControl/Source/RemoteControlUI/Private/UI/Signature/SRCSignatureTree.h
--- a/Engine/Plugins/VirtualProduction/RemoteControl/Source/RemoteControlUI/Private/UI/Signature/SRCSignatureTree.h
+++ b/Engine/Plugins/VirtualProduction/RemoteControl/Source/RemoteControlUI/Private/UI/Signature/SRCSignatureTree.h
@@ -29,6 +29,10 @@ public:
 
 	void Refresh();
 
+	void EnterRenameMode();
+
+	void ProcessRenameQueue();
+
 	TArray<TSharedPtr<FRCSignatureTreeItemBase>> GetSelectedItems() const;
 
 	//~ Begin SRCLogicPanelListBase
@@ -67,5 +71,11 @@ private:
 
 	TWeakPtr<SRCSignaturePanel> SignaturePanelWeak;
 
+	/** Items that are pending to be renamed */
+	TArray<TWeakPtr<FRCSignatureTreeItemBase>> RenameQueue;
+
+	/** Current item being renamed */
+	TWeakPtr<FRCSignatureTreeItemBase> CurrentItemRenamingWeak;
+
 	bool bSyncingSelection = false;
 };
