Remote Control: Add PostEditChange func to Signature Actions

diff --git a/Engine/Plugins/VirtualProduction/RemoteControl/Source/RemoteControl/Private/RCSignatureActionInstance.cpp b/Engine/Plugins/VirtualProduction/RemoteControl/Source/RemoteControl/Private/RCSignatureActionInstance.cpp
--- a/Engine/Plugins/VirtualProduction/RemoteControl/Source/RemoteControl/Private/RCSignatureActionInstance.cpp
+++ b/Engine/Plugins/VirtualProduction/RemoteControl/Source/RemoteControl/Private/RCSignatureActionInstance.cpp
@@ -33,3 +33,13 @@ bool FRCSignatureActionInstance::Execute(const FRCSignatureActionContext& InCont
 	}
 	return false;
 }
+
+#if WITH_EDITOR
+void FRCSignatureActionInstance::PostEditChange(const FPropertyChangedEvent& InPropertyChangedEvent, FEditPropertyChain* InPropertyThatChanged)
+{
+	if (FRCSignatureAction* Action = ActionInstance.GetMutablePtr())
+	{
+		Action->PostEditChange(InPropertyChangedEvent, InPropertyThatChanged);
+	}
+}
+#endif
diff --git a/Engine/Plugins/VirtualProduction/RemoteControl/Source/RemoteControl/Public/RCSignatureAction.h b/Engine/Plugins/VirtualProduction/RemoteControl/Source/RemoteControl/Public/RCSignatureAction.h
--- a/Engine/Plugins/VirtualProduction/RemoteControl/Source/RemoteControl/Public/RCSignatureAction.h
+++ b/Engine/Plugins/VirtualProduction/RemoteControl/Source/RemoteControl/Public/RCSignatureAction.h
@@ -11,6 +11,11 @@ struct FRCSignature;
 struct FRCSignatureField;
 struct FRemoteControlProperty;
 
+#if WITH_EDITOR
+class FEditPropertyChain;
+struct FPropertyChangedEvent;
+#endif
+
 /** The context for a Signature Action to execute */
 struct FRCSignatureActionContext
 {
@@ -85,6 +90,11 @@ struct FRCSignatureAction
 	}
 
 #if WITH_EDITOR
+	/** Called whenever there's a change affecting the action in the Details Panel */
+	virtual void PostEditChange(const FPropertyChangedEvent& InPropertyChangedEvent, FEditPropertyChain* InPropertyThatChanged)
+	{
+	}
+
 	/** Retrieves the icon to use to represent this Signature Action */
 	virtual FRCSignatureActionIcon GetIcon() const
 	{
diff --git a/Engine/Plugins/VirtualProduction/RemoteControl/Source/RemoteControl/Public/RCSignatureActionInstance.h b/Engine/Plugins/VirtualProduction/RemoteControl/Source/RemoteControl/Public/RCSignatureActionInstance.h
--- a/Engine/Plugins/VirtualProduction/RemoteControl/Source/RemoteControl/Public/RCSignatureActionInstance.h
+++ b/Engine/Plugins/VirtualProduction/RemoteControl/Source/RemoteControl/Public/RCSignatureActionInstance.h
@@ -10,6 +10,11 @@ struct FRCSignatureAction;
 struct FRCSignatureActionContext;
 struct FRCSignatureField;
 
+#if WITH_EDITOR
+class FEditPropertyChain;
+struct FPropertyChangedEvent;
+#endif
+
 /** Struct containing an Action Instance and handling its Execution */
 USTRUCT()
 struct FRCSignatureActionInstance
@@ -26,6 +31,10 @@ struct FRCSignatureActionInstance
 
 	bool Execute(const FRCSignatureActionContext& InContext) const;
 
+#if WITH_EDITOR
+	REMOTECONTROL_API void PostEditChange(const FPropertyChangedEvent& InPropertyChangedEvent, FEditPropertyChain* InPropertyThatChanged);
+#endif
+
 private:
 	UPROPERTY()
 	TInstancedStruct<FRCSignatureAction> ActionInstance;
diff --git a/Engine/Plugins/VirtualProduction/RemoteControl/Source/RemoteControlUI/Private/UI/Signature/Details/SRCSignatureDetails.cpp b/Engine/Plugins/VirtualProduction/RemoteControl/Source/RemoteControlUI/Private/UI/Signature/Details/SRCSignatureDetails.cpp
--- a/Engine/Plugins/VirtualProduction/RemoteControl/Source/RemoteControlUI/Private/UI/Signature/Details/SRCSignatureDetails.cpp
+++ b/Engine/Plugins/VirtualProduction/RemoteControl/Source/RemoteControlUI/Private/UI/Signature/Details/SRCSignatureDetails.cpp
@@ -47,22 +47,30 @@ SRCSignatureDetails::~SRCSignatureDetails()
 void SRCSignatureDetails::Refresh()
 {
 	TSharedRef<FStructOnScopeStructureDataProvider> StructProvider = MakeShared<FStructOnScopeStructureDataProvider>();
-	StructProvider->SetStructData(GatherStructOnScopes());
+
+	TArray<TSharedPtr<FStructOnScope>> StructOnScopes;
+	GatherStructOnScopes(StructOnScopes, ViewedItems);
+
+	StructProvider->SetStructData(StructOnScopes);
+
 	StructDetailsView->SetStructureProvider(StructProvider);
 }
 
-TArray<TSharedPtr<FStructOnScope>> SRCSignatureDetails::GatherStructOnScopes() const
+void SRCSignatureDetails::GatherStructOnScopes(TArray<TSharedPtr<FStructOnScope>>& OutStructOnScopes, TArray<TWeakPtr<FRCSignatureTreeItemBase>>& OutItems) const
 {
-	TArray<TSharedPtr<FStructOnScope>> StructOnScopes;
+	OutStructOnScopes.Reset();
+	OutItems.Reset();
 
 	TSharedPtr<FRCSignatureTreeItemSelection> Selection = SelectionWeak.Pin();
 	if (!Selection.IsValid())
 	{
-		return StructOnScopes;
+		return;
 	}
 
 	TConstArrayView<TWeakPtr<FRCSignatureTreeItemBase>> SelectedItems = Selection->GetSelectedItemsView();
-	StructOnScopes.Reserve(SelectedItems.Num());
+
+	OutStructOnScopes.Reserve(SelectedItems.Num());
+	OutItems.Reserve(SelectedItems.Num());
 
 	for (const TWeakPtr<FRCSignatureTreeItemBase>& SelectedItemWeak : SelectedItems)
 	{
@@ -78,10 +86,9 @@ TArray<TSharedPtr<FStructOnScope>> SRCSignatureDetails::GatherStructOnScopes() c
 			continue;
 		}
 
-		StructOnScopes.Add(MoveTemp(StructOnScope));
+		OutStructOnScopes.Add(MoveTemp(StructOnScope));
+		OutItems.Add(SelectedItem);
 	}
-
-	return StructOnScopes;
 }
 
 void SRCSignatureDetails::OnFinishedChangingProperties(const FPropertyChangedEvent& InChangeEvent)
@@ -94,11 +101,24 @@ void SRCSignatureDetails::NotifyPreChange(FEditPropertyChain* InPropertyAboutToC
 {
 	if (URCSignatureRegistry* SignatureRegistry = SignatureRegistryWeak.Get())
 	{
-		// Begin Transaction
-		ensureAlways(!CurrentTransaction.IsValid());
-		CurrentTransaction = MakeShared<FScopedTransaction>(LOCTEXT("EditSignature", "Edit Signature"));
+		// Begin Transaction, if not already in one
+		if (!CurrentTransaction.IsValid())
+		{
+			CurrentTransaction = MakeShared<FScopedTransaction>(LOCTEXT("EditSignature", "Edit Signature"));
+		}
 		SignatureRegistry->Modify();
 	}
 }
 
-#undef LOCTEXT_NAMESPACE 
+void SRCSignatureDetails::NotifyPostChange(const FPropertyChangedEvent& InPropertyChangedEvent, FEditPropertyChain* InPropertyThatChanged)
+{
+	for (TWeakPtr<FRCSignatureTreeItemBase> ViewedItemWeak : ViewedItems)
+	{
+		if (TSharedPtr<FRCSignatureTreeItemBase> ViewedItem = ViewedItemWeak.Pin())
+		{
+			ViewedItem->NotifyPostChange(InPropertyChangedEvent, InPropertyThatChanged);
+		}
+	}
+}
+
+#undef LOCTEXT_NAMESPACE
diff --git a/Engine/Plugins/VirtualProduction/RemoteControl/Source/RemoteControlUI/Private/UI/Signature/Details/SRCSignatureDetails.h b/Engine/Plugins/VirtualProduction/RemoteControl/Source/RemoteControlUI/Private/UI/Signature/Details/SRCSignatureDetails.h
--- a/Engine/Plugins/VirtualProduction/RemoteControl/Source/RemoteControlUI/Private/UI/Signature/Details/SRCSignatureDetails.h
+++ b/Engine/Plugins/VirtualProduction/RemoteControl/Source/RemoteControlUI/Private/UI/Signature/Details/SRCSignatureDetails.h
@@ -7,6 +7,7 @@
 #include "UObject/WeakObjectPtr.h"
 #include "Widgets/SCompoundWidget.h"
 
+class FRCSignatureTreeItemBase;
 class FRCSignatureTreeItemSelection;
 class FScopedTransaction;
 class FStructOnScope;
@@ -27,12 +28,13 @@ public:
 	void Refresh();
 
 private:
-	TArray<TSharedPtr<FStructOnScope>> GatherStructOnScopes() const;
+	void GatherStructOnScopes(TArray<TSharedPtr<FStructOnScope>>& OutStructOnScopes, TArray<TWeakPtr<FRCSignatureTreeItemBase>>& OutItems) const;
 
 	void OnFinishedChangingProperties(const FPropertyChangedEvent& InChangeEvent);
 
 	//~ Begin FNotifyHook
 	virtual void NotifyPreChange(FEditPropertyChain* InPropertyAboutToChange) override;
+	virtual void NotifyPostChange(const FPropertyChangedEvent& InPropertyChangedEvent, FEditPropertyChain* InPropertyThatChanged) override;
 	//~ End FNotifyHook
 
 	TSharedPtr<IStructureDetailsView> StructDetailsView;
@@ -42,4 +44,7 @@ private:
 	TWeakObjectPtr<URCSignatureRegistry> SignatureRegistryWeak;
 
 	TWeakPtr<FRCSignatureTreeItemSelection> SelectionWeak;
+
+	/** Currently viewed items in the Details Panel */
+	TArray<TWeakPtr<FRCSignatureTreeItemBase>> ViewedItems;
 };
diff --git a/Engine/Plugins/VirtualProduction/RemoteControl/Source/RemoteControlUI/Private/UI/Signature/Items/RCSignatureTreeActionItem.cpp b/Engine/Plugins/VirtualProduction/RemoteControl/Source/RemoteControlUI/Private/UI/Signature/Items/RCSignatureTreeActionItem.cpp
--- a/Engine/Plugins/VirtualProduction/RemoteControl/Source/RemoteControlUI/Private/UI/Signature/Items/RCSignatureTreeActionItem.cpp
+++ b/Engine/Plugins/VirtualProduction/RemoteControl/Source/RemoteControlUI/Private/UI/Signature/Items/RCSignatureTreeActionItem.cpp
@@ -108,6 +108,14 @@ TSharedPtr<FStructOnScope> FRCSignatureTreeActionItem::MakeSelectionStruct()
 	return nullptr;
 }
 
+void FRCSignatureTreeActionItem::NotifyPostChange(const FPropertyChangedEvent& InPropertyChangedEvent, FEditPropertyChain* InPropertyThatChanged)
+{
+	if (FRCSignatureActionInstance* ActionInstance = FindActionInstanceMutable())
+	{
+		ActionInstance->PostEditChange(InPropertyChangedEvent, InPropertyThatChanged);
+	}
+}
+
 TSharedPtr<FRCSignatureTreeFieldItem> FRCSignatureTreeActionItem::GetParentFieldItem() const
 {
 	if (TSharedPtr<FRCSignatureTreeItemBase> Parent = GetParent())
diff --git a/Engine/Plugins/VirtualProduction/RemoteControl/Source/RemoteControlUI/Private/UI/Signature/Items/RCSignatureTreeActionItem.h b/Engine/Plugins/VirtualProduction/RemoteControl/Source/RemoteControlUI/Private/UI/Signature/Items/RCSignatureTreeActionItem.h
--- a/Engine/Plugins/VirtualProduction/RemoteControl/Source/RemoteControlUI/Private/UI/Signature/Items/RCSignatureTreeActionItem.h
+++ b/Engine/Plugins/VirtualProduction/RemoteControl/Source/RemoteControlUI/Private/UI/Signature/Items/RCSignatureTreeActionItem.h
@@ -32,6 +32,7 @@ protected:
 	virtual FText GetDisplayNameText() const override;
 	virtual FText GetDescription() const override;
 	virtual TSharedPtr<FStructOnScope> MakeSelectionStruct() override;
+	virtual void NotifyPostChange(const FPropertyChangedEvent& InPropertyChangedEvent, FEditPropertyChain* InPropertyThatChanged) override;
 	virtual ERCSignatureTreeItemType GetItemType() const override { return StaticItemType; }
 	//~ End FRCSignatureTreeItemBase
 
diff --git a/Engine/Plugins/VirtualProduction/RemoteControl/Source/RemoteControlUI/Private/UI/Signature/Items/RCSignatureTreeItemBase.h b/Engine/Plugins/VirtualProduction/RemoteControl/Source/RemoteControlUI/Private/UI/Signature/Items/RCSignatureTreeItemBase.h
--- a/Engine/Plugins/VirtualProduction/RemoteControl/Source/RemoteControlUI/Private/UI/Signature/Items/RCSignatureTreeItemBase.h
+++ b/Engine/Plugins/VirtualProduction/RemoteControl/Source/RemoteControlUI/Private/UI/Signature/Items/RCSignatureTreeItemBase.h
@@ -4,9 +4,11 @@
 
 #include "UI/BaseLogicUI/RCLogicModeBase.h"
 
+class FEditPropertyChain;
 class FRCSignatureTreeItemSelection;
 class FStructOnScope;
 class SRCSignatureTree;
+struct FPropertyChangedEvent;
 
 enum class ERCSignatureTreeItemViewFlags : uint8
 {
@@ -112,6 +114,11 @@ public:
 		return nullptr;
 	}
 
+	/** Called when the Selection Struct created for this Item is changed via the Details Panel */
+	virtual void NotifyPostChange(const FPropertyChangedEvent& InPropertyChangedEvent, FEditPropertyChain* InPropertyThatChanged)
+	{
+	}
+
 	virtual ERCSignatureTreeItemType GetItemType() const
 	{
 		return ERCSignatureTreeItemType::Undefined;
