[Remote Control] Reverting fix for UE-215856 as it breaks the stage app.

diff --git a/Engine/Plugins/VirtualProduction/RemoteControl/Source/WebRemoteControl/Private/WebRemoteControl.cpp b/Engine/Plugins/VirtualProduction/RemoteControl/Source/WebRemoteControl/Private/WebRemoteControl.cpp
--- a/Engine/Plugins/VirtualProduction/RemoteControl/Source/WebRemoteControl/Private/WebRemoteControl.cpp
+++ b/Engine/Plugins/VirtualProduction/RemoteControl/Source/WebRemoteControl/Private/WebRemoteControl.cpp
@@ -2376,7 +2376,6 @@ void FWebRemoteControlModule::HandleWebSocketHttpMessage(const FRemoteControlWeb
 		Wrapper.Passphrase = WebSocketMessage.Header[WebRemoteControlInternalUtils::PassphraseHeader][0];
 	}
 
-	Wrapper.RequestId = WebSocketMessage.MessageId;
 
 	LogRequestExternally(Wrapper.RequestId, TEXT("UE Received"));
 	
