[WebRemoteControl] Http Route handler: Fix error paths of Preset Set Controller.

The rules for http route handlers:
- OnComplete callback must be called only once. some utils like ValidateContentType and DeserializeRequest already call OnComplete with an error message.
- If OnComplete callback is invoked, the handler must return true, even if the request didn't succeed (had some error), it is still considered "handled".

diff --git a/Engine/Plugins/VirtualProduction/RemoteControl/Source/WebRemoteControl/Private/WebRemoteControl.cpp b/Engine/Plugins/VirtualProduction/RemoteControl/Source/WebRemoteControl/Private/WebRemoteControl.cpp
--- a/Engine/Plugins/VirtualProduction/RemoteControl/Source/WebRemoteControl/Private/WebRemoteControl.cpp
+++ b/Engine/Plugins/VirtualProduction/RemoteControl/Source/WebRemoteControl/Private/WebRemoteControl.cpp
@@ -2214,20 +2214,14 @@ bool FWebRemoteControlModule::HandlePresetSetControllerRoute(const FHttpServerRe
 	// 2. Validate Content Type
 	if (!WebRemoteControlInternalUtils::ValidateContentType(Request, TEXT("application/json"), OnComplete))
 	{
-		WebRemoteControlInternalUtils::CreateUTF8ErrorMessage(TEXT("Expected content type to be application/json"), Response->Body);
-		OnComplete(MoveTemp(Response));
-
-		return false;
+		return true;
 	}
 
 	// 3. Deserialize Json POST data into request struct
 	FRCPresetSetControllerRequest SetControllerRequest;
 	if (!WebRemoteControlInternalUtils::DeserializeRequest(Request, &OnComplete, SetControllerRequest))
 	{
-		WebRemoteControlInternalUtils::CreateUTF8ErrorMessage(TEXT("Unable to process JSON body. Expected format: type to be application/json"), Response->Body);
-		OnComplete(MoveTemp(Response));
-
-		return false;
+		return true;
 	}
 
 	// 4. Resolve url arguments
@@ -2242,8 +2236,7 @@ bool FWebRemoteControlModule::HandlePresetSetControllerRoute(const FHttpServerRe
 		Response->Code = EHttpServerResponseCodes::NotFound;
 		WebRemoteControlInternalUtils::CreateUTF8ErrorMessage(TEXT("Unable to resolve the preset."), Response->Body);
 		OnComplete(MoveTemp(Response));
-
-		return false;
+		return true;
 	}
 
 	// 6. Acquire Controller
@@ -2253,8 +2246,7 @@ bool FWebRemoteControlModule::HandlePresetSetControllerRoute(const FHttpServerRe
 		Response->Code = EHttpServerResponseCodes::NotFound;
 		WebRemoteControlInternalUtils::CreateUTF8ErrorMessage(TEXT("Unable to resolve the controller input."), Response->Body);
 		OnComplete(MoveTemp(Response));
-
-		return false;
+		return true;
 	}
 
 	// 7. Reformat Payload to represent our internal structure
@@ -2276,13 +2268,12 @@ bool FWebRemoteControlModule::HandlePresetSetControllerRoute(const FHttpServerRe
 	}
 	else
 	{
+		Response->Code = EHttpServerResponseCodes::ServerError;
 		WebRemoteControlInternalUtils::CreateUTF8ErrorMessage(FString::Printf(TEXT("Error while trying to set controller %s."), *Args.ControllerName), Response->Body);
 	}
 
 	OnComplete(MoveTemp(Response));
-
 	return true;
-
 }
 
 bool FWebRemoteControlModule::HandlePresetGetControllerRoute(const FHttpServerRequest& Request, const FHttpResultCallback& OnComplete)
