Remote Control Signature: Added Post Load functionality for Signature Actions

diff --git a/Engine/Plugins/VirtualProduction/RemoteControl/Source/RemoteControl/Private/RCSignature.cpp b/Engine/Plugins/VirtualProduction/RemoteControl/Source/RemoteControl/Private/RCSignature.cpp
--- a/Engine/Plugins/VirtualProduction/RemoteControl/Source/RemoteControl/Private/RCSignature.cpp
+++ b/Engine/Plugins/VirtualProduction/RemoteControl/Source/RemoteControl/Private/RCSignature.cpp
@@ -91,6 +91,22 @@ FRCSignatureField FRCSignatureField::CreateField(const FRCFieldPathInfo& InField
 	return Field;
 }
 
+void FRCSignatureField::PostLoad()
+{
+	for (FRCSignatureActionInstance& Action : Actions)
+	{
+		Action.PostLoad(*this);
+	}
+}
+
+void FRCSignature::PostLoad()
+{
+	for (FRCSignatureField& Field : Fields)
+	{
+		Field.PostLoad();
+	}
+}
+
 int32 FRCSignature::AddFields(TConstArrayView<FRCSignatureField> InFields)
 {
 	const int32 PreviousNum = Fields.Num();
diff --git a/Engine/Plugins/VirtualProduction/RemoteControl/Source/RemoteControl/Private/RCSignatureActionInstance.cpp b/Engine/Plugins/VirtualProduction/RemoteControl/Source/RemoteControl/Private/RCSignatureActionInstance.cpp
--- a/Engine/Plugins/VirtualProduction/RemoteControl/Source/RemoteControl/Private/RCSignatureActionInstance.cpp
+++ b/Engine/Plugins/VirtualProduction/RemoteControl/Source/RemoteControl/Private/RCSignatureActionInstance.cpp
@@ -25,6 +25,14 @@ TSharedRef<FStructOnScope> FRCSignatureActionInstance::MakeStructOnScope()
 	return MakeShared<FStructOnScope>(ActionInstance.GetScriptStruct(), ActionInstance.GetMutableMemory());
 }
 
+void FRCSignatureActionInstance::PostLoad(const FRCSignatureField& InFieldOwner)
+{
+	if (FRCSignatureAction* Action = ActionInstance.GetMutablePtr())
+	{
+		return Action->Initialize(InFieldOwner);
+	}
+}
+
 bool FRCSignatureActionInstance::Execute(const FRCSignatureActionContext& InContext) const
 {
 	if (const FRCSignatureAction* Action = ActionInstance.GetPtr())
diff --git a/Engine/Plugins/VirtualProduction/RemoteControl/Source/RemoteControl/Private/RCSignatureRegistry.cpp b/Engine/Plugins/VirtualProduction/RemoteControl/Source/RemoteControl/Private/RCSignatureRegistry.cpp
--- a/Engine/Plugins/VirtualProduction/RemoteControl/Source/RemoteControl/Private/RCSignatureRegistry.cpp
+++ b/Engine/Plugins/VirtualProduction/RemoteControl/Source/RemoteControl/Private/RCSignatureRegistry.cpp
@@ -35,4 +35,14 @@ void URCSignatureRegistry::EmptySignatures()
 	Signatures.Empty();
 }
 
+void URCSignatureRegistry::PostLoad()
+{
+	Super::PostLoad();
+
+	for (FRCSignature& Signature : Signatures)
+	{
+		Signature.PostLoad();
+	}
+}
+
 #undef LOCTEXT_NAMESPACE
diff --git a/Engine/Plugins/VirtualProduction/RemoteControl/Source/RemoteControl/Public/RCSignature.h b/Engine/Plugins/VirtualProduction/RemoteControl/Source/RemoteControl/Public/RCSignature.h
--- a/Engine/Plugins/VirtualProduction/RemoteControl/Source/RemoteControl/Public/RCSignature.h
+++ b/Engine/Plugins/VirtualProduction/RemoteControl/Source/RemoteControl/Public/RCSignature.h
@@ -20,6 +20,8 @@ struct FRCSignatureField
 
 	REMOTECONTROL_API static FRCSignatureField CreateField(const FRCFieldPathInfo& InFieldPathInfo, const UObject* InOwnerObject, const FProperty* InProperty);
 
+	void PostLoad();
+
 	bool operator==(const FRCSignatureField& InOtherField) const
 	{
 		return FieldPath == InOtherField.FieldPath
@@ -62,6 +64,8 @@ struct FRCSignature
 		return Id == InSignatureId;
 	}
 
+	void PostLoad();
+
 	/**
 	 * Adds the given fields to this signature and ensuring no repeated field is present
 	 * @return the number of new fields added
diff --git a/Engine/Plugins/VirtualProduction/RemoteControl/Source/RemoteControl/Public/RCSignatureAction.h b/Engine/Plugins/VirtualProduction/RemoteControl/Source/RemoteControl/Public/RCSignatureAction.h
--- a/Engine/Plugins/VirtualProduction/RemoteControl/Source/RemoteControl/Public/RCSignatureAction.h
+++ b/Engine/Plugins/VirtualProduction/RemoteControl/Source/RemoteControl/Public/RCSignatureAction.h
@@ -62,7 +62,7 @@ struct FRCSignatureAction
 	virtual ~FRCSignatureAction() = default;
 
 	/**
-	 * Called only once when the Signature Action is first added to the action list
+	 * Called when the Signature Action is first added to the action list or loaded
 	 * @param InField the owner of the action list
 	 */
 	virtual void Initialize(const FRCSignatureField& InField)
diff --git a/Engine/Plugins/VirtualProduction/RemoteControl/Source/RemoteControl/Public/RCSignatureActionInstance.h b/Engine/Plugins/VirtualProduction/RemoteControl/Source/RemoteControl/Public/RCSignatureActionInstance.h
--- a/Engine/Plugins/VirtualProduction/RemoteControl/Source/RemoteControl/Public/RCSignatureActionInstance.h
+++ b/Engine/Plugins/VirtualProduction/RemoteControl/Source/RemoteControl/Public/RCSignatureActionInstance.h
@@ -29,6 +29,8 @@ struct FRCSignatureActionInstance
 
 	REMOTECONTROL_API TSharedRef<FStructOnScope> MakeStructOnScope();
 
+	void PostLoad(const FRCSignatureField& InFieldOwner);
+
 	bool Execute(const FRCSignatureActionContext& InContext) const;
 
 #if WITH_EDITOR
diff --git a/Engine/Plugins/VirtualProduction/RemoteControl/Source/RemoteControl/Public/RCSignatureRegistry.h b/Engine/Plugins/VirtualProduction/RemoteControl/Source/RemoteControl/Public/RCSignatureRegistry.h
--- a/Engine/Plugins/VirtualProduction/RemoteControl/Source/RemoteControl/Public/RCSignatureRegistry.h
+++ b/Engine/Plugins/VirtualProduction/RemoteControl/Source/RemoteControl/Public/RCSignatureRegistry.h
@@ -44,6 +44,10 @@ public:
 	/** Removes all the Signature in this registry */
 	REMOTECONTROL_API void EmptySignatures();
 
+	//~ Begin UObject
+	virtual void PostLoad() override;
+	//~ End UObject
+
 private:
 	/** Holds all the Signatures in this Registry */
 	UPROPERTY()
