RemoteControl: Fix scrollbar is now visible for protocol bindings when a preset with existing protocol bindings is opened, allowing too edit the binding as expected

diff --git a/Engine/Plugins/VirtualProduction/RemoteControl/Source/RemoteControlProtocolWidgets/Private/Widgets/SRCProtocolBindingList.cpp b/Engine/Plugins/VirtualProduction/RemoteControl/Source/RemoteControlProtocolWidgets/Private/Widgets/SRCProtocolBindingList.cpp
--- a/Engine/Plugins/VirtualProduction/RemoteControl/Source/RemoteControlProtocolWidgets/Private/Widgets/SRCProtocolBindingList.cpp
+++ b/Engine/Plugins/VirtualProduction/RemoteControl/Source/RemoteControlProtocolWidgets/Private/Widgets/SRCProtocolBindingList.cpp
@@ -17,6 +17,7 @@
 #include "Widgets/Images/SImage.h"
 #include "Widgets/Input/SButton.h"
 #include "Widgets/Views/SListView.h"
+#include "Widgets/Layout/SScrollBox.h"
 #include "Widgets/Views/STableRow.h"
 #include "Widgets/Views/STableViewBase.h"
 
@@ -59,14 +60,10 @@ void SRCProtocolBindingList::Construct(const FArguments& InArgs, TSharedRef<FPro
 		Refresh();
 	});
 
-	TSharedRef<SScrollBar> ExternalScrollBar = SNew(SScrollBar);
-	ExternalScrollBar->SetVisibility(TAttribute<EVisibility>(this, &SRCProtocolBindingList::GetScrollBarVisibility));
-
 	SAssignNew(BindingList, SListView<TSharedPtr<IRCTreeNodeViewModel>>)
 	.ListItemsSource(&FilteredBindings)
 	.OnGenerateRow(this, &SRCProtocolBindingList::ConstructBindingWidget)
-	.SelectionMode(ESelectionMode::None)
-	.ExternalScrollbar(ExternalScrollBar);
+	.SelectionMode(ESelectionMode::None);
 
 	ChildSlot
 	[
@@ -155,8 +152,9 @@ void SRCProtocolBindingList::Construct(const FArguments& InArgs, TSharedRef<FPro
 		.FillHeight(1.0f)
 		.Padding(0)
 		[
-			SNew(SHorizontalBox)
-			+ SHorizontalBox::Slot()
+			SNew(SScrollBox)
+
+			+ SScrollBox::Slot()
 			[
 				SNew(SBorder)
 				.BorderImage(FAppStyle::GetBrush("ToolPanel.DarkGroupBorder"))
@@ -165,15 +163,6 @@ void SRCProtocolBindingList::Construct(const FArguments& InArgs, TSharedRef<FPro
 					BindingList.ToSharedRef()				
 				]				
 			]
-			+ SHorizontalBox::Slot()
-			.AutoWidth()
-			[
-				SNew(SBox)
-				.WidthOverride(16.0f)
-				[
-					ExternalScrollBar
-				]
-			]
 		]
 	];
 }
@@ -257,12 +246,6 @@ FText SRCProtocolBindingList::HandleAddBindingToolTipText() const
 	return FText::GetEmpty();
 }
 
-EVisibility SRCProtocolBindingList::GetScrollBarVisibility() const
-{
-	const bool bHasAnythingToShow = FilteredBindings.Num() > 0;
-	return bHasAnythingToShow ? EVisibility::Visible : EVisibility::Collapsed;
-}
-
 void SRCProtocolBindingList::OnStartRecording(TSharedPtr<TStructOnScope<FRemoteControlProtocolEntity>> InEntity)
 {
 	AwaitingProtocolEntities.Add(InEntity);
diff --git a/Engine/Plugins/VirtualProduction/RemoteControl/Source/RemoteControlProtocolWidgets/Private/Widgets/SRCProtocolBindingList.h b/Engine/Plugins/VirtualProduction/RemoteControl/Source/RemoteControlProtocolWidgets/Private/Widgets/SRCProtocolBindingList.h
--- a/Engine/Plugins/VirtualProduction/RemoteControl/Source/RemoteControlProtocolWidgets/Private/Widgets/SRCProtocolBindingList.h
+++ b/Engine/Plugins/VirtualProduction/RemoteControl/Source/RemoteControlProtocolWidgets/Private/Widgets/SRCProtocolBindingList.h
@@ -59,9 +59,6 @@ private:
 	/** Retrieves the appropriate tooltip based on the active protocol. */
 	FText HandleAddBindingToolTipText() const;
 
-	/** Called to get the visibility of the scrollbar based on options, needs to be dynamic to avoid layout changing on expansion */
-	EVisibility GetScrollBarVisibility() const;
-
 	/** Primary Column is between Input and Output values. */
 	float OnGetPrimaryLeftColumnWidth() const
 	{
