RemoteControl: Add interfaces to allow for protocol specific widgets and sorting in the exposed entities list

diff --git a/Engine/Plugins/VirtualProduction/RemoteControl/Source/RemoteControl/Public/RemoteControlPreset.h b/Engine/Plugins/VirtualProduction/RemoteControl/Source/RemoteControl/Public/RemoteControlPreset.h
--- a/Engine/Plugins/VirtualProduction/RemoteControl/Source/RemoteControl/Public/RemoteControlPreset.h
+++ b/Engine/Plugins/VirtualProduction/RemoteControl/Source/RemoteControl/Public/RemoteControlPreset.h
@@ -3,8 +3,8 @@
 
 #include "Algo/Transform.h"
 #include "CoreTypes.h"
-#include "RemoteControlField.h"
 #include "RemoteControlEntity.h"
+#include "RemoteControlField.h"
 #include "RemoteControlPropertyIdRegistry.h"
 #include "Templates/PimplPtr.h"
 #include "Templates/UnrealTypeTraits.h"
@@ -725,6 +725,10 @@ public:
 	UPROPERTY(EditAnywhere, Category = "Remote Control Preset")
 	TArray<TObjectPtr<URemoteControlBinding>> Bindings;
 
+	/** This preset's user data */
+	UPROPERTY()
+	TArray<TObjectPtr<UObject>> UserData;
+
 #if WITH_EDITOR
 	/** Current selected world used only in the editor to let you switch world */
 	TWeakObjectPtr<const UWorld> SelectedWorld = nullptr;
@@ -1066,7 +1070,6 @@ public:
 #endif
 
 private:
-
 #if WITH_EDITORONLY_DATA
 	UPROPERTY()
 	TArray<FName> DetailsTabIdentifierOverrides;
diff --git a/Engine/Plugins/VirtualProduction/RemoteControl/Source/RemoteControl/Public/RemoteControlProtocolBinding.h b/Engine/Plugins/VirtualProduction/RemoteControl/Source/RemoteControl/Public/RemoteControlProtocolBinding.h
--- a/Engine/Plugins/VirtualProduction/RemoteControl/Source/RemoteControl/Public/RemoteControlProtocolBinding.h
+++ b/Engine/Plugins/VirtualProduction/RemoteControl/Source/RemoteControl/Public/RemoteControlProtocolBinding.h
@@ -437,15 +437,16 @@ public:
 	virtual bool HasMask(ERCMask InMaskBit) const;
 
 #if WITH_EDITOR
-
 	/** Retrieves the name of property corresponding to the given column name. */
 	const FName GetPropertyName(const FName& ForColumnName);
 	
 	/** Register(s) all the properties (to be exposed) of this protocol entity. */
 	virtual void RegisterProperties() {};
-
 #endif // WITH_EDITOR
 
+	/** Returns the preset that owns this protocol entity */
+	const TWeakObjectPtr<URemoteControlPreset>& GetOwner() const { return Owner; }
+
 public:
 	/** Container for range and mapping value pointers, and an optional number of elements (arrays, strings). */
 	struct FRangeMappingData
diff --git a/Engine/Plugins/VirtualProduction/RemoteControl/Source/RemoteControl/Public/UI/RCFieldGroupOrder.h b/Engine/Plugins/VirtualProduction/RemoteControl/Source/RemoteControl/Public/UI/RCFieldGroupOrder.h
new file mode 100644
--- /dev/null
+++ b/Engine/Plugins/VirtualProduction/RemoteControl/Source/RemoteControl/Public/UI/RCFieldGroupOrder.h
@@ -0,0 +1,14 @@
+// Copyright Epic Games, Inc. All Rights Reserved.
+
+#pragma once
+
+#include "RCFieldGroupOrder.generated.h"
+
+/** Order of remote control field groups */
+UENUM()
+enum class ERCFieldGroupOrder : uint8
+{
+	None,
+	Ascending,
+	Descending
+};
diff --git a/Engine/Plugins/VirtualProduction/RemoteControl/Source/RemoteControl/Public/UI/RCFieldGroupType.h b/Engine/Plugins/VirtualProduction/RemoteControl/Source/RemoteControl/Public/UI/RCFieldGroupType.h
new file mode 100644
--- /dev/null
+++ b/Engine/Plugins/VirtualProduction/RemoteControl/Source/RemoteControl/Public/UI/RCFieldGroupType.h
@@ -0,0 +1,22 @@
+// Copyright Epic Games, Inc. All Rights Reserved.
+
+#pragma once
+
+#include "RCFieldGroupType.generated.h"
+
+/** Types of remote control field groups */
+UENUM()
+enum class ERCFieldGroupType : uint8
+{
+	/** No Grouping. Will be shown as a flat list */
+	None,
+
+	/** Fields with same property id are grouped together */
+	PropertyId,
+
+	/** Fields with same owner are grouped together */
+	Owner,
+};
+
+ENUM_CLASS_FLAGS(ERCFieldGroupType)
+
diff --git a/Engine/Plugins/VirtualProduction/RemoteControl/Source/RemoteControl/Public/UI/RCPanelExposedEntitiesListSettings.h b/Engine/Plugins/VirtualProduction/RemoteControl/Source/RemoteControl/Public/UI/RCPanelExposedEntitiesListSettings.h
new file mode 100644
--- /dev/null
+++ b/Engine/Plugins/VirtualProduction/RemoteControl/Source/RemoteControl/Public/UI/RCPanelExposedEntitiesListSettings.h
@@ -0,0 +1,23 @@
+// Copyright Epic Games, Inc. All Rights Reserved.
+
+#pragma once
+
+#include "RCPanelExposedEntitiesListSettings.generated.h"
+
+enum class ERCFieldGroupType : uint8;
+enum class ERCFieldGroupOrder : uint8;
+
+/** Data for exposed entities list settings */
+USTRUCT()
+struct FRCPanelExposedEntitiesListSettingsData
+{
+	GENERATED_BODY()
+
+	/** The field group type for the entity list. */
+	UPROPERTY()
+	ERCFieldGroupType FieldGroupType;
+
+	/** Whether the field groups are expanded */
+	UPROPERTY()
+	ERCFieldGroupOrder FieldGroupOrder;
+};
diff --git a/Engine/Plugins/VirtualProduction/RemoteControl/Source/RemoteControlUI/Private/RemoteControlUIModule.cpp b/Engine/Plugins/VirtualProduction/RemoteControl/Source/RemoteControlUI/Private/RemoteControlUIModule.cpp
--- a/Engine/Plugins/VirtualProduction/RemoteControl/Source/RemoteControlUI/Private/RemoteControlUIModule.cpp
+++ b/Engine/Plugins/VirtualProduction/RemoteControl/Source/RemoteControlUI/Private/RemoteControlUIModule.cpp
@@ -18,8 +18,8 @@
 #include "Engine/StaticMesh.h"
 #include "Framework/MultiBox/MultiBoxBuilder.h"
 #include "HAL/PlatformApplicationMisc.h"
-#include "IRemoteControlModule.h"
 #include "Interfaces/IMainFrameModule.h"
+#include "IRemoteControlModule.h"
 #include "Kismet2/ComponentEditorUtils.h"
 #include "LevelEditorSubsystem.h"
 #include "MaterialDomain.h"
@@ -44,6 +44,9 @@
 #include "UI/Customizations/NetworkAddressCustomization.h"
 #include "UI/Customizations/RCAssetPathElementCustomization.h"
 #include "UI/Customizations/RemoteControlEntityCustomization.h"
+#include "UI/IRCPanelExposedEntitiesGroupWidgetFactory.h"
+#include "UI/IRCPanelExposedEntitiesListSettingsForProtocol.h"
+#include "UI/IRCPanelExposedEntityWidgetFactory.h"
 #include "UI/RemoteControlExposeMenuStyle.h"
 #include "UI/RemoteControlPanelStyle.h"
 #include "UI/SRCPanelExposedActor.h"
@@ -616,6 +619,46 @@ void FRemoteControlUIModule::UnregisterSignatureCustomization(const TSharedPtr<I
 	}
 }
 
+void FRemoteControlUIModule::RegisterExposedEntitiesListSettingsForProtocol(const TSharedRef<IRCPanelExposedEntitiesListSettingsForProtocol>& InSettings)
+{
+	ExposedEntitiesListSettingsForProtocols.AddUnique(InSettings);
+}
+
+void FRemoteControlUIModule::UnregisterExposedEntitiesListSettingsForProtocol(const TSharedRef<IRCPanelExposedEntitiesListSettingsForProtocol>& InSettings)
+{
+	ExposedEntitiesListSettingsForProtocols.RemoveSingle(InSettings);
+}
+
+void FRemoteControlUIModule::RegisterExposedEntitiesPanelExtender(const TSharedRef<IRCExposedEntitiesPanelExtender>& InExtender)
+{
+	ExposedEntitiesPanelExtenders.AddUnique(InExtender);
+}
+
+void FRemoteControlUIModule::UnregisterExposedEntitiesPanelExtender(const TSharedRef<IRCExposedEntitiesPanelExtender>& InExtender)
+{
+	ExposedEntitiesPanelExtenders.RemoveSingle(InExtender);
+}
+
+void FRemoteControlUIModule::RegisterExposedEntitiesGroupWidgetFactory(const TSharedRef<IRCPanelExposedEntitiesGroupWidgetFactory>& InFactory)
+{
+	ExposedEntitiesGroupWidgetFactories.AddUnique(InFactory);
+}
+
+void FRemoteControlUIModule::UnregisterExposedEntitiesGroupWidgetFactory(const TSharedRef<IRCPanelExposedEntitiesGroupWidgetFactory>& InFactory)
+{
+	ExposedEntitiesGroupWidgetFactories.RemoveSingle(InFactory);
+}
+
+void FRemoteControlUIModule::RegisterExposedEntityWidgetFactory(const TSharedRef<IRCPanelExposedEntityWidgetFactory>& InFactory)
+{
+	ExposedEntityWidgetFactories.AddUnique(InFactory);
+}
+
+void FRemoteControlUIModule::UnregisterExposedEntityWidgetFactory(const TSharedRef<IRCPanelExposedEntityWidgetFactory>& InFactory)
+{
+	ExposedEntityWidgetFactories.RemoveSingle(InFactory);
+}
+
 void FRemoteControlUIModule::RegisterAssetTools()
 {
 	if (FAssetToolsModule* AssetToolsModule = FModuleManager::GetModulePtr<FAssetToolsModule>("AssetTools"))
@@ -1632,6 +1675,50 @@ void FRemoteControlUIModule::UnregisterRemoteControlPanel(SRemoteControlPanel* P
 	}
 }
 
+const TSharedRef<IRCPanelExposedEntitiesListSettingsForProtocol>* FRemoteControlUIModule::GetExposedEntitiesListSettingsForProtocol(const FName& ProtocolName) const
+{
+	return Algo::FindByPredicate(ExposedEntitiesListSettingsForProtocols, [&ProtocolName, this](const TSharedRef<IRCPanelExposedEntitiesListSettingsForProtocol>& Settings)
+		{
+			return Settings->GetProtocolName() == ProtocolName;
+		});
+}
+
+const TSharedRef<IRCPanelExposedEntitiesGroupWidgetFactory>* FRemoteControlUIModule::GetExposedEntitiesGroupWidgetFactory(const FName& ForColumnName, const FName& InActiveProtocol) const
+{
+	return Algo::FindByPredicate(ExposedEntitiesGroupWidgetFactories, [&ForColumnName, &InActiveProtocol](const TSharedRef<IRCPanelExposedEntitiesGroupWidgetFactory>& WidgetFactory)
+		{
+			if (WidgetFactory->GetColumnName() == ForColumnName)
+			{
+				const FName ProtocolName = WidgetFactory->GetProtocolName();
+
+				// Customize the column for any protcol if the protocol name is none
+				return 
+					ProtocolName == InActiveProtocol || 
+					ProtocolName == NAME_None;
+			}
+			
+			return false;
+		});
+}
+
+const TSharedRef<IRCPanelExposedEntityWidgetFactory>* FRemoteControlUIModule::GetExposedEntityWidgetFactory(const FName& ForColumnName, const FName& InActiveProtocol) const
+{
+	return Algo::FindByPredicate(ExposedEntityWidgetFactories, [&ForColumnName, &InActiveProtocol](const TSharedRef<IRCPanelExposedEntityWidgetFactory>& ExposedEntityCustomization)
+		{
+			if (ExposedEntityCustomization->GetColumnName() == ForColumnName)
+			{
+				const FName ProtocolName = ExposedEntityCustomization->GetProtocolName();
+
+				// Customize the column for any protcol if the protocol name is none
+				return
+					ProtocolName == InActiveProtocol ||
+					ProtocolName == NAME_None;
+			}
+
+			return false;
+		});
+}
+
 IMPLEMENT_MODULE(FRemoteControlUIModule, RemoteControlUI);
 
 #undef LOCTEXT_NAMESPACE /*RemoteControlUI*/
diff --git a/Engine/Plugins/VirtualProduction/RemoteControl/Source/RemoteControlUI/Private/RemoteControlUIModule.h b/Engine/Plugins/VirtualProduction/RemoteControl/Source/RemoteControlUI/Private/RemoteControlUIModule.h
--- a/Engine/Plugins/VirtualProduction/RemoteControl/Source/RemoteControlUI/Private/RemoteControlUIModule.h
+++ b/Engine/Plugins/VirtualProduction/RemoteControl/Source/RemoteControlUI/Private/RemoteControlUIModule.h
@@ -11,10 +11,12 @@
 #include "PropertyHandle.h"
 #include "PropertyPath.h"
 
+class IRCPanelExposedEntitiesListSettingsForProtocol;
 class IToolkitHost;
 class SRemoteControlPanel;
 class UMeshComponent;
 class URemoteControlPreset;
+enum class ERCPanelMode : uint8;
 
 /**
  * A Remote Control module that allows exposing objects and properties from the editor.
@@ -58,6 +60,14 @@ public:
 	virtual uint32 GetRemoteControlAssetCategory() const override;
 	virtual void RegisterSignatureCustomization(const TSharedPtr<IRCSignatureCustomization>& InCustomization) override;
 	virtual void UnregisterSignatureCustomization(const TSharedPtr<IRCSignatureCustomization>& InCustomization) override;
+	virtual void RegisterExposedEntitiesListSettingsForProtocol(const TSharedRef<IRCPanelExposedEntitiesListSettingsForProtocol>& InSettings) override;
+	virtual void UnregisterExposedEntitiesListSettingsForProtocol(const TSharedRef<IRCPanelExposedEntitiesListSettingsForProtocol>& InSettings) override;
+	virtual void RegisterExposedEntitiesPanelExtender(const TSharedRef<IRCExposedEntitiesPanelExtender>& InExtender) override;
+	virtual void UnregisterExposedEntitiesPanelExtender(const TSharedRef<IRCExposedEntitiesPanelExtender>& InExtender) override;
+	virtual void RegisterExposedEntitiesGroupWidgetFactory(const TSharedRef<IRCPanelExposedEntitiesGroupWidgetFactory>& InFactory) override;
+	virtual void UnregisterExposedEntitiesGroupWidgetFactory(const TSharedRef<IRCPanelExposedEntitiesGroupWidgetFactory>& InFactory) override;
+	virtual void RegisterExposedEntityWidgetFactory(const TSharedRef<IRCPanelExposedEntityWidgetFactory>& InFactory) override;
+	virtual void UnregisterExposedEntityWidgetFactory(const TSharedRef<IRCPanelExposedEntityWidgetFactory>& InFactory) override;
 	virtual void RegisterWidgetFactoryForType(UScriptStruct* RemoteControlEntityType, const FOnGenerateRCWidget& OnGenerateRCWidgetDelegate) override;
 	virtual void UnregisterWidgetFactoryForType(UScriptStruct* RemoteControlEntityType) override;
 	virtual void HighlightPropertyInDetailsPanel(const FPropertyPath& Path) const override;
@@ -87,6 +97,20 @@ public:
 		return SignatureCustomizations;
 	}
 
+	TConstArrayView<TSharedRef<IRCExposedEntitiesPanelExtender>> GetExposedEntitiesPanelExtenders() const
+	{
+		return ExposedEntitiesPanelExtenders;
+	}
+
+	/** Returns the exposed entity list settings for a protocol, or nullptr if no entities list settings could be found */
+	const TSharedRef<IRCPanelExposedEntitiesListSettingsForProtocol>* GetExposedEntitiesListSettingsForProtocol(const FName& ProtocolName) const;
+
+	/** Returns the exposed entities group widget factory, or nullptr if no widget factory is registered for the specified protocol and column name. */
+	const TSharedRef<IRCPanelExposedEntitiesGroupWidgetFactory>* GetExposedEntitiesGroupWidgetFactory(const FName& ForColumnName, const FName& InActiveProtocol) const;
+
+	/** Returns the exposed entity widget factory, or nullptr if no widget factory is registered for the specified protocol and column name. */
+	const TSharedRef<IRCPanelExposedEntityWidgetFactory>* GetExposedEntityWidgetFactory(const FName& ForColumnName, const FName& InActiveProtocol) const;
+
 public:
 	static const FName RemoteControlPanelTabName;
 
@@ -278,6 +302,18 @@ private:
 	/** Registered Signature Customizations */
 	TArray<TSharedRef<IRCSignatureCustomization>> SignatureCustomizations; 
 
+	/** Registered Exposed Entities List Settings for Protocols */
+	TArray<TSharedRef<IRCPanelExposedEntitiesListSettingsForProtocol>> ExposedEntitiesListSettingsForProtocols;
+
+	/** Registered Exposed Entities Panel Extenders */
+	TArray<TSharedRef<IRCExposedEntitiesPanelExtender>> ExposedEntitiesPanelExtenders;
+
+	/** Registered Exposed Entities Group Widget Factories */
+	TArray<TSharedRef<IRCPanelExposedEntitiesGroupWidgetFactory>> ExposedEntitiesGroupWidgetFactories;
+
+	/** Registered Exposed Entity Widget Factories */
+	TArray<TSharedRef<IRCPanelExposedEntityWidgetFactory>> ExposedEntityWidgetFactories;
+
 	TMap<TWeakObjectPtr<UScriptStruct>, FOnGenerateRCWidget> GenerateWidgetDelegates;
 
 	/** Holds the advanced asset type category bit registered by Remote Control. */
diff --git a/Engine/Plugins/VirtualProduction/RemoteControl/Source/RemoteControlUI/Private/UI/Action/SRCActionPanel.cpp b/Engine/Plugins/VirtualProduction/RemoteControl/Source/RemoteControlUI/Private/UI/Action/SRCActionPanel.cpp
--- a/Engine/Plugins/VirtualProduction/RemoteControl/Source/RemoteControlUI/Private/UI/Action/SRCActionPanel.cpp
+++ b/Engine/Plugins/VirtualProduction/RemoteControl/Source/RemoteControlUI/Private/UI/Action/SRCActionPanel.cpp
@@ -604,7 +604,7 @@ FReply SRCActionPanel::OnAddAllSelectedFields()
 				{
 					if (const TSharedPtr<SRCPanelExposedEntitiesGroup> RCFieldGroup = StaticCastSharedPtr<SRCPanelExposedEntitiesGroup>(RCEntity))
 					{
-						if (RCFieldGroup->GetGroupType() == EFieldGroupType::PropertyId)
+						if (RCFieldGroup->GetGroupType() == ERCFieldGroupType::PropertyId)
 						{
 							const URCBehaviour* Behaviour = BehaviourItem->GetBehaviour();
 							if (Behaviour && Behaviour->SupportPropertyId())
diff --git a/Engine/Plugins/VirtualProduction/RemoteControl/Source/RemoteControlUI/Private/UI/Action/SRCActionPanelList.h b/Engine/Plugins/VirtualProduction/RemoteControl/Source/RemoteControlUI/Private/UI/Action/SRCActionPanelList.h
--- a/Engine/Plugins/VirtualProduction/RemoteControl/Source/RemoteControlUI/Private/UI/Action/SRCActionPanelList.h
+++ b/Engine/Plugins/VirtualProduction/RemoteControl/Source/RemoteControlUI/Private/UI/Action/SRCActionPanelList.h
@@ -8,14 +8,15 @@
 #include "Action/RCPropertyAction.h"
 #include "RCActionModel.h"
 #include "RemoteControlPreset.h"
+#include "SDropTarget.h"
 #include "SlateOptMacros.h"
 #include "SRCActionPanel.h"
-#include "SDropTarget.h"
 #include "Styling/RemoteControlStyles.h"
 #include "UI/Action/Bind/RCActionBindModel.h"
 #include "UI/Action/Conditional/RCActionConditionalModel.h"
 #include "UI/BaseLogicUI/SRCLogicPanelListBase.h"
 #include "UI/Behaviour/RCBehaviourModel.h"
+#include "UI/RCFieldGroupType.h"
 #include "UI/RCUIHelpers.h"
 #include "UI/RemoteControlPanelStyle.h"
 #include "UI/SRCPanelExposedEntitiesGroup.h"
@@ -379,9 +380,9 @@ private:
 						{
 							if (const TSharedPtr<SRCPanelExposedEntitiesGroup>& FieldGroup = StaticCastSharedPtr<SRCPanelExposedEntitiesGroup>(ExposedEntity))
 							{
-								const EFieldGroupType GroupType = FieldGroup->GetGroupType();
+								const ERCFieldGroupType GroupType = FieldGroup->GetGroupType();
 								// Property Id group
-								if (GroupType == EFieldGroupType::PropertyId)
+								if (GroupType == ERCFieldGroupType::PropertyId)
 								{
 									if (bBehaviorSupportPropertyId)
 									{
@@ -389,7 +390,7 @@ private:
 									}
 								}
 								// Owner group
-								else if (GroupType == EFieldGroupType::Owner)
+								else if (GroupType == ERCFieldGroupType::Owner)
 								{
 									TArray<TSharedPtr<SRCPanelTreeNode>> OwnerGroupChildren;
 									FieldGroup->GetNodeChildren(OwnerGroupChildren);
@@ -485,9 +486,9 @@ private:
 						{
 							if (const TSharedPtr<SRCPanelExposedEntitiesGroup>& FieldGroup = StaticCastSharedPtr<SRCPanelExposedEntitiesGroup>(ExposedEntity))
 							{
-								const EFieldGroupType GroupType = FieldGroup->GetGroupType();
+								const ERCFieldGroupType GroupType = FieldGroup->GetGroupType();
 								// Property Id Group
-								if (GroupType == EFieldGroupType::PropertyId)
+								if (GroupType == ERCFieldGroupType::PropertyId)
 								{
 									if (bBehaviorSupportPropertyId)
 									{
@@ -495,7 +496,7 @@ private:
 									}
 								}
 								// Owner Group
-								else if (GroupType == EFieldGroupType::Owner)
+								else if (GroupType == ERCFieldGroupType::Owner)
 								{
 									TArray<TSharedPtr<SRCPanelTreeNode>> OwnerGroupChildren;
 									FieldGroup->GetNodeChildren(OwnerGroupChildren);
diff --git a/Engine/Plugins/VirtualProduction/RemoteControl/Source/RemoteControlUI/Private/UI/SRCPanelExposedEntitiesGroup.cpp b/Engine/Plugins/VirtualProduction/RemoteControl/Source/RemoteControlUI/Private/UI/SRCPanelExposedEntitiesGroup.cpp
--- a/Engine/Plugins/VirtualProduction/RemoteControl/Source/RemoteControlUI/Private/UI/SRCPanelExposedEntitiesGroup.cpp
+++ b/Engine/Plugins/VirtualProduction/RemoteControl/Source/RemoteControlUI/Private/UI/SRCPanelExposedEntitiesGroup.cpp
@@ -3,13 +3,15 @@
 #include "SRCPanelExposedEntitiesGroup.h"
 
 #include "RemoteControlPreset.h"
+#include "RemoteControlUIModule.h"
 #include "SRCPanelExposedField.h"
-#include "Widgets/Layout/SBox.h"
+#include "UI/IRCPanelExposedEntitiesGroupWidgetFactory.h"
 #include "Widgets/Input/SEditableTextBox.h"
+#include "Widgets/Layout/SBox.h"
 
 #define LOCTEXT_NAMESPACE "SRCPanelExposedEntitiesGroup"
 
-void SRCPanelExposedEntitiesGroup::Construct(const FArguments& InArgs, EFieldGroupType InFieldGroupType, URemoteControlPreset* Preset)
+void SRCPanelExposedEntitiesGroup::Construct(const FArguments& InArgs, ERCFieldGroupType InFieldGroupType, URemoteControlPreset* Preset)
 {
 	FieldKey = InArgs._FieldKey;
 	GroupType = InFieldGroupType;
@@ -20,6 +22,25 @@ void SRCPanelExposedEntitiesGroup::Construct(const FArguments& InArgs, EFieldGro
 	MakeNodeWidget(Args);
 }
 
+TSharedRef<SWidget> SRCPanelExposedEntitiesGroup::GetWidget(const FName ForColumnName, const FName InActiveProtocol)
+{
+	const FRemoteControlUIModule& RemoteControlUIModule = FModuleManager::GetModuleChecked<FRemoteControlUIModule>("RemoteControlUI");
+	const TSharedRef<IRCPanelExposedEntitiesGroupWidgetFactory>* CustomFactoryPtr = RemoteControlUIModule.GetExposedEntitiesGroupWidgetFactory(ForColumnName, InActiveProtocol);
+	if (CustomFactoryPtr)
+	{
+		const FRCPanelExposedEntitiesGroupWidgetFactoryArgs Args(
+			PresetWeak,
+			GetChildProperties()
+		);
+
+		return (*CustomFactoryPtr)->MakeWidget(Args);
+	}
+	else
+	{
+		return SRCPanelTreeNode::GetWidget(ForColumnName, InActiveProtocol);
+	}
+}
+
 SRCPanelTreeNode::FMakeNodeWidgetArgs SRCPanelExposedEntitiesGroup::CreateNodeWidgetArgs()
 {
 	FMakeNodeWidgetArgs Args;
@@ -48,11 +69,11 @@ SRCPanelTreeNode::FMakeNodeWidgetArgs SRCPanelExposedEntitiesGroup::CreateNodeWi
 
 	switch (GroupType)
 	{
-	case EFieldGroupType::PropertyId:
+	case ERCFieldGroupType::PropertyId:
 		GroupText = LOCTEXT("GroupPropertyId", "Group by Id");
 		break;
 
-	case EFieldGroupType::Owner:
+	case ERCFieldGroupType::Owner:
 		GroupText = LOCTEXT("GroupOwner", "Group by Owner");
 		break;
 	}
@@ -129,8 +150,8 @@ void SRCPanelExposedEntitiesGroup::AssignChildren(const TArray<TSharedPtr<SRCPan
 			const FName FieldOwnerName  = ExposedField->GetOwnerName();
 			const FName FieldPropertyId = ExposedField->GetPropertyId();
 
-			const bool bOwnerMatch = GroupType == EFieldGroupType::Owner && FieldOwnerName == FieldKey;
-			const bool bPropertyIdMatch = GroupType == EFieldGroupType::PropertyId && FieldPropertyId == FieldKey;
+			const bool bOwnerMatch = GroupType == ERCFieldGroupType::Owner && FieldOwnerName == FieldKey;
+			const bool bPropertyIdMatch = GroupType == ERCFieldGroupType::PropertyId && FieldPropertyId == FieldKey;
 
 			if (bOwnerMatch || bPropertyIdMatch)
 			{
@@ -148,4 +169,28 @@ void SRCPanelExposedEntitiesGroup::GetNodeChildren(TArray<TSharedPtr<SRCPanelTre
 	SRCPanelTreeNode::GetNodeChildren(OutChildren);
 }
 
+TArray<TSharedRef<FRemoteControlProperty>> SRCPanelExposedEntitiesGroup::GetChildProperties() const
+{
+	URemoteControlPreset* Preset = PresetWeak.Get();
+	if (!Preset)
+	{
+		return {};
+	}
+
+	TArray<TSharedPtr<SRCPanelTreeNode>> Children;
+	GetNodeChildren(Children);
+
+	TArray<TSharedRef<FRemoteControlProperty>> ChildProperties;
+	for (const TSharedPtr<SRCPanelTreeNode>& Child : Children)
+	{
+		const TWeakPtr<FRemoteControlProperty> WeakProperty = Child.IsValid() ? Preset->GetExposedEntity<FRemoteControlProperty>(Child->GetRCId()) : nullptr;
+		if (WeakProperty.IsValid())
+		{
+			ChildProperties.Add(WeakProperty.Pin().ToSharedRef());
+		}
+	}
+
+	return ChildProperties;
+}
+
 #undef LOCTEXT_NAMESPACE
diff --git a/Engine/Plugins/VirtualProduction/RemoteControl/Source/RemoteControlUI/Private/UI/SRCPanelExposedEntitiesGroup.h b/Engine/Plugins/VirtualProduction/RemoteControl/Source/RemoteControlUI/Private/UI/SRCPanelExposedEntitiesGroup.h
--- a/Engine/Plugins/VirtualProduction/RemoteControl/Source/RemoteControlUI/Private/UI/SRCPanelExposedEntitiesGroup.h
+++ b/Engine/Plugins/VirtualProduction/RemoteControl/Source/RemoteControlUI/Private/UI/SRCPanelExposedEntitiesGroup.h
@@ -2,25 +2,12 @@
 
 #pragma once
 
-#include "CoreMinimal.h"
 #include "SRCPanelExposedEntity.h"
 #include "SRCPanelTreeNode.h"
+#include "UI/RCFieldGroupType.h"
 #include "Widgets/SCompoundWidget.h"
 
-/** Grouping type available */
-enum class EFieldGroupType : uint8
-{
-	/** No Grouping. Will be shown as a flat list */
-	None,
-
-	/** Fields with same property id are grouped together */
-	PropertyId,
-
-	/** Fields with same owner are grouped together */
-	Owner,
-};
-
-ENUM_CLASS_FLAGS(EFieldGroupType)
+struct FRemoteControlProperty;
 
 /** Delegate called when the PropertyId of a group change */
 DECLARE_DELEGATE(FOnGroupPropertyIdChanged)
@@ -28,6 +15,7 @@ DECLARE_DELEGATE(FOnGroupPropertyIdChanged)
 /** Widget that is used as a parent for field widget */
 class SRCPanelExposedEntitiesGroup : public SRCPanelTreeNode
 {
+public:
 	SLATE_BEGIN_ARGS(SRCPanelExposedEntitiesGroup)
 	{}
 
@@ -43,7 +31,11 @@ class SRCPanelExposedEntitiesGroup : public SRCPanelTreeNode
 
 	//~ End SRCPanelTreeNode interface
 
-	void Construct(const FArguments& InArgs, EFieldGroupType InFieldGroupType, URemoteControlPreset* Preset);
+	void Construct(const FArguments& InArgs, ERCFieldGroupType InFieldGroupType, URemoteControlPreset* Preset);
+
+	//~ Begin SRCPanelTreeNode interface
+	virtual TSharedRef<SWidget> GetWidget(const FName ForColumnName, const FName InActiveProtocol) override;
+	//~ End SRCPanelTreeNode interface
 
 	/**
 	 * Take the current fields in RC and assign them to its child widgets based on the group key
@@ -58,7 +50,7 @@ class SRCPanelExposedEntitiesGroup : public SRCPanelTreeNode
 	virtual bool HasChildren() const override { return true; }
 
 	/** Get the group type of this group */
-	EFieldGroupType GetGroupType() const { return GroupType; }
+	ERCFieldGroupType GetGroupType() const { return GroupType; }
 
 	/** Get the field key of this group */
 	FName GetFieldKey() const { return FieldKey; }
@@ -74,7 +66,9 @@ private:
 	 */
 	void OnPropertyIdTextCommitted(const FText& InText, ETextCommit::Type InTextCommitType);
 
-private:
+	/** Returns the properties that are children of this group */
+	TArray<TSharedRef<FRemoteControlProperty>> GetChildProperties() const;
+
 	/** Child widgets of this group */
 	TArray<TSharedPtr<SRCPanelTreeNode>> ChildWidgets;
 
@@ -88,7 +82,7 @@ private:
 	FName PropertyIdName;
 
 	/** Field group type of this group */
-	EFieldGroupType GroupType = EFieldGroupType::None;
+	ERCFieldGroupType GroupType = ERCFieldGroupType::None;
 
 	/** Weak ptr to the preset where this group reside */
 	TWeakObjectPtr<URemoteControlPreset> PresetWeak;
diff --git a/Engine/Plugins/VirtualProduction/RemoteControl/Source/RemoteControlUI/Private/UI/SRCPanelExposedEntitiesList.cpp b/Engine/Plugins/VirtualProduction/RemoteControl/Source/RemoteControlUI/Private/UI/SRCPanelExposedEntitiesList.cpp
--- a/Engine/Plugins/VirtualProduction/RemoteControl/Source/RemoteControlUI/Private/UI/SRCPanelExposedEntitiesList.cpp
+++ b/Engine/Plugins/VirtualProduction/RemoteControl/Source/RemoteControlUI/Private/UI/SRCPanelExposedEntitiesList.cpp
@@ -2,15 +2,11 @@
 
 #include "SRCPanelExposedEntitiesList.h"
 
-#include "Algo/ForEach.h"
 #include "Commands/RemoteControlCommands.h"
 #include "Editor.h"
-#include "Editor/EditorEngine.h"
-#include "EditorFontGlyphs.h"
 #include "Engine/Selection.h"
 #include "Filters/SRCPanelFilter.h"
 #include "GameFramework/Actor.h"
-#include "ISettingsModule.h"
 #include "Input/DragAndDrop.h"
 #include "IRemoteControlProtocolModule.h"
 #include "IRemoteControlProtocolWidgetsModule.h"
@@ -19,7 +15,6 @@
 #include "Misc/MessageDialog.h"
 #include "PropertyPath.h"
 #include "RCPanelWidgetRegistry.h"
-#include "RemoteControlPropertyIdRegistry.h"
 #include "RemoteControlPanelStyle.h"
 #include "RemoteControlPreset.h"
 #include "RemoteControlSettings.h"
@@ -28,19 +23,17 @@
 #include "SRCHeaderRow.h"
 #include "SRCModeSwitcher.h"
 #include "SRCPanelExposedEntitiesGroup.h"
-#include "SRCPanelFieldGroup.h"
 #include "SRCPanelExposedField.h"
+#include "SRCPanelFieldGroup.h"
 #include "SSearchToggleButton.h"
 #include "Styling/RemoteControlStyles.h"
+#include "UI/IRCExposedEntitiesPanelExtender.h"
+#include "UI/IRCPanelExposedEntitiesListSettingsForProtocol.h"
 #include "UI/Panels/SRCDockPanel.h"
 #include "UObject/Object.h"
 #include "Widgets/Input/SButton.h"
-#include "Widgets/Input/SComboBox.h"
 #include "Widgets/Input/SComboButton.h"
 #include "Widgets/Input/SSearchBox.h"
-#include "Widgets/Layout/SSeparator.h"
-#include "Widgets/Layout/SSpacer.h"
-#include "Widgets/Layout/SWidgetSwitcher.h"
 #include "Widgets/Views/SListView.h"
 #include "Widgets/Views/STableRow.h"
 #include "Widgets/Views/STreeView.h"
@@ -64,7 +57,7 @@ public:
 
 	SLATE_BEGIN_ARGS(SEntityRow)
 	{}
-	
+		
 		SLATE_ATTRIBUTE(FName, ActiveProtocol)
 		SLATE_ARGUMENT(TSharedPtr<SRCPanelTreeNode>, Entity)
 
@@ -266,6 +259,8 @@ void SRCPanelExposedEntitiesList::Construct(const FArguments& InArgs, URemoteCon
 	bFilterApplicationRequested = false;
 	bSearchRequested = false;
 	
+	RecallListSettings();
+
 	// Setup search filter.
 	SearchTextFilter = MakeShared<TTextFilter<const SRCPanelTreeNode&>>(TTextFilter<const SRCPanelTreeNode&>::FItemToStringArray::CreateSP(this, &SRCPanelExposedEntitiesList::PopulateSearchStrings));
 	SearchedText = MakeShared<FText>(FText::GetEmpty());
@@ -387,12 +382,39 @@ void SRCPanelExposedEntitiesList::Construct(const FArguments& InArgs, URemoteCon
 		);
 
 	// Exposed Entities Dock Panel
-	TSharedPtr<SRCMinorPanel> ExposeDockPanel = SNew(SRCMinorPanel)
+	const TSharedRef<SVerticalBox> ExposedEntitiesPanelContent = SNew(SVerticalBox);
+
+	const FRemoteControlUIModule& RemoteControlUIModule = FModuleManager::Get().GetModuleChecked<FRemoteControlUIModule>("RemoteControlUI");
+	for (const TSharedRef<IRCExposedEntitiesPanelExtender>& Extender : RemoteControlUIModule.GetExposedEntitiesPanelExtenders())
+	{
+		// Add extensions
+		IRCExposedEntitiesPanelExtender::FArgs Args;
+		Args.ActiveProtocolAttribute = TAttribute<FName>::CreateLambda([this]()
+			{
+				return bIsInProtocolsMode.Get() ? ActiveProtocol : NAME_None;
+			});
+		const TSharedRef<SWidget> ExtenderWidget = Extender->MakeWidget(Preset.Get(), Args);
+
+		ExposedEntitiesPanelContent->AddSlot()
+			.AutoHeight()
+			[
+				ExtenderWidget
+			];
+	}
+
+
+	ExposedEntitiesPanelContent->AddSlot()
+		.FillHeight(1.f)
+		[
+			FieldsListView.ToSharedRef()
+		];
+
+	const TSharedPtr<SRCMinorPanel> ExposeDockPanel = SNew(SRCMinorPanel)
 		.HeaderLabel(this, &SRCPanelExposedEntitiesList::HandleEntityListHeaderLabel)
 		.Visibility_Lambda([this]() {return bIsInLiveMode.Get() ? EVisibility::Collapsed : EVisibility::Visible; })
 		.EnableFooter(false)
 		[
-			FieldsListView.ToSharedRef()
+			ExposedEntitiesPanelContent
 		];
 
 	// Mode Switcher
@@ -401,6 +423,8 @@ void SRCPanelExposedEntitiesList::Construct(const FArguments& InArgs, URemoteCon
 			{
 				if (ActiveProtocol != NewMode.ModeId)
 				{
+					StoreListSettings();
+
 					ActiveListMode = EEntitiesListMode::Default; // Hack to trigger mode change via mode switcher.
 
 					ActiveProtocol = NewMode.ModeId;
@@ -539,6 +563,8 @@ SRCPanelExposedEntitiesList::~SRCPanelExposedEntitiesList()
 
 	UnregisterPresetDelegates();
 	UnregisterEvents();
+
+	StoreListSettings();
 }
 
 void SRCPanelExposedEntitiesList::Tick(const FGeometry& AllottedGeometry, const double InCurrentTime, const float InDeltaTime)
@@ -656,6 +682,8 @@ void SRCPanelExposedEntitiesList::SetBackendFilter(const FRCFilter& InBackendFil
 
 void SRCPanelExposedEntitiesList::RebuildListWithColumns(EEntitiesListMode InListMode)
 {
+	RecallListSettings();
+
 	// Only activate when the active mode changes to different one.
 	// Drop multiple calls to same mode change except for Protocols Mode.
 	if (ActiveListMode != InListMode && !ActiveProtocol.IsNone())
@@ -778,7 +806,7 @@ void SRCPanelExposedEntitiesList::OnPropertyIdRenamed(const FName InNewId, TShar
 		}
 	}
 
-	if (CurrentGroupType == EFieldGroupType::PropertyId)
+	if (CurrentGroupType == ERCFieldGroupType::PropertyId)
 	{
 		CreateFieldGroup();
 		OrderGroups();
@@ -835,6 +863,51 @@ FReply SRCPanelExposedEntitiesList::OnNodeDragDetected(const FGeometry& InGeomet
 	return FReply::Unhandled();
 }
 
+void SRCPanelExposedEntitiesList::StoreListSettings()
+{
+	const TSharedRef<IRCPanelExposedEntitiesListSettingsForProtocol>* SettingsPtr = bIsInProtocolsMode.Get() ?
+		FRemoteControlUIModule::Get().GetExposedEntitiesListSettingsForProtocol(ActiveProtocol) :
+		nullptr;
+
+	if (SettingsPtr)
+	{
+		FRCPanelExposedEntitiesListSettingsData Settings;
+		Settings.FieldGroupType = CurrentGroupType;
+		Settings.FieldGroupOrder = CurrentGroupSortType;
+
+		(*SettingsPtr)->OnSettingsChanged(Preset.Get(), Settings);
+	}
+	else
+	{
+		CurrentGroupType = DefaultSettings.FieldGroupType;
+		CurrentGroupSortType = DefaultSettings.FieldGroupOrder;
+	}
+}
+
+void SRCPanelExposedEntitiesList::RecallListSettings()
+{
+	const TSharedRef<IRCPanelExposedEntitiesListSettingsForProtocol>* SettingsPtr = bIsInProtocolsMode.Get() ? 
+		FRemoteControlUIModule::Get().GetExposedEntitiesListSettingsForProtocol(ActiveProtocol) : 
+		nullptr;
+
+	if (SettingsPtr)
+	{
+		const FRCPanelExposedEntitiesListSettingsData Settings = (*SettingsPtr)->GetListSettings(Preset.Get());
+		CurrentGroupType = Settings.FieldGroupType;
+		CurrentGroupSortType = Settings.FieldGroupOrder;
+	}
+	else
+	{
+		CurrentGroupType = DefaultSettings.FieldGroupType;
+		CurrentGroupSortType = DefaultSettings.FieldGroupOrder;
+	}
+
+	if (FieldsListView.IsValid())
+	{
+		OnCreateFieldGroup(CurrentGroupType);
+	}
+}
+
 void SRCPanelExposedEntitiesList::OnObjectPropertyChange(UObject* InObject, FPropertyChangedEvent& InChangeEvent)
 {
 	static constexpr EPropertyChangeType::Type TypesNeedingRefresh = EPropertyChangeType::ArrayAdd | EPropertyChangeType::ArrayClear | EPropertyChangeType::ArrayRemove | EPropertyChangeType::ValueSet | EPropertyChangeType::Duplicate;
@@ -1151,7 +1224,7 @@ FReply SRCPanelExposedEntitiesList::OnDropOnGroup(const TSharedPtr<FDragDropOper
 	checkSlow(DragTargetGroup);
 
 	// will be changed, for now as long as you have a grouping type active you will not be able to drag fields to re-order them
-	if (DragDropOperation->IsOfType<FExposedEntityDragDrop>() && CurrentGroupType == EFieldGroupType::None)
+	if (DragDropOperation->IsOfType<FExposedEntityDragDrop>() && CurrentGroupType == ERCFieldGroupType::None)
 	{
 		if (TSharedPtr<FExposedEntityDragDrop> DragDropOp = StaticCastSharedPtr<FExposedEntityDragDrop>(DragDropOperation))
 		{
@@ -1403,9 +1476,9 @@ TSharedRef<SWidget> SRCPanelExposedEntitiesList::GetGroupMenuContentWidget()
 		LOCTEXT("RCNoGroupingTooltip", "No grouping set. Show items as a flat list"),
 		FSlateIcon(),
 		FUIAction(
-		FExecuteAction::CreateSP(this, &SRCPanelExposedEntitiesList::OnCreateFieldGroup, EFieldGroupType::None),
+		FExecuteAction::CreateSP(this, &SRCPanelExposedEntitiesList::OnCreateFieldGroup, ERCFieldGroupType::None),
 		FCanExecuteAction(),
-		FIsActionChecked::CreateSPLambda(this, [this]{ return CurrentGroupType == EFieldGroupType::None; })),
+		FIsActionChecked::CreateSPLambda(this, [this]{ return CurrentGroupType == ERCFieldGroupType::None; })),
 		NAME_None,
 		EUserInterfaceActionType::RadioButton);
 
@@ -1414,9 +1487,9 @@ TSharedRef<SWidget> SRCPanelExposedEntitiesList::GetGroupMenuContentWidget()
 		LOCTEXT("RCGroupPropertyIdTooltip", "Group by Property Id"),
 		FSlateIcon(),
 		FUIAction(
-		FExecuteAction::CreateSP(this, &SRCPanelExposedEntitiesList::OnCreateFieldGroup, EFieldGroupType::PropertyId),
+		FExecuteAction::CreateSP(this, &SRCPanelExposedEntitiesList::OnCreateFieldGroup, ERCFieldGroupType::PropertyId),
 		FCanExecuteAction(),
-		FIsActionChecked::CreateSPLambda(this, [this]{ return CurrentGroupType == EFieldGroupType::PropertyId; })
+		FIsActionChecked::CreateSPLambda(this, [this]{ return CurrentGroupType == ERCFieldGroupType::PropertyId; })
 		),
 		NAME_None,
 		EUserInterfaceActionType::RadioButton);
@@ -1426,9 +1499,9 @@ TSharedRef<SWidget> SRCPanelExposedEntitiesList::GetGroupMenuContentWidget()
 		LOCTEXT("RCGroupOwnerTooltip", "Group by Owner"),
 		FSlateIcon(),
 		FUIAction(
-		FExecuteAction::CreateSP(this, &SRCPanelExposedEntitiesList::OnCreateFieldGroup, EFieldGroupType::Owner),
+		FExecuteAction::CreateSP(this, &SRCPanelExposedEntitiesList::OnCreateFieldGroup, ERCFieldGroupType::Owner),
 		FCanExecuteAction(),
-		FIsActionChecked::CreateSPLambda(this, [this]{ return CurrentGroupType == EFieldGroupType::Owner; })
+		FIsActionChecked::CreateSPLambda(this, [this]{ return CurrentGroupType == ERCFieldGroupType::Owner; })
 		),
 		NAME_None,
 		EUserInterfaceActionType::RadioButton);
@@ -1443,9 +1516,9 @@ TSharedRef<SWidget> SRCPanelExposedEntitiesList::GetGroupMenuContentWidget()
 		LOCTEXT("RCSortAscendingGroupTooltip", "Sort groups ascending"),
 		FSlateIcon(),
 		FUIAction(
-		FExecuteAction::CreateSP(this, &SRCPanelExposedEntitiesList::OnGroupOrderChanged, ERCGroupOrder::Ascending),
+		FExecuteAction::CreateSP(this, &SRCPanelExposedEntitiesList::OnGroupOrderChanged, ERCFieldGroupOrder::Ascending),
 		FCanExecuteAction(),
-		FIsActionChecked::CreateSPLambda(this, [this]{ return CurrentGroupSortType == ERCGroupOrder::Ascending; })
+		FIsActionChecked::CreateSPLambda(this, [this]{ return CurrentGroupSortType == ERCFieldGroupOrder::Ascending; })
 		),
 		NAME_None,
 		EUserInterfaceActionType::RadioButton);
@@ -1455,9 +1528,9 @@ TSharedRef<SWidget> SRCPanelExposedEntitiesList::GetGroupMenuContentWidget()
 		LOCTEXT("RCSortDescendingGroupTooltip", "Sort groups descending"),
 		FSlateIcon(),
 		FUIAction(
-		FExecuteAction::CreateSP(this, &SRCPanelExposedEntitiesList::OnGroupOrderChanged, ERCGroupOrder::Descending),
+		FExecuteAction::CreateSP(this, &SRCPanelExposedEntitiesList::OnGroupOrderChanged, ERCFieldGroupOrder::Descending),
 		FCanExecuteAction(),
-		FIsActionChecked::CreateSPLambda(this, [this]{ return CurrentGroupSortType == ERCGroupOrder::Descending; })
+		FIsActionChecked::CreateSPLambda(this, [this]{ return CurrentGroupSortType == ERCFieldGroupOrder::Descending; })
 		),
 		NAME_None,
 		EUserInterfaceActionType::RadioButton);
@@ -1467,18 +1540,20 @@ TSharedRef<SWidget> SRCPanelExposedEntitiesList::GetGroupMenuContentWidget()
 	return MenuBuilder.MakeWidget();
 }
 
-void SRCPanelExposedEntitiesList::OnCreateFieldGroup(EFieldGroupType InFieldGroupType)
+void SRCPanelExposedEntitiesList::OnCreateFieldGroup(ERCFieldGroupType InFieldGroupType)
 {
 	CurrentGroupType = InFieldGroupType;
 	CreateFieldGroup();
 	OrderGroups();
+
+	StoreListSettings();
 }
 
 void SRCPanelExposedEntitiesList::CreateFieldGroup()
 {
 	FieldEntitiesGroups.Reset();
 	TSet<FName> GroupKeys;
-	if (CurrentGroupType == EFieldGroupType::PropertyId)
+	if (CurrentGroupType == ERCFieldGroupType::PropertyId)
 	{
 		for (TSharedPtr<SRCPanelTreeNode> Entity : CachedFieldEntities)
 		{
@@ -1497,7 +1572,7 @@ void SRCPanelExposedEntitiesList::CreateFieldGroup()
 			}
 		}
 	}
-	else if (CurrentGroupType == EFieldGroupType::Owner)
+	else if (CurrentGroupType == ERCFieldGroupType::Owner)
 	{
 		for (TSharedPtr<SRCPanelTreeNode> Entity : CachedFieldEntities)
 		{
@@ -1516,7 +1591,7 @@ void SRCPanelExposedEntitiesList::CreateFieldGroup()
 		}
 	}
 
-	if (CurrentGroupType == EFieldGroupType::None)
+	if (CurrentGroupType == ERCFieldGroupType::None)
 	{
 		FieldEntities = CachedFieldEntities;
 		FieldsListView->RequestTreeRefresh();
@@ -1533,20 +1608,20 @@ void SRCPanelExposedEntitiesList::CreateFieldGroup()
 	}
 }
 
-void SRCPanelExposedEntitiesList::OnGroupOrderChanged(ERCGroupOrder InGroupOrder)
+void SRCPanelExposedEntitiesList::OnGroupOrderChanged(ERCFieldGroupOrder InGroupOrder)
 {
-	CurrentGroupSortType = CurrentGroupSortType == InGroupOrder? ERCGroupOrder::None : InGroupOrder;
+	CurrentGroupSortType = CurrentGroupSortType == InGroupOrder? ERCFieldGroupOrder::None : InGroupOrder;
 	OrderGroups();
 }
 
 void SRCPanelExposedEntitiesList::OrderGroups()
 {
-	if (CurrentGroupType == EFieldGroupType::None)
+	if (CurrentGroupType == ERCFieldGroupType::None)
 	{
 		return;
 	}
 
-	if (CurrentGroupSortType == ERCGroupOrder::Ascending)
+	if (CurrentGroupSortType == ERCFieldGroupOrder::Ascending)
 	{
 		FieldEntities.Sort([](const TSharedPtr<SRCPanelTreeNode>& InFirst, const TSharedPtr<SRCPanelTreeNode>& InSecond)
 		{
@@ -1560,7 +1635,7 @@ void SRCPanelExposedEntitiesList::OrderGroups()
 			return false;
 		});
 	}
-	else if (CurrentGroupSortType == ERCGroupOrder::Descending)
+	else if (CurrentGroupSortType == ERCFieldGroupOrder::Descending)
 	{
 		FieldEntities.Sort([](const TSharedPtr<SRCPanelTreeNode>& InFirst, const TSharedPtr<SRCPanelTreeNode>& InSecond)
 		{
@@ -1610,6 +1685,8 @@ void SRCPanelExposedEntitiesList::CreateGroupsAndSort()
 {
 	CreateFieldGroup();
 	OrderGroups();
+
+	StoreListSettings();
 }
 
 void SRCPanelExposedEntitiesList::RegisterPresetDelegates()
diff --git a/Engine/Plugins/VirtualProduction/RemoteControl/Source/RemoteControlUI/Private/UI/SRCPanelExposedEntitiesList.h b/Engine/Plugins/VirtualProduction/RemoteControl/Source/RemoteControlUI/Private/UI/SRCPanelExposedEntitiesList.h
--- a/Engine/Plugins/VirtualProduction/RemoteControl/Source/RemoteControlUI/Private/UI/SRCPanelExposedEntitiesList.h
+++ b/Engine/Plugins/VirtualProduction/RemoteControl/Source/RemoteControlUI/Private/UI/SRCPanelExposedEntitiesList.h
@@ -7,11 +7,14 @@
 #include "RemoteControlPreset.h"
 #include "SRCPanelExposedEntitiesGroup.h"
 #include "SRCPanelTreeNode.h"
+#include "UI/Filters/RCFilter.h"
+#include "UI/RCFieldGroupOrder.h"
+#include "UI/RCFieldGroupType.h"
+#include "UI/RCPanelExposedEntitiesListSettings.h"
+#include "UObject/StrongObjectPtr.h"
 #include "Widgets/DeclarativeSyntaxSupport.h"
 #include "Widgets/SCompoundWidget.h"
 #include "Widgets/Views/STreeView.h"
-#include "UI/Filters/RCFilter.h"
-#include "UObject/StrongObjectPtr.h"
 
 struct FRCPanelGroup;
 class FRCPanelWidgetRegistry;
@@ -40,14 +43,6 @@ enum class EEntitiesListMode : uint8
 	Protocols
 };
 
-/** Ordering types while grouping is active */
-enum class ERCGroupOrder
-{
-	None,
-	Ascending,
-	Descending
-};
-
 /** Holds information about a group drag and drop event  */
 struct FGroupDragEvent
 {
@@ -72,6 +67,7 @@ class SRCPanelExposedEntitiesList : public SCompoundWidget
 public:
 	SLATE_BEGIN_ARGS(SRCPanelExposedEntitiesList)
 		: _LiveMode(false)
+		, _ProtocolsMode(false)
 	{}
 		SLATE_ATTRIBUTE(bool, LiveMode)
 		SLATE_ATTRIBUTE(bool, ProtocolsMode)
@@ -190,7 +186,7 @@ private:
 	 * Called when the group type changed, if the new group type is the same as the current one it will be set to none
 	 * @param InFieldGroupType New grouping type
 	 */
-	void OnCreateFieldGroup(EFieldGroupType InFieldGroupType);
+	void OnCreateFieldGroup(ERCFieldGroupType InFieldGroupType);
 
 	/** Group fields based on the current group type (PropertyId/Owner) */
 	void CreateFieldGroup();
@@ -199,7 +195,7 @@ private:
 	 * Called when the order type changed, if the new order type is the same as the current one it will be set to none
 	 * @param InGroupOrder New group order
 	 */
-	void OnGroupOrderChanged(ERCGroupOrder InGroupOrder);
+	void OnGroupOrderChanged(ERCFieldGroupOrder InGroupOrder);
 
 	/** Order the groups based on the current order assigned (Ascending/Descending) */
 	void OrderGroups();
@@ -264,6 +260,12 @@ private:
 	/** Executed when a drag is detected, will create the Drag and Drop widget of the node(s) */
 	FReply OnNodeDragDetected(const FGeometry& InGeometry, const FPointerEvent& InPointerEvent, TSharedPtr<SRCPanelTreeNode> InNode);
 
+	/** Stores the list settings */
+	void StoreListSettings();
+
+	/** Recalls the list settings */
+	void RecallListSettings();
+
 private:
 	/** Holds the Groups list view. */
 	TSharedPtr<SListView<TSharedPtr<SRCPanelTreeNode>>> GroupsListView;
@@ -278,9 +280,9 @@ private:
 	/** Holds all the exposed entities groups. */
 	TArray<TSharedPtr<SRCPanelExposedEntitiesGroup>> ExposedEntitiesGroups;
 	/** Holds the current group type */
-	EFieldGroupType CurrentGroupType = EFieldGroupType::None;
+	ERCFieldGroupType CurrentGroupType = ERCFieldGroupType::None;
 	/** Holds the current sorting type */
-	ERCGroupOrder CurrentGroupSortType = ERCGroupOrder::None;
+	ERCFieldGroupOrder CurrentGroupSortType = ERCFieldGroupOrder::None;
 	/** Holds all the entities groups currently in the list */
 	TArray<TSharedPtr<SRCPanelExposedEntitiesGroup>> FieldEntitiesGroups;
 	/** Map of field ids to field widgets. */
@@ -333,7 +335,9 @@ private:
 	static TSet<FName> DefaultProtocolColumns;
 	/** Holds identifier of the selected group. */
 	FGuid CurrentlySelectedGroup;
-
+	/** Default settings for this entities list */
+	FRCPanelExposedEntitiesListSettingsData DefaultSettings;
+	
 	bool bRefreshRequested = false;
 	bool bRefreshEntitiesGroups = false;
 
diff --git a/Engine/Plugins/VirtualProduction/RemoteControl/Source/RemoteControlUI/Private/UI/SRCPanelExposedField.cpp b/Engine/Plugins/VirtualProduction/RemoteControl/Source/RemoteControlUI/Private/UI/SRCPanelExposedField.cpp
--- a/Engine/Plugins/VirtualProduction/RemoteControl/Source/RemoteControlUI/Private/UI/SRCPanelExposedField.cpp
+++ b/Engine/Plugins/VirtualProduction/RemoteControl/Source/RemoteControlUI/Private/UI/SRCPanelExposedField.cpp
@@ -5,6 +5,7 @@
 #include "Algo/Transform.h"
 #include "Components/MeshComponent.h"
 #include "Components/SceneComponent.h"
+#include "Editor/UnrealEdEngine.h"
 #include "EditorFontGlyphs.h"
 #include "Factories/IRCDefaultValueFactory.h"
 #include "Framework/Application/SlateApplication.h"
@@ -21,14 +22,15 @@
 #include "RemoteControlPanelStyle.h"
 #include "RemoteControlPreset.h"
 #include "RemoteControlSettings.h"
+#include "RemoteControlUIModule.h"
+#include "ScopedTransaction.h"
 #include "SRCPanelTreeNode.h"
 #include "SResetToDefaultPropertyEditor.h"
-#include "ScopedTransaction.h"
-#include "UnrealEdGlobals.h"
-#include "Editor/UnrealEdEngine.h"
 #include "Styling/AppStyle.h"
 #include "Styling/RemoteControlStyles.h"
 #include "Styling/SlateBrush.h"
+#include "UI/IRCPanelExposedEntityWidgetFactory.h"
+#include "UnrealEdGlobals.h"
 #include "UObject/Object.h"
 #include "Widgets/Input/SButton.h"
 #include "Widgets/Input/SCheckBox.h"
@@ -124,7 +126,10 @@ namespace ExposedFieldUtils
 
 TSharedPtr<SRCPanelTreeNode> SRCPanelExposedField::MakeInstance(const FGenerateWidgetArgs& Args)
 {
-	return SNew(SRCPanelExposedField, StaticCastSharedPtr<FRemoteControlField>(Args.Entity), Args.ColumnSizeData, Args.WidgetRegistry).Preset(Args.Preset).LiveMode(Args.bIsInLiveMode).HighlightText(Args.HighlightText);
+	return SNew(SRCPanelExposedField, StaticCastSharedPtr<FRemoteControlField>(Args.Entity), Args.ColumnSizeData, Args.WidgetRegistry)
+		.Preset(Args.Preset)
+		.LiveMode(Args.bIsInLiveMode)
+		.HighlightText(Args.HighlightText);
 }
 
 void SRCPanelExposedField::Construct(const FArguments& InArgs, TWeakPtr<FRemoteControlField> InField, FRCColumnSizeData InColumnSizeData, TWeakPtr<FRCPanelWidgetRegistry> InWidgetRegistry)
@@ -232,85 +237,101 @@ void SRCPanelExposedField::SetIsHovered(bool bInIsHovered)
 
 TSharedRef<SWidget> SRCPanelExposedField::GetProtocolWidget(const FName ForColumnName, const FName InProtocolName)
 {
-	if (TSharedPtr<FRemoteControlField> RCField = WeakField.Pin())
+	const TSharedPtr<FRemoteControlField> RCField = WeakField.Pin();
+	const TSharedPtr<FRemoteControlProperty> RCProperty = RCField.IsValid() ? StaticCastSharedPtr<FRemoteControlProperty>(RCField) : nullptr;
+	if (!RCProperty.IsValid())
 	{
-		if (const TSharedPtr<FRemoteControlProperty> RCProperty = StaticCastSharedPtr<FRemoteControlProperty>(RCField))
+		return SNullWidget::NullWidget;
+	}
+
+	// If a custom widget factory for this protocol and column exists, let it create the widget.
+	const FRemoteControlUIModule& RemoteControlUIModule = FModuleManager::GetModuleChecked<FRemoteControlUIModule>("RemoteControlUI");
+	const TSharedRef<IRCPanelExposedEntityWidgetFactory>* CustomFactoryPtr = RemoteControlUIModule.GetExposedEntityWidgetFactory(ForColumnName, InProtocolName);
+	if (CustomFactoryPtr)
+	{
+		const FRCPanelExposedPropertyWidgetArgs Args(
+			Preset,
+			RCProperty.ToSharedRef()
+		);
+
+		return (*CustomFactoryPtr)->MakePropertyWidget(Args);
+	}
+	
+	for (const FRemoteControlProtocolBinding& RCProtocolIter : RCProperty->ProtocolBindings)
+	{
+		if (RCProtocolIter.GetProtocolName() == InProtocolName)
 		{
-			for (const FRemoteControlProtocolBinding& RCProtocolIter : RCProperty->ProtocolBindings)
+			if (TSharedPtr<TStructOnScope<FRemoteControlProtocolEntity>> RCProtocolEntityPtr = RCProtocolIter.GetRemoteControlProtocolEntityPtr())
 			{
-				if (RCProtocolIter.GetProtocolName() == InProtocolName)
+				// Create a binding status widget
+				if (ForColumnName == RemoteControlPresetColumns::BindingStatus)
 				{
-					if (TSharedPtr<TStructOnScope<FRemoteControlProtocolEntity>> RCProtocolEntityPtr = RCProtocolIter.GetRemoteControlProtocolEntityPtr())
-					{
-						if (ForColumnName == RemoteControlPresetColumns::BindingStatus)
-						{
-							TSharedPtr<SWidget> BindingStatusWidget = SNew(SButton)
-								.ButtonStyle(FAppStyle::Get(), "HoverHintOnly")
-								.ToolTipText(LOCTEXT("RecordingButtonToolTip", "Status of the protocol entity binding"))
-								.ForegroundColor(FSlateColor::UseForeground())
-								.OnClicked_Lambda([RCProtocolEntityPtr]()
-									{
-										const ERCBindingStatus BindingStatus = (*RCProtocolEntityPtr)->ToggleBindingStatus();
+					TSharedPtr<SWidget> BindingStatusWidget = SNew(SButton)
+						.ButtonStyle(FAppStyle::Get(), "HoverHintOnly")
+						.ToolTipText(LOCTEXT("RecordingButtonToolTip", "Status of the protocol entity binding"))
+						.ForegroundColor(FSlateColor::UseForeground())
+						.OnClicked_Lambda([RCProtocolEntityPtr]()
+							{
+								const ERCBindingStatus BindingStatus = (*RCProtocolEntityPtr)->ToggleBindingStatus();
 
-										IRemoteControlProtocolWidgetsModule& RCProtocolWidgetsModule = IRemoteControlProtocolWidgetsModule::Get();
+								IRemoteControlProtocolWidgetsModule& RCProtocolWidgetsModule = IRemoteControlProtocolWidgetsModule::Get();
 
-										if (TSharedPtr<IRCProtocolBindingList> RCProtocolBindingList = RCProtocolWidgetsModule.GetProtocolBindingList())
-										{
-											if (BindingStatus == ERCBindingStatus::Awaiting)
-											{
-												RCProtocolBindingList->OnStartRecording(RCProtocolEntityPtr);
-											}
-											else if (BindingStatus == ERCBindingStatus::Bound)
-											{
-												RCProtocolBindingList->OnStopRecording(RCProtocolEntityPtr);
-											}
-											else
-											{
-												checkNoEntry();
-											}
-										}
-
-										return FReply::Handled();
+								if (TSharedPtr<IRCProtocolBindingList> RCProtocolBindingList = RCProtocolWidgetsModule.GetProtocolBindingList())
+								{
+									if (BindingStatus == ERCBindingStatus::Awaiting)
+									{
+										RCProtocolBindingList->OnStartRecording(RCProtocolEntityPtr);
 									}
-								)
-								.Content()
-								[
-									SNew(SImage)
-									.ColorAndOpacity_Lambda([RCProtocolEntityPtr]()
-										{
-											const ERCBindingStatus ActiveBindingStatus = (*RCProtocolEntityPtr)->GetBindingStatus();
-
-											switch (ActiveBindingStatus)
-											{
-												case ERCBindingStatus::Awaiting:
-													return FLinearColor::Red;
-												case ERCBindingStatus::Bound:
-													return FLinearColor::Green;
-												case ERCBindingStatus::Unassigned:
-													return FLinearColor::Gray;
-												default:
-													checkNoEntry();
-											}
-
-											return FLinearColor::Black;
-										}
-									)
-									.Image(FAppStyle::Get().GetBrush(TEXT("Icons.FilledCircle")))
-								];
-
-							return BindingStatusWidget.ToSharedRef();
-						}
-						else if (TSharedPtr<FRCPanelWidgetRegistry> Registry = WidgetRegistry.Pin())
-						{
-							if (TSharedPtr<IDetailTreeNode> Node = Registry->GetStructTreeNode(RCProtocolEntityPtr, (*RCProtocolEntityPtr)->GetPropertyName(ForColumnName).ToString(), ERCFindNodeMethod::Name))
-							{
-								FNodeWidgets NodeWidgets = Node->CreateNodeWidgets();
+									else if (BindingStatus == ERCBindingStatus::Bound)
+									{
+										RCProtocolBindingList->OnStopRecording(RCProtocolEntityPtr);
+									}
+									else
+									{
+										checkNoEntry();
+									}
+								}
 
-								if (NodeWidgets.ValueWidget)
+								return FReply::Handled();
+							}
+						)
+						.Content()
+						[
+							SNew(SImage)
+							.ColorAndOpacity_Lambda([RCProtocolEntityPtr]()
 								{
-									return NodeWidgets.ValueWidget.ToSharedRef();
+									const ERCBindingStatus ActiveBindingStatus = (*RCProtocolEntityPtr)->GetBindingStatus();
+
+									switch (ActiveBindingStatus)
+									{
+										case ERCBindingStatus::Awaiting:
+											return FLinearColor::Red;
+										case ERCBindingStatus::Bound:
+											return FLinearColor::Green;
+										case ERCBindingStatus::Unassigned:
+											return FLinearColor::Gray;
+										default:
+											checkNoEntry();
+									}
+
+									return FLinearColor::Black;
 								}
-							}
+							)
+							.Image(FAppStyle::Get().GetBrush(TEXT("Icons.FilledCircle")))
+						];
+
+					return BindingStatusWidget.ToSharedRef();
+				}
+				else if (TSharedPtr<FRCPanelWidgetRegistry> Registry = WidgetRegistry.Pin())
+				{
+					// Create a common property widget
+					if (TSharedPtr<IDetailTreeNode> Node = Registry->GetStructTreeNode(RCProtocolEntityPtr, (*RCProtocolEntityPtr)->GetPropertyName(ForColumnName).ToString(), ERCFindNodeMethod::Name))
+					{
+						FNodeWidgets NodeWidgets = Node->CreateNodeWidgets();
+
+						if (NodeWidgets.ValueWidget)
+						{
+							return NodeWidgets.ValueWidget.ToSharedRef();
 						}
 					}
 				}
diff --git a/Engine/Plugins/VirtualProduction/RemoteControl/Source/RemoteControlUI/Public/IRemoteControlUIModule.h b/Engine/Plugins/VirtualProduction/RemoteControl/Source/RemoteControlUI/Public/IRemoteControlUIModule.h
--- a/Engine/Plugins/VirtualProduction/RemoteControl/Source/RemoteControlUI/Public/IRemoteControlUIModule.h
+++ b/Engine/Plugins/VirtualProduction/RemoteControl/Source/RemoteControlUI/Public/IRemoteControlUIModule.h
@@ -16,6 +16,10 @@ class FRCPanelWidgetRegistry;
 class IDetailCategoryBuilder;
 class IDetailLayoutBuilder;
 class IPropertyHandle;
+class IRCPanelExposedEntitiesGroupWidgetFactory;
+class IRCPanelExposedEntitiesListSettingsForProtocol;
+class IRCExposedEntitiesPanelExtender;
+class IRCPanelExposedEntityWidgetFactory;
 class IRCSignatureCustomization;
 class URCVirtualPropertyBase;
 class URemoteControlPreset;
@@ -168,8 +172,7 @@ DECLARE_MULTICAST_DELEGATE_OneParam(FOnAddControllerExtensionColumn, TArray<FNam
  */
 class IRemoteControlUIModule : public IModuleInterface
 {
-public:
-	
+public:	
 	/**
 	 * Singleton-like access to this module's interface.  This is just for convenience!
 	 * Beware of calling this during the shutdown phase, though.  Your module might have been unloaded already.
@@ -237,16 +240,64 @@ public:
 
 	/**
 	 * Register the given signature customization to extend the Signature Panel capabilities like Drag/Drop handling
-	 * @param InCustomization the signature customization instance to register
+	 * @param InCustomization	The signature customization instance to register
 	 */
 	virtual void RegisterSignatureCustomization(const TSharedPtr<IRCSignatureCustomization>& InCustomization) = 0;
 
 	/**
 	 * Unregisters the given signature customization
-	 * @param InCustomization the signature customization instance to unregister
+	 * @param InCustomization	The signature customization instance to unregister
 	 */
 	virtual void UnregisterSignatureCustomization(const TSharedPtr<IRCSignatureCustomization>& InCustomization) = 0;
 
+	/**
+	 * Registers the given settings for an exposed entities list.
+	 * @param InSettings		The exposed entities list settings instance to register.
+	 */
+	virtual void RegisterExposedEntitiesListSettingsForProtocol(const TSharedRef<IRCPanelExposedEntitiesListSettingsForProtocol>& InSettings) = 0;
+	
+	/**
+	 * Unregisters the given settings for an exposed entities list.
+	 * @param InSettings		The exposed entities list settings instance to unregister.
+	 */
+	virtual void UnregisterExposedEntitiesListSettingsForProtocol(const TSharedRef<IRCPanelExposedEntitiesListSettingsForProtocol>& InSettings) = 0;
+
+	/**
+	 * Registers the given customization for an exposed entities panel.
+	 * @param InExtender		The exposed entities panel extender instance to register.
+	 */
+	virtual void RegisterExposedEntitiesPanelExtender(const TSharedRef<IRCExposedEntitiesPanelExtender>& InExtender) = 0;
+
+	/**
+	 * Unregisters the given exposed entities panel customization.
+	 * @param InExtender		The exposed entities panel extender instance to unregister.
+	 */
+	virtual void UnregisterExposedEntitiesPanelExtender(const TSharedRef<IRCExposedEntitiesPanelExtender>& InExtender) = 0;
+
+	/**
+	 * Registers a widget factory a column in an exposed entities group row.
+	 * @param InFactory			The widget factory to register.
+	 */
+	virtual void RegisterExposedEntitiesGroupWidgetFactory(const TSharedRef<IRCPanelExposedEntitiesGroupWidgetFactory>& InFactory) = 0;
+
+	/**
+	 * Unregisters the given widget factory.
+	 * @param InFactory			The factory to unregister.
+	 */
+	virtual void UnregisterExposedEntitiesGroupWidgetFactory(const TSharedRef<IRCPanelExposedEntitiesGroupWidgetFactory>& InFactory) = 0;
+
+	/**
+	 * Registers a widget factory a column in an exposed entity row.
+	 * @param InFactory			The widget factory to register.
+	 */
+	virtual void RegisterExposedEntityWidgetFactory(const TSharedRef<IRCPanelExposedEntityWidgetFactory>& InFactory) = 0;
+
+	/**
+	 * Unregisters the given customization for a exposed entity.
+	 * @param InFactory			The factory to unregister.
+	 */
+	virtual void UnregisterExposedEntityWidgetFactory(const TSharedRef<IRCPanelExposedEntityWidgetFactory>& InFactory) = 0;
+
 	/**
 	 * Register a widget factory to handle creating a widget in the control panel.
 	 */
diff --git a/Engine/Plugins/VirtualProduction/RemoteControl/Source/RemoteControlUI/Public/UI/IRCExposedEntitiesPanelExtender.h b/Engine/Plugins/VirtualProduction/RemoteControl/Source/RemoteControlUI/Public/UI/IRCExposedEntitiesPanelExtender.h
new file mode 100644
--- /dev/null
+++ b/Engine/Plugins/VirtualProduction/RemoteControl/Source/RemoteControlUI/Public/UI/IRCExposedEntitiesPanelExtender.h
@@ -0,0 +1,23 @@
+// Copyright Epic Games, Inc. All Rights Reserved.
+
+#pragma once
+
+#include "Misc/Attribute.h"
+#include "Templates/SharedPointer.h"
+#include "UObject/NameTypes.h"
+
+class SWidget;
+class URemoteControlPreset;
+
+/** An externder for the exposed entities panel. Draws a widget above the exposed entities list. */
+class IRCExposedEntitiesPanelExtender
+{
+public:
+	struct FArgs
+	{
+		TAttribute<FName> ActiveProtocolAttribute;
+	};
+
+	/** Creates a widget above the exposed entities list. */
+	virtual TSharedRef<SWidget> MakeWidget(URemoteControlPreset* Preset, const FArgs& Args) = 0;
+};
diff --git a/Engine/Plugins/VirtualProduction/RemoteControl/Source/RemoteControlUI/Public/UI/IRCPanelExposedEntitiesGroupWidgetFactory.h b/Engine/Plugins/VirtualProduction/RemoteControl/Source/RemoteControlUI/Public/UI/IRCPanelExposedEntitiesGroupWidgetFactory.h
new file mode 100644
--- /dev/null
+++ b/Engine/Plugins/VirtualProduction/RemoteControl/Source/RemoteControlUI/Public/UI/IRCPanelExposedEntitiesGroupWidgetFactory.h
@@ -0,0 +1,43 @@
+// Copyright Epic Games, Inc. All Rights Reserved.
+
+#pragma once
+
+#include "Templates/SharedPointer.h"
+#include "UI/RCFieldGroupType.h"
+#include "UObject/NameTypes.h"
+#include "UObject/WeakObjectPtr.h"
+
+class SWidget;
+class URemoteControlPreset;
+struct FRemoteControlProperty;
+
+/** Arguments provided by the exposed entities group widget factory */
+struct FRCPanelExposedEntitiesGroupWidgetFactoryArgs
+{
+	FRCPanelExposedEntitiesGroupWidgetFactoryArgs(
+		const TWeakObjectPtr<URemoteControlPreset> InWeakPreset,
+		const TArray<TSharedRef<FRemoteControlProperty>> InChildProperties)
+		: WeakPreset(InWeakPreset)
+		, ChildProperties(InChildProperties)
+	{}
+
+	/** The preset for which this exposed entities group widget is painted */
+	TWeakObjectPtr<URemoteControlPreset> WeakPreset;
+
+	/** The child properties that reside in this exposed entities group. */
+	TArray<TSharedRef<FRemoteControlProperty>> ChildProperties;
+};
+
+/** Factory for a widget in a column in an exposed entities group row */
+class IRCPanelExposedEntitiesGroupWidgetFactory
+{
+public:
+	/** Creates a widget in a column in an exposed entities group row. */
+	virtual TSharedRef<SWidget> MakeWidget(const FRCPanelExposedEntitiesGroupWidgetFactoryArgs& Args) = 0;
+
+	/** Returns the column in which the widget should be painted */
+	virtual FName GetColumnName() const = 0;
+
+	/** Returns the protocol name, or name none if the widget factory is not protocol specific. */
+	virtual FName GetProtocolName() const { return NAME_None; }
+};
diff --git a/Engine/Plugins/VirtualProduction/RemoteControl/Source/RemoteControlUI/Public/UI/IRCPanelExposedEntitiesListSettingsForProtocol.h b/Engine/Plugins/VirtualProduction/RemoteControl/Source/RemoteControlUI/Public/UI/IRCPanelExposedEntitiesListSettingsForProtocol.h
new file mode 100644
--- /dev/null
+++ b/Engine/Plugins/VirtualProduction/RemoteControl/Source/RemoteControlUI/Public/UI/IRCPanelExposedEntitiesListSettingsForProtocol.h
@@ -0,0 +1,27 @@
+// Copyright Epic Games, Inc. All Rights Reserved.
+
+#pragma once
+
+#include "UI/RCPanelExposedEntitiesListSettings.h"
+#include "UObject/NameTypes.h"
+
+class URemoteControlPreset;
+
+/** Interface to store and recall settings for an entity list for a specific protocol */
+class IRCPanelExposedEntitiesListSettingsForProtocol
+{
+public:
+	/** Returns the protocol name */
+	virtual FName GetProtocolName() const = 0;
+
+	/** Returns the list settings that should be applied to the exposed entities list */
+	virtual FRCPanelExposedEntitiesListSettingsData GetListSettings(URemoteControlPreset* Preset) const = 0;
+
+	/**
+	 * Called when the corresponding entity list changed. Useful to store the current settings.
+	 *
+	 * @param Preset				The preset for which the list is painted.
+	 * @param ListSettings			The current list settings.
+	 */
+	virtual void OnSettingsChanged(URemoteControlPreset* Preset, const FRCPanelExposedEntitiesListSettingsData& ListSettings) = 0;
+};
diff --git a/Engine/Plugins/VirtualProduction/RemoteControl/Source/RemoteControlUI/Public/UI/IRCPanelExposedEntityWidgetFactory.h b/Engine/Plugins/VirtualProduction/RemoteControl/Source/RemoteControlUI/Public/UI/IRCPanelExposedEntityWidgetFactory.h
new file mode 100644
--- /dev/null
+++ b/Engine/Plugins/VirtualProduction/RemoteControl/Source/RemoteControlUI/Public/UI/IRCPanelExposedEntityWidgetFactory.h
@@ -0,0 +1,43 @@
+// Copyright Epic Games, Inc. All Rights Reserved.
+
+#pragma once
+
+#include "Templates/SharedPointer.h"
+#include "UObject/NameTypes.h"
+#include "UObject/WeakObjectPtr.h"
+
+class SWidget;
+class URemoteControlPreset;
+enum class ERCFieldGroupType : uint8;
+struct FRemoteControlProperty;
+
+/** Arguments to construct an exposed property widget */
+struct FRCPanelExposedPropertyWidgetArgs
+{
+	FRCPanelExposedPropertyWidgetArgs(
+		const TWeakObjectPtr<URemoteControlPreset> InWeakPreset,
+		const TSharedRef<FRemoteControlProperty> InProperty)
+		: WeakPreset(InWeakPreset)
+		, Property(InProperty)
+	{}
+
+	/** The preset for which this exposed entities group widget is painted */
+	const TWeakObjectPtr<URemoteControlPreset> WeakPreset;
+
+	/** The property of shown in this widget. */
+	const TSharedRef<FRemoteControlProperty> Property;
+};
+
+/** Factory for a widget in a column of an exposed entity row */
+class IRCPanelExposedEntityWidgetFactory
+{
+public:
+	/** Creates a widget for a property in a column of an exposed entity row */
+	virtual TSharedRef<SWidget> MakePropertyWidget(const FRCPanelExposedPropertyWidgetArgs& Args) = 0;
+
+	/** Returns the column in which the widget should be drawn */
+	virtual FName GetColumnName() const = 0;
+
+	/** Returns the protocol name, or name none if the widget factory is not protocol specific. */
+	virtual FName GetProtocolName() const { return NAME_None; }
+};
