Remote Control: Added Signature Protocol Action

diff --git a/Engine/Plugins/VirtualProduction/RemoteControl/Source/RemoteControl/Private/RemoteControlModule.cpp b/Engine/Plugins/VirtualProduction/RemoteControl/Source/RemoteControl/Private/RemoteControlModule.cpp
--- a/Engine/Plugins/VirtualProduction/RemoteControl/Source/RemoteControl/Private/RemoteControlModule.cpp
+++ b/Engine/Plugins/VirtualProduction/RemoteControl/Source/RemoteControl/Private/RemoteControlModule.cpp
@@ -1077,13 +1077,17 @@ void FRemoteControlModule::UnregisterMaskingFactoryForType(UScriptStruct* Remote
 	MaskingFactories.Remove(RemoteControlPropertyType);
 }
 
+bool FRemoteControlModule::SupportsMasking(const UScriptStruct* InStruct) const
+{
+	return MaskingFactories.Contains(InStruct);
+}
+
 bool FRemoteControlModule::SupportsMasking(const FProperty* InProperty) const
 {
 	if (const FStructProperty* StructProperty = CastField<FStructProperty>(InProperty))
 	{
-		return MaskingFactories.Contains(StructProperty->Struct);
+		return SupportsMasking(StructProperty->Struct);
 	}
-
 	return false;
 }
 
diff --git a/Engine/Plugins/VirtualProduction/RemoteControl/Source/RemoteControl/Private/RemoteControlModule.h b/Engine/Plugins/VirtualProduction/RemoteControl/Source/RemoteControl/Private/RemoteControlModule.h
--- a/Engine/Plugins/VirtualProduction/RemoteControl/Source/RemoteControl/Private/RemoteControlModule.h
+++ b/Engine/Plugins/VirtualProduction/RemoteControl/Source/RemoteControl/Private/RemoteControlModule.h
@@ -38,6 +38,7 @@ public:
 	virtual void PerformMasking(const TSharedRef<FRCMaskingOperation>& InMaskingOperation) override;
 	virtual void RegisterMaskingFactoryForType(UScriptStruct* RemoteControlPropertyType, const TSharedPtr<IRemoteControlMaskingFactory>& InMaskingFactory) override;
 	virtual void UnregisterMaskingFactoryForType(UScriptStruct* RemoteControlPropertyType) override;
+	virtual bool SupportsMasking(const UScriptStruct* InStruct) const override;
 	virtual bool SupportsMasking(const FProperty* InProperty) const override;
 	virtual bool ResolveCall(const FString& ObjectPath, const FString& FunctionName, FRCCallReference& OutCallRef, FString* OutErrorText) override;
 	virtual bool InvokeCall(FRCCall& InCall, ERCPayloadType InPayloadType = ERCPayloadType::Json, const TArray<uint8>& InInterceptPayload = TArray<uint8>()) override;
diff --git a/Engine/Plugins/VirtualProduction/RemoteControl/Source/RemoteControl/Public/IRemoteControlModule.h b/Engine/Plugins/VirtualProduction/RemoteControl/Source/RemoteControl/Public/IRemoteControlModule.h
--- a/Engine/Plugins/VirtualProduction/RemoteControl/Source/RemoteControl/Public/IRemoteControlModule.h
+++ b/Engine/Plugins/VirtualProduction/RemoteControl/Source/RemoteControl/Public/IRemoteControlModule.h
@@ -328,6 +328,11 @@ public:
 	 */
 	virtual void UnregisterMaskingFactoryForType(UScriptStruct* RemoteControlPropertyType) = 0;
 
+	/**
+	 * Returns true if the given property can be masked, false otherwise.
+	 */
+	virtual bool SupportsMasking(const UScriptStruct* InStruct) const = 0;
+
 	/**
 	 * Returns true if the given property can be masked, false otherwise.
 	 */
diff --git a/Engine/Plugins/VirtualProduction/RemoteControl/Source/RemoteControlProtocol/Private/Signature/RCSignatureProtocolAction.cpp b/Engine/Plugins/VirtualProduction/RemoteControl/Source/RemoteControlProtocol/Private/Signature/RCSignatureProtocolAction.cpp
new file mode 100644
--- /dev/null
+++ b/Engine/Plugins/VirtualProduction/RemoteControl/Source/RemoteControlProtocol/Private/Signature/RCSignatureProtocolAction.cpp
@@ -0,0 +1,414 @@
+// Copyright Epic Games, Inc. All Rights Reserved.
+
+#include "Signature/RCSignatureProtocolAction.h"
+#include "IRemoteControlPropertyHandle.h"
+#include "IRemoteControlProtocolModule.h"
+#include "RCSignature.h"
+#include "RemoteControlField.h"
+#include "RemoteControlProtocolBinding.h"
+
+namespace UE::RemoteControlProtocol::Private
+{
+	uint8 GetPropertyDimension(const UScriptStruct* InPropertyStruct)
+	{
+		if (!InPropertyStruct)
+		{
+			return 1;
+		}
+
+		uint8 Dimensions = 0;
+		const FProperty* Property = InPropertyStruct->PropertyLink;
+
+		while (Property)
+		{
+			++Dimensions;
+			Property = Property->PropertyLinkNext;
+		}
+		return Dimensions;
+	}
+
+	// Logic taken from FProtocolBindingViewModel::AddDefaultRangeMappings.
+	// TODO: Move the logic for re-use
+	void SetMappingRange(const FRemoteControlProtocolEntity& InEntity
+		, const IRemoteControlProtocol& InProtocol
+		, FRemoteControlProtocolMapping& InOutMinMapping
+		, FRemoteControlProtocolMapping& InOutMaxMapping)
+	{
+		const FNumericProperty* NumericProperty = CastField<FNumericProperty>(InProtocol.GetRangeInputTemplateProperty());
+		if (!NumericProperty)
+		{
+			return;
+		}
+
+		FName RangePropertyTypeName = InEntity.GetRangePropertyName();
+		const uint8 RangePropertySize = InEntity.GetRangePropertySize();
+
+		if (NumericProperty->IsInteger())
+		{
+			int64 IntMin = RemoteControlTypeUtilities::GetDefaultRangeValueMin<int64>(NumericProperty);
+			int64 IntMax = RemoteControlTypeUtilities::GetDefaultRangeValueMax<int64>(NumericProperty);
+			// fixup typename according to typesize, ie. it can be a UInt32 but a typesize of 2 makes its a UInt16
+			if (RangePropertyTypeName == NAME_UInt32Property && RangePropertySize > 0)
+			{
+				if (RangePropertySize == sizeof(uint8))
+				{
+					IntMin = RemoteControlTypeUtilities::GetDefaultRangeValueMin<uint8>(NumericProperty);
+					IntMax = RemoteControlTypeUtilities::GetDefaultRangeValueMax<uint8>(NumericProperty);
+					RangePropertyTypeName = NAME_ByteProperty;
+				}
+				else if (RangePropertySize == sizeof(uint16))
+				{
+					IntMin = RemoteControlTypeUtilities::GetDefaultRangeValueMin<uint16>(NumericProperty);
+					IntMax = RemoteControlTypeUtilities::GetDefaultRangeValueMax<uint16>(NumericProperty);
+					RangePropertyTypeName = NAME_UInt16Property; 
+				}
+				else if (RangePropertySize == sizeof(uint64))
+				{
+					IntMin = RemoteControlTypeUtilities::GetDefaultRangeValueMin<uint64>(NumericProperty);
+					IntMax = RemoteControlTypeUtilities::GetDefaultRangeValueMax<uint64>(NumericProperty);
+					RangePropertyTypeName = NAME_UInt64Property;
+				}
+			}
+			if (RangePropertyTypeName == NAME_ByteProperty)
+			{
+				InOutMinMapping.SetRangeValue<uint8>(IntMin);
+				InOutMaxMapping.SetRangeValue<uint8>(IntMax);
+			}
+			else if (RangePropertyTypeName == NAME_UInt16Property)
+			{
+				InOutMinMapping.SetRangeValue<uint16>(IntMin);
+				InOutMaxMapping.SetRangeValue<uint16>(IntMax);
+			}
+			else if (RangePropertyTypeName == NAME_UInt32Property)
+			{
+				InOutMinMapping.SetRangeValue<uint32>(IntMin);
+				InOutMaxMapping.SetRangeValue<uint32>(IntMax);
+			}
+			else if (RangePropertyTypeName == NAME_UInt64Property)
+			{
+				InOutMinMapping.SetRangeValue<uint64>(IntMin);
+				InOutMaxMapping.SetRangeValue<uint64>(IntMax);
+			}
+		}
+		else if(NumericProperty->IsFloatingPoint())
+		{
+			const float FloatMin = RemoteControlTypeUtilities::GetDefaultRangeValueMin<float>(NumericProperty);
+			const float FloatMax = RemoteControlTypeUtilities::GetDefaultRangeValueMax<float>(NumericProperty);
+			if (RangePropertyTypeName == NAME_FloatProperty)
+			{
+				InOutMinMapping.SetRangeValue<float>(FloatMin);
+				InOutMaxMapping.SetRangeValue<float>(FloatMax);
+			}
+		}
+	}
+
+	template<typename InCppType>
+	struct FRelativeOperation
+	{
+		static InCppType CalculateAdditive(const InCppType& InPropertyValue, const InCppType& InMappingValue)
+		{
+			return InPropertyValue + InMappingValue;
+		}
+
+		static InCppType CalculateScalar(const InCppType& InPropertyValue, const InCppType& InMappingValue)
+		{
+			return InPropertyValue * InMappingValue;
+		}
+	};
+
+	template<>
+	struct FRelativeOperation<FRotator>
+	{
+		static FRotator CalculateAdditive(const FRotator& InPropertyValue, const FRotator& InMappingValue)
+		{
+			return InPropertyValue + InMappingValue;
+		}
+
+		// FRotators have no Scalar support. Treat this as Absolute
+		static FRotator CalculateScalar(const FRotator& InPropertyValue, const FRotator& InMappingValue)
+		{
+			return InMappingValue;
+		}
+	};
+
+	template<>
+	struct FRelativeOperation<FColor>
+	{
+		static FColor CalculateAdditive(const FColor& InPropertyValue, const FColor& InMappingValue)
+		{
+			return FRelativeOperation<FLinearColor>::CalculateAdditive(InPropertyValue, InMappingValue)
+				.ToFColor(/*bSRGB*/true);
+		}
+
+		static FColor CalculateScalar(const FColor& InPropertyValue, const FColor& InMappingValue)
+		{
+			return FRelativeOperation<FLinearColor>::CalculateScalar(InPropertyValue, InMappingValue)
+				.ToFColor(/*bSRGB*/true);
+		}
+	};
+
+	struct FMappingTypeHelper
+	{
+		FMappingTypeHelper(const FRCSignatureActionContext& InContext, ERCSignatureProtocolActionMappingSpace InMappingType)
+			: PropertyHandle(InContext.Property.IsValid() ? InContext.Property->GetPropertyHandle() : nullptr)
+			, MappingType(InMappingType)
+		{
+		}
+
+		bool TryApply(FRemoteControlProtocolMapping& InOutMapping)
+		{
+			if (!PropertyHandle.IsValid())
+			{
+				return false;
+			}
+
+			// Nothing to do with an Absolute Mapping
+			if (MappingType == ERCSignatureProtocolActionMappingSpace::Absolute)
+			{
+				return false;
+			}
+
+			return TryApply<double>(InOutMapping)
+				|| TryApply<float>(InOutMapping)
+				|| TryApply<FVector>(InOutMapping)
+				|| TryApply<FVector4>(InOutMapping)
+				|| TryApply<FRotator>(InOutMapping)
+				|| TryApply<FColor>(InOutMapping)
+				|| TryApply<FLinearColor>(InOutMapping)
+				|| TryApply<int8>(InOutMapping)
+				|| TryApply<uint8>(InOutMapping)
+				|| TryApply<int32>(InOutMapping)
+				|| TryApply<uint32>(InOutMapping)
+				|| TryApply<int16>(InOutMapping)
+				|| TryApply<uint16>(InOutMapping)
+				|| TryApply<int64>(InOutMapping)
+				|| TryApply<uint64>(InOutMapping);
+		}
+
+	private:
+		template<typename InCppType>
+		bool TryApply(FRemoteControlProtocolMapping& InOutMapping) const
+		{
+			InCppType PropertyValue;
+			if (!PropertyHandle->GetValue(PropertyValue))
+			{
+				return false;
+			}
+
+			const InCppType MappingValue = InOutMapping.GetMappingValueAsPrimitive<InCppType>();
+
+			switch (MappingType)
+			{
+			case ERCSignatureProtocolActionMappingSpace::Additive:
+				InOutMapping.SetMappingValueAsPrimitive(FRelativeOperation<InCppType>::CalculateAdditive(PropertyValue, MappingValue));
+				break;
+
+			case ERCSignatureProtocolActionMappingSpace::Scalar:
+				InOutMapping.SetMappingValueAsPrimitive(FRelativeOperation<InCppType>::CalculateScalar(PropertyValue, MappingValue));
+				break;
+
+			default:
+				checkNoEntry();
+			}
+
+			return true;
+		}
+
+		TSharedPtr<IRemoteControlPropertyHandle> PropertyHandle;
+		ERCSignatureProtocolActionMappingSpace MappingType;
+	};
+}
+
+void FRCSignatureProtocolAction::Initialize(const FRCSignatureField& InField)
+{
+	MinMappingDesc = InField.PropertyDesc;
+	MaxMappingDesc = InField.PropertyDesc;
+
+	MinMappingDesc.Name = TEXT("Min");
+	MaxMappingDesc.Name = TEXT("Max");
+
+#if WITH_EDITORONLY_DATA
+	// Reset MetaData that could affect Display Name, or Min/Max clamping.
+	// The mapping descs are meant to just match the type
+	MinMappingDesc.MetaData.Reset();
+	MinMappingDesc.MetaClass = nullptr;
+
+	MaxMappingDesc.MetaData.Reset();
+	MaxMappingDesc.MetaClass = nullptr;
+#endif
+
+	PropertyDimension = UE::RemoteControlProtocol::Private::GetPropertyDimension(GetPropertyStruct());
+
+	TArray<FName> ProtocolNames = IRemoteControlProtocolModule::Get().GetProtocolNames();
+	if (!ProtocolNames.IsEmpty())
+	{
+		if (ProtocolName.IsNone())
+		{
+			ProtocolName = ProtocolNames[0];
+		}
+		UpdateProtocolEntity();
+	}
+
+	UpdateMappingType();
+}
+
+bool FRCSignatureProtocolAction::IsSupported(const FRCSignatureField& InField) const
+{
+	return InField.PropertyDesc.ValueType == EPropertyBagPropertyType::Struct
+		|| InField.PropertyDesc.IsNumericType();
+}
+
+bool FRCSignatureProtocolAction::Execute(const FRCSignatureActionContext& InContext) const
+{
+	TSharedPtr<IRemoteControlProtocol> Protocol = GetProtocol();
+	if (!Protocol.IsValid())
+	{
+		return false;
+	}
+
+	if (PropertyDimension == 1)
+	{
+		CreateProtocolEntity(InContext, *Protocol, 0xFF);
+		return true;
+	}
+
+	switch (MaskingType)
+	{
+	case ERCSignatureProtocolActionMaskingType::SingleChannel:
+		// Add Mask Single Channel per Entity instead of all Channels
+		for (uint8 Dimension = 0; Dimension < PropertyDimension; ++Dimension)
+		{
+			const uint8 Mask = 1 << Dimension;
+			if (Mask & OverrideMask)
+			{
+				CreateProtocolEntity(InContext, *Protocol, Mask);
+			}
+		}
+		break;
+
+	case ERCSignatureProtocolActionMaskingType::MultiChannel:
+		CreateProtocolEntity(InContext, *Protocol, OverrideMask);
+		break;
+
+	default:
+		break;
+	}
+	return true;
+}
+
+#if WITH_EDITOR
+void FRCSignatureProtocolAction::PostEditChange(const FPropertyChangedEvent& InPropertyChangedEvent, FEditPropertyChain* InPropertyThatChanged)
+{
+	if (InPropertyChangedEvent.GetMemberPropertyName() == GET_MEMBER_NAME_CHECKED(FRCSignatureProtocolAction, ProtocolName))
+	{
+		UpdateProtocolEntity();
+	}
+	UpdateMappingType();
+}
+
+FRCSignatureActionIcon FRCSignatureProtocolAction::GetIcon() const
+{
+	FRCSignatureActionIcon Icon;
+	Icon.StyleSetName = TEXT("EditorStyle");
+	Icon.StyleName = TEXT("LevelEditor.Tabs.StatsViewer");
+	return Icon;
+}
+#endif
+
+void FRCSignatureProtocolAction::UpdateProtocolEntity()
+{
+	const TSharedPtr<IRemoteControlProtocol> Protocol = GetProtocol();
+	if (!Protocol.IsValid())
+	{
+		return;
+	}
+
+	const UScriptStruct* ProtocolStruct = Protocol->GetProtocolScriptStruct();
+
+	check(ProtocolStruct && ProtocolStruct->IsChildOf(FRemoteControlProtocolEntity::StaticStruct()));
+	if (ProtocolStruct != ProtocolEntity.GetScriptStruct())
+	{
+		ProtocolEntity.InitializeAsScriptStruct(ProtocolStruct, /*StructMemory*/nullptr);
+	}
+}
+
+void FRCSignatureProtocolAction::UpdateMappingType()
+{
+	if (!Mappings.FindPropertyDescByName(MinMappingDesc.Name) || !Mappings.FindPropertyDescByName(MaxMappingDesc.Name))
+	{
+		Mappings.Reset();
+		Mappings.AddProperties({ MinMappingDesc, MaxMappingDesc });
+	}
+}
+
+TSharedPtr<IRemoteControlProtocol> FRCSignatureProtocolAction::GetProtocol() const
+{
+	return IRemoteControlProtocolModule::Get().GetProtocolByName(ProtocolName);
+}
+
+void FRCSignatureProtocolAction::CreateProtocolEntity(const FRCSignatureActionContext& InContext, IRemoteControlProtocol& InProtocol, uint8 InMask) const
+{
+	const FGuid EntityId = InContext.Property->GetId();
+	FProperty* Property = InContext.Property->GetProperty();
+
+	TSharedPtr<TStructOnScope<FRemoteControlProtocolEntity>> NewEntity = InProtocol.CreateNewProtocolEntity(Property
+		, InContext.Preset
+		, EntityId);
+
+	if (const UScriptStruct* ScriptStruct = Cast<UScriptStruct>(NewEntity->GetStruct()))
+	{
+		check(ProtocolEntity.GetScriptStruct() == ScriptStruct);
+		ScriptStruct->CopyScriptStruct(NewEntity->Get(), ProtocolEntity.GetPtr());
+	}
+
+	FRemoteControlProtocolBinding ProtocolBinding(ProtocolName, EntityId, NewEntity, FGuid::NewGuid());
+	AddMappings(InContext, InProtocol, /*InOut*/ProtocolBinding);
+
+	(*NewEntity)->ClearMask(RC_AllMasks);
+	(*NewEntity)->EnableMask(static_cast<ERCMask>(InMask));
+
+	InProtocol.Bind(NewEntity);
+	InContext.Property->ProtocolBindings.Emplace(MoveTemp(ProtocolBinding));
+}
+
+void FRCSignatureProtocolAction::AddMappings(const FRCSignatureActionContext& InContext
+	, const IRemoteControlProtocol& InProtocol
+	, FRemoteControlProtocolBinding& InOutBinding) const
+{
+	const FRemoteControlProtocolEntity& ProtocolEntityRef = ProtocolEntity.Get();
+
+	FProperty* Property = InContext.Property->GetProperty();
+
+	const uint8 RangePropertySize = ProtocolEntityRef.GetRangePropertySize();
+
+	FRemoteControlProtocolMapping MinMapping(Property, RangePropertySize);
+	FRemoteControlProtocolMapping MaxMapping(Property, RangePropertySize);
+
+	UE::RemoteControlProtocol::Private::SetMappingRange(ProtocolEntity.Get(), InProtocol, MinMapping, MaxMapping);
+
+	const uint8* MappingValue = Mappings.GetValue().GetMemory();
+
+	UE::RemoteControlProtocol::Private::FMappingTypeHelper MappingHelper(InContext, MappingSpace);
+
+	if (const FPropertyBagPropertyDesc* Desc = Mappings.FindPropertyDescByName(MinMappingDesc.Name))
+	{
+		const uint8* MinValue = MappingValue + Desc->CachedProperty->GetOffset_ForInternal();
+		MinMapping.SetRawMappingData(InContext.Preset, Property, MinValue);
+		MappingHelper.TryApply(MinMapping);
+	}
+
+	if (const FPropertyBagPropertyDesc* Desc = Mappings.FindPropertyDescByName(MaxMappingDesc.Name))
+	{
+		const uint8* MaxValue = MappingValue + Desc->CachedProperty->GetOffset_ForInternal();
+		MaxMapping.SetRawMappingData(InContext.Preset, Property, MaxValue);
+		MappingHelper.TryApply(MaxMapping);
+	}
+
+	InOutBinding.AddMapping(MinMapping);
+	InOutBinding.AddMapping(MaxMapping);
+}
+
+const UScriptStruct* FRCSignatureProtocolAction::GetPropertyStruct() const
+{
+	return Cast<UScriptStruct>(MinMappingDesc.ValueTypeObject);
+}
diff --git a/Engine/Plugins/VirtualProduction/RemoteControl/Source/RemoteControlProtocol/Public/Signature/RCSignatureProtocolAction.h b/Engine/Plugins/VirtualProduction/RemoteControl/Source/RemoteControlProtocol/Public/Signature/RCSignatureProtocolAction.h
new file mode 100644
--- /dev/null
+++ b/Engine/Plugins/VirtualProduction/RemoteControl/Source/RemoteControlProtocol/Public/Signature/RCSignatureProtocolAction.h
@@ -0,0 +1,82 @@
+// Copyright Epic Games, Inc. All Rights Reserved.
+
+#pragma once
+
+#include "RCSignatureAction.h"
+#include "RemoteControlProtocolBinding.h"
+#include "StructUtils/InstancedStruct.h"
+#include "StructUtils/PropertyBag.h"
+#include "RCSignatureProtocolAction.generated.h"
+
+class IRemoteControlProtocol;
+struct FRemoteControlProtocolEntity;
+
+UENUM()
+enum class ERCSignatureProtocolActionMaskingType : uint8
+{
+	SingleChannel,
+	MultiChannel,
+};
+
+UENUM()
+enum class ERCSignatureProtocolActionMappingSpace : uint8
+{
+	Additive,
+	Scalar,
+	Absolute
+};
+
+USTRUCT(DisplayName="Protocol Action")
+struct FRCSignatureProtocolAction : public FRCSignatureAction
+{
+	GENERATED_BODY()
+
+	//~ Begin FRCSignatureAction
+	virtual void Initialize(const FRCSignatureField& InField) override;
+	virtual bool IsSupported(const FRCSignatureField& InField) const override;
+	virtual bool Execute(const FRCSignatureActionContext& InContext) const override;
+#if WITH_EDITOR
+	virtual void PostEditChange(const FPropertyChangedEvent& InPropertyChangedEvent, FEditPropertyChain* InPropertyThatChanged) override;
+	virtual FRCSignatureActionIcon GetIcon() const override;
+#endif
+	//~ End FRCSignatureAction
+
+	void UpdateProtocolEntity();
+
+	void UpdateMappingType();
+
+	TSharedPtr<IRemoteControlProtocol> GetProtocol() const;
+
+	void CreateProtocolEntity(const FRCSignatureActionContext& InContext, IRemoteControlProtocol& InProtocol, uint8 InMask) const;
+
+	void AddMappings(const FRCSignatureActionContext& InContext, const IRemoteControlProtocol& InProtocol, FRemoteControlProtocolBinding& InOutBinding) const;
+
+	REMOTECONTROLPROTOCOL_API const UScriptStruct* GetPropertyStruct() const;
+
+	UPROPERTY(EditAnywhere, Category="Protocol")
+	FName ProtocolName;
+
+	UPROPERTY(EditAnywhere, Category="Protocol", meta = (StructTypeConst))
+	TInstancedStruct<FRemoteControlProtocolEntity> ProtocolEntity;
+
+	UPROPERTY(EditAnywhere, Category="Protocol", meta=(EditCondition="PropertyDimension > 1", EditConditionHides))
+	ERCSignatureProtocolActionMaskingType MaskingType = ERCSignatureProtocolActionMaskingType::SingleChannel;
+
+	UPROPERTY(EditAnywhere, Category="Protocol", meta=(EditCondition="PropertyDimension > 1", EditConditionHides))
+	uint8 OverrideMask = 0;
+
+	UPROPERTY(EditAnywhere, Category="Protocol")
+	ERCSignatureProtocolActionMappingSpace MappingSpace = ERCSignatureProtocolActionMappingSpace::Additive;
+
+	UPROPERTY(EditAnywhere, Category="Protocol", meta = (FixedLayout))
+	FInstancedPropertyBag Mappings;
+
+	UPROPERTY()
+	FPropertyBagPropertyDesc MinMappingDesc;
+
+	UPROPERTY()
+	FPropertyBagPropertyDesc MaxMappingDesc;
+
+	UPROPERTY()
+	uint8 PropertyDimension = 1;
+};
diff --git a/Engine/Plugins/VirtualProduction/RemoteControl/Source/RemoteControlProtocolWidgets/Private/Customizations/RCSignatureProtocolActionCustomization.cpp b/Engine/Plugins/VirtualProduction/RemoteControl/Source/RemoteControlProtocolWidgets/Private/Customizations/RCSignatureProtocolActionCustomization.cpp
new file mode 100644
--- /dev/null
+++ b/Engine/Plugins/VirtualProduction/RemoteControl/Source/RemoteControlProtocolWidgets/Private/Customizations/RCSignatureProtocolActionCustomization.cpp
@@ -0,0 +1,123 @@
+// Copyright Epic Games, Inc. All Rights Reserved.
+
+#include "RCSignatureProtocolActionCustomization.h"
+#include "DetailCategoryBuilder.h"
+#include "DetailLayoutBuilder.h"
+#include "DetailWidgetRow.h"
+#include "IDetailChildrenBuilder.h"
+#include "IRemoteControlProtocolModule.h"
+#include "Signature/RCSignatureProtocolAction.h"
+#include "Widgets/Input/SComboBox.h"
+#include "Widgets/Masks/SRCSignatureProtocolMaskProperty.h"
+#include "Widgets/Text/STextBlock.h"
+
+namespace UE::RemoteControlProtocolWidgets::Private
+{
+	TSharedRef<SWidget> CreateProtocolPicker(const TSharedRef<IPropertyHandle>& InProtocolNameHandle)
+	{
+		const TSharedRef<TArray<FName>> ProtocolOptions = MakeShared<TArray<FName>>();
+
+		FName CurrentValue;
+		InProtocolNameHandle->GetValue(CurrentValue);
+
+		return SNew(SComboBox<FName>)
+			.OptionsSource(&*ProtocolOptions) // ProtocolOptions is captured by copy in a lambda below, keeping it valid
+			.InitiallySelectedItem(CurrentValue)
+			.OnComboBoxOpening_Lambda([ProtocolOptions]()
+			{
+				*ProtocolOptions = IRemoteControlProtocolModule::Get().GetProtocolNames();
+			})
+			.OnGenerateWidget_Lambda([](FName InName)
+			{
+				return SNew(STextBlock)
+					.Font(IDetailLayoutBuilder::GetDetailFont())
+					.Text(FText::FromName(InName));
+			})
+			.OnSelectionChanged_Lambda([InProtocolNameHandle](FName InName, ESelectInfo::Type)
+			{
+				InProtocolNameHandle->SetValue(InName);
+			})
+			.Content()
+			[
+				SNew(STextBlock)
+				.Font(IDetailLayoutBuilder::GetDetailFont())
+				.Text_Lambda([InProtocolNameHandle]()->FText
+				{
+						FName CurrentValue;
+						InProtocolNameHandle->GetValue(CurrentValue);
+						return FText::FromName(CurrentValue);
+				})
+			];
+	}
+}
+
+TSharedRef<IDetailCustomization> FRCSignatureProtocolActionCustomization::MakeInstance()
+{
+	return MakeShared<FRCSignatureProtocolActionCustomization>();
+}
+
+void FRCSignatureProtocolActionCustomization::CustomizeDetails(IDetailLayoutBuilder& InDetailBuilder)
+{
+	const TSharedRef<IPropertyHandle> ProtocolNameHandle = InDetailBuilder.GetProperty(GET_MEMBER_NAME_CHECKED(FRCSignatureProtocolAction, ProtocolName));
+	const TSharedRef<IPropertyHandle> ProtocolEntityHandle = InDetailBuilder.GetProperty(GET_MEMBER_NAME_CHECKED(FRCSignatureProtocolAction, ProtocolEntity));
+	const TSharedRef<IPropertyHandle> MaskingTypeHandle = InDetailBuilder.GetProperty(GET_MEMBER_NAME_CHECKED(FRCSignatureProtocolAction, MaskingType));
+	const TSharedRef<IPropertyHandle> OverrideMaskHandle = InDetailBuilder.GetProperty(GET_MEMBER_NAME_CHECKED(FRCSignatureProtocolAction, OverrideMask));
+
+	InDetailBuilder.HideProperty(ProtocolNameHandle);
+	InDetailBuilder.HideProperty(ProtocolEntityHandle);
+	InDetailBuilder.HideProperty(MaskingTypeHandle);
+	InDetailBuilder.HideProperty(OverrideMaskHandle);
+
+	IDetailCategoryBuilder& ProtocolCategory = InDetailBuilder.EditCategory(TEXT("Protocol"));
+
+	ProtocolCategory.AddProperty(ProtocolNameHandle)
+		.CustomWidget()
+		.NameContent()
+		[
+			ProtocolNameHandle->CreatePropertyNameWidget()
+		]
+		.ValueContent()
+		[
+			UE::RemoteControlProtocolWidgets::Private::CreateProtocolPicker(ProtocolNameHandle)
+		];
+
+	ProtocolCategory.AddProperty(ProtocolEntityHandle);
+	ProtocolCategory.AddProperty(MaskingTypeHandle);
+
+	ProtocolCategory.AddProperty(OverrideMaskHandle)
+		.CustomWidget()
+		.NameContent()
+		[
+			OverrideMaskHandle->CreatePropertyNameWidget()
+		]
+		.ValueContent()
+		[
+			SNew(SRCSignatureProtocolMaskProperty, OverrideMaskHandle, GetProtocolActions(InDetailBuilder))
+		];
+}
+
+TArray<const FRCSignatureProtocolAction*, TInlineAllocator<1>> FRCSignatureProtocolActionCustomization::GetProtocolActions(IDetailLayoutBuilder& InDetailBuilder) const
+{
+	TArray<const FRCSignatureProtocolAction*, TInlineAllocator<1>> ProtocolActions;
+
+	const UScriptStruct* ProtocolActionStruct = FRCSignatureProtocolAction::StaticStruct();
+
+	TArray<TSharedPtr<FStructOnScope>> StructOnScopes;
+	InDetailBuilder.GetStructsBeingCustomized(StructOnScopes);
+
+	if (!StructOnScopes.IsEmpty())
+	{
+		ProtocolActions.Reserve(StructOnScopes.Num());
+		for (const TSharedPtr<FStructOnScope>& StructOnScope : StructOnScopes)
+		{
+			if (StructOnScope.IsValid())
+			{
+				const UStruct* ScriptStruct = StructOnScope->GetStruct();
+				check(ScriptStruct == ProtocolActionStruct);
+				ProtocolActions.Add(reinterpret_cast<const FRCSignatureProtocolAction*>(StructOnScope->GetStructMemory()));
+			}
+		}
+	}
+
+	return ProtocolActions;
+}
diff --git a/Engine/Plugins/VirtualProduction/RemoteControl/Source/RemoteControlProtocolWidgets/Private/Customizations/RCSignatureProtocolActionCustomization.h b/Engine/Plugins/VirtualProduction/RemoteControl/Source/RemoteControlProtocolWidgets/Private/Customizations/RCSignatureProtocolActionCustomization.h
new file mode 100644
--- /dev/null
+++ b/Engine/Plugins/VirtualProduction/RemoteControl/Source/RemoteControlProtocolWidgets/Private/Customizations/RCSignatureProtocolActionCustomization.h
@@ -0,0 +1,20 @@
+// Copyright Epic Games, Inc. All Rights Reserved.
+
+#pragma once
+
+#include "IDetailCustomization.h"
+
+struct FRCSignatureProtocolAction;
+
+class FRCSignatureProtocolActionCustomization : public IDetailCustomization
+{
+public:
+	static TSharedRef<IDetailCustomization> MakeInstance();
+
+	//~ Begin IDetailCustomization
+	virtual void CustomizeDetails(IDetailLayoutBuilder& InDetailBuilder) override;
+	//~ End IDetailCustomization
+
+private:
+	TArray<const FRCSignatureProtocolAction*, TInlineAllocator<1>> GetProtocolActions(IDetailLayoutBuilder& InDetailBuilder) const;
+};
diff --git a/Engine/Plugins/VirtualProduction/RemoteControl/Source/RemoteControlProtocolWidgets/Private/RemoteControlProtocolWidgetsModule.cpp b/Engine/Plugins/VirtualProduction/RemoteControl/Source/RemoteControlProtocolWidgets/Private/RemoteControlProtocolWidgetsModule.cpp
--- a/Engine/Plugins/VirtualProduction/RemoteControl/Source/RemoteControlProtocolWidgets/Private/RemoteControlProtocolWidgetsModule.cpp
+++ b/Engine/Plugins/VirtualProduction/RemoteControl/Source/RemoteControlProtocolWidgets/Private/RemoteControlProtocolWidgetsModule.cpp
@@ -1,13 +1,13 @@
 // Copyright Epic Games, Inc. All Rights Reserved.
 
 #include "RemoteControlProtocolWidgetsModule.h"
-
+#include "Customizations/RCSignatureProtocolActionCustomization.h"
 #include "IRemoteControlModule.h"
 #include "IRemoteControlProtocol.h"
 #include "IRemoteControlProtocolModule.h"
-#include "RemoteControlProtocolWidgetsSettings.h"
-#include "IRemoteControlModule.h"
 #include "RemoteControlPreset.h"
+#include "RemoteControlProtocolWidgetsSettings.h"
+#include "Signature/RCSignatureProtocolAction.h"
 #include "Styling/ProtocolPanelStyle.h"
 #include "ViewModels/ProtocolEntityViewModel.h"
 #include "Widgets/SNullWidget.h"
@@ -19,6 +19,8 @@ void FRemoteControlProtocolWidgetsModule::StartupModule()
 {
 	FProtocolPanelStyle::Initialize();
 
+	RegisterPropertyEditorCustomizations();
+
 	OnActiveProtocolChanged().AddRaw(this, &FRemoteControlProtocolWidgetsModule::SetActiveProtocolName);
 }
 
@@ -26,6 +28,8 @@ void FRemoteControlProtocolWidgetsModule::ShutdownModule()
 {
 	FProtocolPanelStyle::Shutdown();
 
+	UnregisterPropertyEditorCustomizations();
+
 	OnActiveProtocolChanged().RemoveAll(this);
 }
 
@@ -126,4 +130,25 @@ void FRemoteControlProtocolWidgetsModule::SetActiveProtocolName(const FName InPr
 	}
 }
 
+void FRemoteControlProtocolWidgetsModule::RegisterPropertyEditorCustomizations()
+{
+	FPropertyEditorModule& PropertyEditorModule = FModuleManager::LoadModuleChecked<FPropertyEditorModule>("PropertyEditor");
+
+	PropertyEditorModule.RegisterCustomClassLayout(FRCSignatureProtocolAction::StaticStruct()->GetFName()
+		, FOnGetDetailCustomizationInstance::CreateStatic(&FRCSignatureProtocolActionCustomization::MakeInstance));
+}
+
+void FRemoteControlProtocolWidgetsModule::UnregisterPropertyEditorCustomizations()
+{
+	if (!UObjectInitialized())
+	{
+		return;
+	}
+
+	if (FPropertyEditorModule* PropertyEditorModule = FModuleManager::GetModulePtr<FPropertyEditorModule>("PropertyEditor"))
+	{
+		PropertyEditorModule->UnregisterCustomClassLayout(FRCSignatureProtocolAction::StaticStruct()->GetFName());
+	}
+}
+
 IMPLEMENT_MODULE(FRemoteControlProtocolWidgetsModule, RemoteControlProtocolWidgets);
diff --git a/Engine/Plugins/VirtualProduction/RemoteControl/Source/RemoteControlProtocolWidgets/Private/RemoteControlProtocolWidgetsModule.h b/Engine/Plugins/VirtualProduction/RemoteControl/Source/RemoteControlProtocolWidgets/Private/RemoteControlProtocolWidgetsModule.h
--- a/Engine/Plugins/VirtualProduction/RemoteControl/Source/RemoteControlProtocolWidgets/Private/RemoteControlProtocolWidgetsModule.h
+++ b/Engine/Plugins/VirtualProduction/RemoteControl/Source/RemoteControlProtocolWidgets/Private/RemoteControlProtocolWidgetsModule.h
@@ -41,6 +41,9 @@ private:
 	/** Sets the selected protocol in the list and user setting. */
 	void SetActiveProtocolName(const FName InProtocolName);
 
+	void RegisterPropertyEditorCustomizations();
+	void UnregisterPropertyEditorCustomizations();
+
 private:
 	/** Binding list public interface instance */
 	TSharedPtr<IRCProtocolBindingList> RCProtocolBindingList;
diff --git a/Engine/Plugins/VirtualProduction/RemoteControl/Source/RemoteControlProtocolWidgets/Private/Widgets/Masks/SRCSignatureProtocolMaskProperty.cpp b/Engine/Plugins/VirtualProduction/RemoteControl/Source/RemoteControlProtocolWidgets/Private/Widgets/Masks/SRCSignatureProtocolMaskProperty.cpp
new file mode 100644
--- /dev/null
+++ b/Engine/Plugins/VirtualProduction/RemoteControl/Source/RemoteControlProtocolWidgets/Private/Widgets/Masks/SRCSignatureProtocolMaskProperty.cpp
@@ -0,0 +1,106 @@
+// Copyright Epic Games, Inc. All Rights Reserved.
+
+#include "SRCSignatureProtocolMaskProperty.h"
+
+#include "IRemoteControlModule.h"
+#include "RemoteControlCommon.h"
+#include "Signature/RCSignatureProtocolAction.h"
+
+namespace UE::RemoteControlProtocolWidgets::Private
+{
+	EMaskingType GetMaskingType(const UScriptStruct* InPropertyStruct)
+	{
+		if (const EMaskingType* FoundMaskingType = FRemoteControlProtocolMasking::GetStructsToMaskingTypes().Find(InPropertyStruct))
+		{
+			return *FoundMaskingType;
+		}
+		return EMaskingType::Unsupported;
+	}
+}
+
+void SRCSignatureProtocolMaskProperty::Construct(const FArguments& InArgs
+	, const TSharedRef<IPropertyHandle>& InMaskPropertyHandle
+	, TConstArrayView<const FRCSignatureProtocolAction*> InProtocolActions)
+{
+	using namespace UE::RemoteControlProtocolWidgets;
+
+	MaskPropertyHandle = InMaskPropertyHandle;
+
+	// Find the Property Struct for all the selected protocol actions
+	// if the protocol action have different structs, this will remain null
+	const UScriptStruct* PropertyStruct = nullptr;
+	for (const FRCSignatureProtocolAction* Action : InProtocolActions)
+	{
+		if (!Action)
+		{
+			continue;
+		}
+
+		const UScriptStruct* CurrentStruct = Action->GetPropertyStruct();
+		if (!PropertyStruct)
+		{
+			PropertyStruct = CurrentStruct;
+		}
+
+		// If the property structs don't match, do not show mask widget (multiple different protocol actions are being visualized)
+		if (PropertyStruct != CurrentStruct)
+		{
+			PropertyStruct = nullptr;
+			break;
+		}
+	}
+
+	if (PropertyStruct)
+	{
+		SRCProtocolMaskTriplet::Construct(SRCProtocolMaskTriplet::FArguments()
+			.MaskA(ERCMask::MaskA)
+			.MaskB(ERCMask::MaskB)
+			.MaskC(ERCMask::MaskC)
+			.OptionalMask(ERCMask::MaskD)
+			.MaskingType(Private::GetMaskingType(PropertyStruct))
+			.CanBeMasked(IRemoteControlModule::Get().SupportsMasking(PropertyStruct))
+			.EnableOptionalMask(FRemoteControlProtocolMasking::GetOptionalMaskStructs().Contains(PropertyStruct))
+		);
+	}
+}
+
+ECheckBoxState SRCSignatureProtocolMaskProperty::IsMaskEnabled(ERCMask InMaskBit) const
+{
+	if (!MaskPropertyHandle.IsValid())
+	{
+		return ECheckBoxState::Unchecked;
+	}
+
+	uint8 MaskValue = 0;
+	if (MaskPropertyHandle->GetValue(MaskValue) != FPropertyAccess::Success)
+	{
+		return ECheckBoxState::Unchecked;
+	}
+
+	return MaskValue & static_cast<uint8>(InMaskBit) ? ECheckBoxState::Checked : ECheckBoxState::Unchecked;
+}
+
+void SRCSignatureProtocolMaskProperty::SetMaskEnabled(ECheckBoxState InNewState, ERCMask InNewMaskBit)
+{
+	if (InNewState == ECheckBoxState::Undetermined || !MaskPropertyHandle.IsValid())
+	{
+		return;
+	}
+
+	uint8 MaskValue = 0;
+	if (MaskPropertyHandle->GetValue(MaskValue) != FPropertyAccess::Success)
+	{
+		return;
+	}
+
+	if (InNewState == ECheckBoxState::Checked)
+	{
+		MaskValue |= static_cast<uint8>(InNewMaskBit);
+	}
+	else
+	{
+		MaskValue &= ~static_cast<uint8>(InNewMaskBit);
+	}
+
+	MaskPropertyHandle->SetValue(MaskValue);
+}
diff --git a/Engine/Plugins/VirtualProduction/RemoteControl/Source/RemoteControlProtocolWidgets/Private/Widgets/Masks/SRCSignatureProtocolMaskProperty.h b/Engine/Plugins/VirtualProduction/RemoteControl/Source/RemoteControlProtocolWidgets/Private/Widgets/Masks/SRCSignatureProtocolMaskProperty.h
new file mode 100644
--- /dev/null
+++ b/Engine/Plugins/VirtualProduction/RemoteControl/Source/RemoteControlProtocolWidgets/Private/Widgets/Masks/SRCSignatureProtocolMaskProperty.h
@@ -0,0 +1,27 @@
+// Copyright Epic Games, Inc. All Rights Reserved.
+
+#pragma once
+
+#include "Widgets/Masks/SRCProtocolMaskTriplet.h"
+
+class IPropertyHandle;
+struct FRCSignatureProtocolAction;
+
+class SRCSignatureProtocolMaskProperty : public SRCProtocolMaskTriplet
+{
+public:
+	SLATE_BEGIN_ARGS(SRCSignatureProtocolMaskProperty) {}
+	SLATE_END_ARGS()
+
+	void Construct(const FArguments& InArgs
+		, const TSharedRef<IPropertyHandle>& InMaskPropertyHandle
+		, TConstArrayView<const FRCSignatureProtocolAction*> InProtocolActions);
+
+private:
+	//~ Begin SRCProtocolMaskTriplet
+	virtual ECheckBoxState IsMaskEnabled(ERCMask InMaskBit) const override;
+	virtual void SetMaskEnabled(ECheckBoxState InNewState, ERCMask InNewMaskBit) override;
+	//~ End SRCProtocolMaskTriplet
+
+	TSharedPtr<IPropertyHandle> MaskPropertyHandle;
+};
diff --git a/Engine/Plugins/VirtualProduction/RemoteControl/Source/RemoteControlProtocolWidgets/RemoteControlProtocolWidgets.Build.cs b/Engine/Plugins/VirtualProduction/RemoteControl/Source/RemoteControlProtocolWidgets/RemoteControlProtocolWidgets.Build.cs
--- a/Engine/Plugins/VirtualProduction/RemoteControl/Source/RemoteControlProtocolWidgets/RemoteControlProtocolWidgets.Build.cs
+++ b/Engine/Plugins/VirtualProduction/RemoteControl/Source/RemoteControlProtocolWidgets/RemoteControlProtocolWidgets.Build.cs
@@ -23,7 +23,8 @@ public class RemoteControlProtocolWidgets : ModuleRules
 			"RemoteControlProtocol",
 			"Slate",
 			"SlateCore",
-			"UnrealEd"
+			"StructUtilsEditor",
+			"UnrealEd",
 		});
 	}
 }
