Remote Control: Fixed issue with signatures always adding new exposed properties instead of reusing existing ones with the same path

diff --git a/Engine/Plugins/VirtualProduction/RemoteControl/Source/RemoteControl/Private/RCSignature.cpp b/Engine/Plugins/VirtualProduction/RemoteControl/Source/RemoteControl/Private/RCSignature.cpp
--- a/Engine/Plugins/VirtualProduction/RemoteControl/Source/RemoteControl/Private/RCSignature.cpp
+++ b/Engine/Plugins/VirtualProduction/RemoteControl/Source/RemoteControl/Private/RCSignature.cpp
@@ -164,7 +164,18 @@ int32 FRCSignature::ApplySignature(URemoteControlPreset* InPreset, TConstArrayVi
 			FRCSignatureActionContext ActionContext;
 			ActionContext.Preset = InPreset;
 			ActionContext.Object = Object;
-			ActionContext.Property = InPreset->ExposeProperty(Context, Field.FieldPath, ExposeArgs).Pin();
+
+			// Try to re-use existing property in Preset
+			const FName FieldName = InPreset->GetEntityName(NAME_None, Object, Path);
+			const FGuid EntityId = InPreset->GetExposedEntityId(FieldName);
+			ActionContext.Property = InPreset->GetExposedEntity<FRemoteControlProperty>(EntityId).Pin();
+
+			// if not found, create new one
+			if (!ActionContext.Property.IsValid())
+			{
+				ActionContext.Property = InPreset->ExposeProperty(Context, Field.FieldPath, ExposeArgs).Pin();
+			}
+
 			if (!ActionContext.Property.IsValid())
 			{
 				continue;
