Remote Control: Improved Mask UI for clarity

diff --git a/Engine/Plugins/VirtualProduction/RemoteControl/Source/RemoteControlProtocolWidgets/Private/Widgets/Masks/SRCProtocolMaskButton.cpp b/Engine/Plugins/VirtualProduction/RemoteControl/Source/RemoteControlProtocolWidgets/Private/Widgets/Masks/SRCProtocolMaskButton.cpp
--- a/Engine/Plugins/VirtualProduction/RemoteControl/Source/RemoteControlProtocolWidgets/Private/Widgets/Masks/SRCProtocolMaskButton.cpp
+++ b/Engine/Plugins/VirtualProduction/RemoteControl/Source/RemoteControlProtocolWidgets/Private/Widgets/Masks/SRCProtocolMaskButton.cpp
@@ -24,6 +24,8 @@ void SRCProtocolMaskButton::Construct(const FArguments& InArgs)
 
 	MaskBit = InArgs._MaskBit;
 
+	MaskColor = InArgs._MaskColor;
+
 	MaskedState = ECheckBoxState::Unchecked;
 
 	IsInMaskedState = InArgs._IsMasked;
@@ -56,14 +58,14 @@ void SRCProtocolMaskButton::Construct(const FArguments& InArgs)
 				[
 					SNew(SBorder)
 					.BorderImage(&WidgetStyle->ContentAreaBrushLight)
-					.BorderBackgroundColor(InArgs._MaskColor)
+					.BorderBackgroundColor(MaskColor)
 				]
 
 				+SOverlay::Slot()
 				.HAlign(HAlign_Center)
 				.Padding(0.f, 0.f, 0.f, 2.f)
 				[
-					SNew(SBorder)
+					SAssignNew(MaskBorder, SBorder)
 					.BorderImage(&WidgetStyle->ContentAreaBrushDark)
 					.HAlign(HAlign_Center)
 					.VAlign(VAlign_Center)
@@ -96,8 +98,9 @@ void SRCProtocolMaskButton::ToggleMaskedState(bool bSoftToggle)
 	if (bSoftToggle && MaskLabel.IsValid())
 	{
 		MaskLabel->SetTextStyle(&WidgetStyle->BoldTextStyle);
-
 		MaskLabel->SetColorAndOpacity(FLinearColor::White);
+		MaskBorder->SetBorderBackgroundColor(MaskColor);
+		MaskBorder->SetBorderImage(&WidgetStyle->ContentAreaBrushLight);
 	}
 	else if (Mask.IsValid())
 	{
@@ -134,14 +137,16 @@ void SRCProtocolMaskButton::SetMaskEnabled(ECheckBoxState NewState)
 	if (bIsChecked)
 	{
 		MaskLabel->SetTextStyle(&WidgetStyle->BoldTextStyle);
-
 		MaskLabel->SetColorAndOpacity(FLinearColor::White);
+		MaskBorder->SetBorderBackgroundColor(MaskColor);
+		MaskBorder->SetBorderImage(&WidgetStyle->ContentAreaBrushLight);
 	}
 	else
 	{
 		MaskLabel->SetTextStyle(&WidgetStyle->PlainTextStyle);
-
 		MaskLabel->SetColorAndOpacity(FSlateColor::UseForeground());
+		MaskBorder->SetBorderBackgroundColor(FLinearColor::White);
+		MaskBorder->SetBorderImage(&WidgetStyle->ContentAreaBrushDark);
 	}
 
 	OnMaskStateChanged.ExecuteIfBound(NewState, MaskBit);
diff --git a/Engine/Plugins/VirtualProduction/RemoteControl/Source/RemoteControlProtocolWidgets/Private/Widgets/Masks/SRCProtocolMaskButton.h b/Engine/Plugins/VirtualProduction/RemoteControl/Source/RemoteControlProtocolWidgets/Private/Widgets/Masks/SRCProtocolMaskButton.h
--- a/Engine/Plugins/VirtualProduction/RemoteControl/Source/RemoteControlProtocolWidgets/Private/Widgets/Masks/SRCProtocolMaskButton.h
+++ b/Engine/Plugins/VirtualProduction/RemoteControl/Source/RemoteControlProtocolWidgets/Private/Widgets/Masks/SRCProtocolMaskButton.h
@@ -94,12 +94,18 @@ private:
 	/** Holds the mask that should be affected by this button as an uint8. */
 	ERCMask MaskBit;
 
+	/** The color of the mask to set */
+	FLinearColor MaskColor = FLinearColor::White;
+
 	/** Holds a shared reference to the underlying mask button. */
 	TSharedPtr<SCheckBox> Mask;
 	
 	/** Holds a shared reference to the underlying mask label. */
 	TSharedPtr<STextBlock> MaskLabel;
 
+	/** Holds a shared reference to the underlying mask border. */
+	TSharedPtr<SBorder> MaskBorder;
+
 	/** Holds the "Mask" string as text. */
 	static FText MaskText;
 
