Remote Control Signature: Set the default override to all masks

diff --git a/Engine/Plugins/VirtualProduction/RemoteControl/Source/RemoteControlProtocol/Public/Signature/RCSignatureProtocolAction.h b/Engine/Plugins/VirtualProduction/RemoteControl/Source/RemoteControlProtocol/Public/Signature/RCSignatureProtocolAction.h
--- a/Engine/Plugins/VirtualProduction/RemoteControl/Source/RemoteControlProtocol/Public/Signature/RCSignatureProtocolAction.h
+++ b/Engine/Plugins/VirtualProduction/RemoteControl/Source/RemoteControlProtocol/Public/Signature/RCSignatureProtocolAction.h
@@ -63,7 +63,7 @@ struct FRCSignatureProtocolAction : public FRCSignatureAction
 	ERCSignatureProtocolActionMaskingType MaskingType = ERCSignatureProtocolActionMaskingType::SingleChannel;
 
 	UPROPERTY(EditAnywhere, Category="Protocol", meta=(EditCondition="PropertyDimension > 1", EditConditionHides))
-	uint8 OverrideMask = 0;
+	uint8 OverrideMask = 0xFF;
 
 	UPROPERTY(EditAnywhere, Category="Protocol")
 	ERCSignatureProtocolActionMappingSpace MappingSpace = ERCSignatureProtocolActionMappingSpace::Additive;
