[Remote Control] RCPanel: Fix Function Picker

The function items where recycled and it seems like the SListView doesn't like having the same items multiple time.

diff --git a/Engine/Plugins/VirtualProduction/RemoteControl/Source/RemoteControlUI/Private/UI/SRCPanelFunctionPicker.cpp b/Engine/Plugins/VirtualProduction/RemoteControl/Source/RemoteControlUI/Private/UI/SRCPanelFunctionPicker.cpp
--- a/Engine/Plugins/VirtualProduction/RemoteControl/Source/RemoteControlUI/Private/UI/SRCPanelFunctionPicker.cpp
+++ b/Engine/Plugins/VirtualProduction/RemoteControl/Source/RemoteControlUI/Private/UI/SRCPanelFunctionPicker.cpp
@@ -109,13 +109,13 @@ namespace FunctionPickerUtils
 			}
 		}
 
-		//~ Begin FRCBlueprintPickerTreeNode interface
+		//~ Begin FRCFunctionPickerTreeNode
 		virtual FText GetName() const override { return Name; }
 		virtual bool IsFunctionNode() const override { return true; }
 		virtual UFunction* GetFunction() const override { return Function.Get(); }
 		virtual UObject* GetObject() const override { return nullptr; }
-		virtual void GetChildNodes(TArray<TSharedPtr<FRCFunctionPickerTreeNode>>& OutNodes) const {}
-		//~ End FRCBlueprintPickerTreeNode interface
+		virtual void GetChildNodes(TArray<TSharedPtr<FRCFunctionPickerTreeNode>>& OutNodes) const override {}
+		//~ End FRCFunctionPickerTreeNode
 
 	private:
 		/** This node's function. */
@@ -124,11 +124,29 @@ namespace FunctionPickerUtils
 		FText Name;
 	};
 
+	struct FRCFunctionNodeWrapper : public FRCFunctionPickerTreeNode
+	{
+		FRCFunctionNodeWrapper(TSharedPtr<FRCFunctionPickerTreeNode> InNode)
+			: Node(InNode)
+		{
+		}
+
+		//~ Begin FRCFunctionPickerTreeNode
+		virtual FText GetName() const override { return Node ? Node->GetName() : FText::GetEmpty(); }
+		virtual bool IsFunctionNode() const override { return true; }
+		virtual UFunction* GetFunction() const override { return Node ? Node->GetFunction() : nullptr; }
+		virtual UObject* GetObject() const override { return nullptr; }
+		virtual void GetChildNodes(TArray<TSharedPtr<FRCFunctionPickerTreeNode>>& OutNodes) const override {}
+		//~ End FRCFunctionPickerTreeNode
+
+		private:
+		TSharedPtr<FRCFunctionPickerTreeNode> Node;
+	};
+
 	struct FRCObjectNode : public FRCFunctionPickerTreeNode
 	{
 		FRCObjectNode(UObject* InObject, const TArray<TSharedPtr<FRCFunctionPickerTreeNode>>& InFunctionNodes)
 			: WeakObject(InObject)
-			, ChildNodes(InFunctionNodes)
 		{
 			TRACE_CPUPROFILER_EVENT_SCOPE(FRCObjectNode::FRCObjectNode());
 
@@ -148,10 +166,17 @@ namespace FunctionPickerUtils
 				}
 
 				Name = FText::FromString(NameString);
+
+				// Children Items must be a different item to avoid issues when list is flattened in search view.
+				ChildNodes.Reserve(InFunctionNodes.Num());
+				for (const TSharedPtr<FRCFunctionPickerTreeNode>& FunctionNode : InFunctionNodes)
+				{
+					ChildNodes.Add(MakeShared<FRCFunctionNodeWrapper>(FunctionNode));
+				}
 			}
 		}
 
-		//~ Begin FRCBlueprintPickerTreeNode interface
+		//~ Begin FRCFunctionPickerTreeNode
 		virtual FText GetName() const override { return Name; }
 		virtual bool IsFunctionNode() const override { return false; }
 		virtual UObject* GetObject() const override { return WeakObject.Get(); }
@@ -160,7 +185,7 @@ namespace FunctionPickerUtils
 		{
 			OutNodes.Append(ChildNodes);
 		}
-		//~ End FRCBlueprintPickerTreeNode
+		//~ End FRCFunctionPickerTreeNode
 
 	private:
 		/** The object represented by this node. */
@@ -286,7 +311,7 @@ void SRCPanelFunctionPicker::Refresh()
 
 			TArray<UFunction*> Functions = GetExposableFunctions(Object->GetClass());
 			TArray<TSharedPtr<FRCFunctionPickerTreeNode>> FunctionNodes;
-			FunctionNodes.Reserve(FunctionNodes.Num());
+			FunctionNodes.Reserve(Functions.Num());
 			for (UFunction* Function : Functions)
 			{
 				FunctionNodes.Add(MakeShared<FRCFunctionNode>(Function));
