Remote Control: Protocols no longer create Transactions by default. Add a performance mode for Protocols to Project Settings

diff --git a/Engine/Plugins/VirtualProduction/RemoteControl/Source/RemoteControl/Private/Factories/RemoteControlMaskingFactories.cpp b/Engine/Plugins/VirtualProduction/RemoteControl/Source/RemoteControl/Private/Factories/RemoteControlMaskingFactories.cpp
--- a/Engine/Plugins/VirtualProduction/RemoteControl/Source/RemoteControl/Private/Factories/RemoteControlMaskingFactories.cpp
+++ b/Engine/Plugins/VirtualProduction/RemoteControl/Source/RemoteControl/Private/Factories/RemoteControlMaskingFactories.cpp
@@ -10,16 +10,19 @@ TSharedRef<IRemoteControlMaskingFactory> FVectorMaskingFactory::MakeInstance()
 	return MakeShared<FVectorMaskingFactory>();
 }
 
-void FVectorMaskingFactory::ApplyMaskedValues(const TSharedRef<FRCMaskingOperation>& InMaskingOperation, bool bIsInteractive)
+void FVectorMaskingFactory::ApplyMaskedValues(const TSharedRef<FRCMaskingOperation>& InMaskingOperation, bool bIsInteractive, const ERCModifyOperationFlags ModifyOperationFlags)
 {
 	if (FStructProperty* ToStructProp = CastField<FStructProperty>(InMaskingOperation->ObjectRef.Property.Get()))
 	{
 		if (UObject* OwningObject = InMaskingOperation->ObjectRef.Object.Get())
 		{
 #if WITH_EDITOR
-			OwningObject->PreEditChange(ToStructProp);
-			OwningObject->Modify();
-#endif // WITH_EDITOR
+			const bool bWithPropertyChangedEvents = !EnumHasAnyFlags(ModifyOperationFlags, ERCModifyOperationFlags::SkipPropertyChangeEvents);
+			if (bWithPropertyChangedEvents)
+			{
+				OwningObject->PreEditChange(ToStructProp);
+			}
+#endif
 
 			if (const FVector* VectorProp = ToStructProp->ContainerPtrToValuePtr<FVector>(OwningObject))
 			{
@@ -33,9 +36,12 @@ void FVectorMaskingFactory::ApplyMaskedValues(const TSharedRef<FRCMaskingOperati
 			}
 
 #if WITH_EDITOR
-			FPropertyChangedEvent ChangeEvent(ToStructProp, bIsInteractive ? EPropertyChangeType::Interactive : EPropertyChangeType::ValueSet);
-			OwningObject->PostEditChangeProperty(ChangeEvent);
-#endif // WITH_EDITOR
+			if (bWithPropertyChangedEvents)
+			{
+				FPropertyChangedEvent ChangeEvent(ToStructProp, bIsInteractive ? EPropertyChangeType::Interactive : EPropertyChangeType::ValueSet);
+				OwningObject->PostEditChangeProperty(ChangeEvent);
+			}
+#endif
 		}
 	}
 }
@@ -69,15 +75,18 @@ TSharedRef<IRemoteControlMaskingFactory> FVector4MaskingFactory::MakeInstance()
 	return MakeShared<FVector4MaskingFactory>();
 }
 
-void FVector4MaskingFactory::ApplyMaskedValues(const TSharedRef<FRCMaskingOperation>& InMaskingOperation, bool bIsInteractive)
+void FVector4MaskingFactory::ApplyMaskedValues(const TSharedRef<FRCMaskingOperation>& InMaskingOperation, bool bIsInteractive, const ERCModifyOperationFlags ModifyOperationFlags)
 {
 	if (FStructProperty* ToStructProp = CastField<FStructProperty>(InMaskingOperation->ObjectRef.Property.Get()))
 	{
 		if (UObject* OwningObject = InMaskingOperation->ObjectRef.Object.Get())
 		{
 #if WITH_EDITOR
-			OwningObject->PreEditChange(ToStructProp);
-			OwningObject->Modify();
+			const bool bWithPropertyChangedEvents = !EnumHasAnyFlags(ModifyOperationFlags, ERCModifyOperationFlags::SkipPropertyChangeEvents);
+			if (bWithPropertyChangedEvents)
+			{
+				OwningObject->PreEditChange(ToStructProp);
+			}
 #endif // WITH_EDITOR
 
 			if (const FVector4* Vector4Prop = ToStructProp->ContainerPtrToValuePtr<FVector4>(OwningObject))
@@ -93,8 +102,11 @@ void FVector4MaskingFactory::ApplyMaskedValues(const TSharedRef<FRCMaskingOperat
 			}
 
 #if WITH_EDITOR
-			FPropertyChangedEvent ChangeEvent(ToStructProp, bIsInteractive ? EPropertyChangeType::Interactive : EPropertyChangeType::ValueSet);
-			OwningObject->PostEditChangeProperty(ChangeEvent);
+			if (bWithPropertyChangedEvents)
+			{
+				FPropertyChangedEvent ChangeEvent(ToStructProp, bIsInteractive ? EPropertyChangeType::Interactive : EPropertyChangeType::ValueSet);
+				OwningObject->PostEditChangeProperty(ChangeEvent);
+			}
 #endif // WITH_EDITOR
 		}
 	}
@@ -130,15 +142,18 @@ TSharedRef<IRemoteControlMaskingFactory> FIntVectorMaskingFactory::MakeInstance(
 	return MakeShared<FIntVectorMaskingFactory>();
 }
 
-void FIntVectorMaskingFactory::ApplyMaskedValues(const TSharedRef<FRCMaskingOperation>& InMaskingOperation, bool bIsInteractive)
+void FIntVectorMaskingFactory::ApplyMaskedValues(const TSharedRef<FRCMaskingOperation>& InMaskingOperation, bool bIsInteractive, const ERCModifyOperationFlags ModifyOperationFlags)
 {
 	if (FStructProperty* ToStructProp = CastField<FStructProperty>(InMaskingOperation->ObjectRef.Property.Get()))
 	{
 		if (UObject* OwningObject = InMaskingOperation->ObjectRef.Object.Get())
 		{
 #if WITH_EDITOR
-			OwningObject->PreEditChange(ToStructProp);
-			OwningObject->Modify();
+			const bool bWithPropertyChangedEvents = !EnumHasAnyFlags(ModifyOperationFlags, ERCModifyOperationFlags::SkipPropertyChangeEvents);
+			if (bWithPropertyChangedEvents)
+			{
+				OwningObject->PreEditChange(ToStructProp);
+			}
 #endif // WITH_EDITOR
 
 			if (const FIntVector* IntVectorProp = ToStructProp->ContainerPtrToValuePtr<FIntVector>(OwningObject))
@@ -153,8 +168,11 @@ void FIntVectorMaskingFactory::ApplyMaskedValues(const TSharedRef<FRCMaskingOper
 			}
 
 #if WITH_EDITOR
-			FPropertyChangedEvent ChangeEvent(ToStructProp, bIsInteractive ? EPropertyChangeType::Interactive : EPropertyChangeType::ValueSet);
-			OwningObject->PostEditChangeProperty(ChangeEvent);
+			if (bWithPropertyChangedEvents)
+			{
+				FPropertyChangedEvent ChangeEvent(ToStructProp, bIsInteractive ? EPropertyChangeType::Interactive : EPropertyChangeType::ValueSet);
+				OwningObject->PostEditChangeProperty(ChangeEvent);
+			}
 #endif // WITH_EDITOR
 		}
 	}
@@ -189,15 +207,18 @@ TSharedRef<IRemoteControlMaskingFactory> FIntVector4MaskingFactory::MakeInstance
 	return MakeShared<FIntVector4MaskingFactory>();
 }
 
-void FIntVector4MaskingFactory::ApplyMaskedValues(const TSharedRef<FRCMaskingOperation>& InMaskingOperation, bool bIsInteractive)
+void FIntVector4MaskingFactory::ApplyMaskedValues(const TSharedRef<FRCMaskingOperation>& InMaskingOperation, bool bIsInteractive, const ERCModifyOperationFlags ModifyOperationFlags)
 {
 	if (FStructProperty* ToStructProp = CastField<FStructProperty>(InMaskingOperation->ObjectRef.Property.Get()))
 	{
 		if (UObject* OwningObject = InMaskingOperation->ObjectRef.Object.Get())
 		{
 #if WITH_EDITOR
-			OwningObject->PreEditChange(ToStructProp);
-			OwningObject->Modify();
+			const bool bWithPropertyChangedEvents = !EnumHasAnyFlags(ModifyOperationFlags, ERCModifyOperationFlags::SkipPropertyChangeEvents);
+			if (bWithPropertyChangedEvents)
+			{
+				OwningObject->PreEditChange(ToStructProp);
+			}
 #endif // WITH_EDITOR
 
 			if (const FIntVector4* IntVector4Prop = ToStructProp->ContainerPtrToValuePtr<FIntVector4>(OwningObject))
@@ -213,8 +234,11 @@ void FIntVector4MaskingFactory::ApplyMaskedValues(const TSharedRef<FRCMaskingOpe
 			}
 
 #if WITH_EDITOR
-			FPropertyChangedEvent ChangeEvent(ToStructProp, bIsInteractive ? EPropertyChangeType::Interactive : EPropertyChangeType::ValueSet);
-			OwningObject->PostEditChangeProperty(ChangeEvent);
+			if (bWithPropertyChangedEvents)
+			{
+				FPropertyChangedEvent ChangeEvent(ToStructProp, bIsInteractive ? EPropertyChangeType::Interactive : EPropertyChangeType::ValueSet);
+				OwningObject->PostEditChangeProperty(ChangeEvent);
+			}
 #endif // WITH_EDITOR
 		}
 	}
@@ -250,15 +274,18 @@ TSharedRef<IRemoteControlMaskingFactory> FRotatorMaskingFactory::MakeInstance()
 	return MakeShared<FRotatorMaskingFactory>();
 }
 
-void FRotatorMaskingFactory::ApplyMaskedValues(const TSharedRef<FRCMaskingOperation>& InMaskingOperation, bool bIsInteractive)
+void FRotatorMaskingFactory::ApplyMaskedValues(const TSharedRef<FRCMaskingOperation>& InMaskingOperation, bool bIsInteractive, const ERCModifyOperationFlags ModifyOperationFlags)
 {
 	if (FStructProperty* ToStructProp = CastField<FStructProperty>(InMaskingOperation->ObjectRef.Property.Get()))
 	{
 		if (UObject* OwningObject = InMaskingOperation->ObjectRef.Object.Get())
 		{
 #if WITH_EDITOR
-			OwningObject->PreEditChange(ToStructProp);
-			OwningObject->Modify();
+			const bool bWithPropertyChangedEvents = !EnumHasAnyFlags(ModifyOperationFlags, ERCModifyOperationFlags::SkipPropertyChangeEvents);
+			if (bWithPropertyChangedEvents)
+			{
+				OwningObject->PreEditChange(ToStructProp);
+			}
 #endif // WITH_EDITOR
 
 			if (const FRotator* RotatorProp = ToStructProp->ContainerPtrToValuePtr<FRotator>(OwningObject))
@@ -272,9 +299,12 @@ void FRotatorMaskingFactory::ApplyMaskedValues(const TSharedRef<FRCMaskingOperat
 				ToStructProp->SetValue_InContainer(OwningObject, &MaskedRotator);
 			}
 
-#if WITH_EDITOR
-			FPropertyChangedEvent ChangeEvent(ToStructProp, bIsInteractive ? EPropertyChangeType::Interactive : EPropertyChangeType::ValueSet);
-			OwningObject->PostEditChangeProperty(ChangeEvent);
+#if WITH_EDITOR			
+			if (bWithPropertyChangedEvents)
+			{
+				FPropertyChangedEvent ChangeEvent(ToStructProp, bIsInteractive ? EPropertyChangeType::Interactive : EPropertyChangeType::ValueSet);
+				OwningObject->PostEditChangeProperty(ChangeEvent);
+			}
 #endif // WITH_EDITOR
 		}
 	}
@@ -309,15 +339,18 @@ TSharedRef<IRemoteControlMaskingFactory> FColorMaskingFactory::MakeInstance()
 	return MakeShared<FColorMaskingFactory>();
 }
 
-void FColorMaskingFactory::ApplyMaskedValues(const TSharedRef<FRCMaskingOperation>& InMaskingOperation, bool bIsInteractive)
+void FColorMaskingFactory::ApplyMaskedValues(const TSharedRef<FRCMaskingOperation>& InMaskingOperation, bool bIsInteractive, const ERCModifyOperationFlags ModifyOperationFlags)
 {
 	if (FStructProperty* ToStructProp = CastField<FStructProperty>(InMaskingOperation->ObjectRef.Property.Get()))
 	{
 		if (UObject* OwningObject = InMaskingOperation->ObjectRef.Object.Get())
 		{
 #if WITH_EDITOR
-			OwningObject->PreEditChange(ToStructProp);
-			OwningObject->Modify();
+			const bool bWithPropertyChangedEvents = !EnumHasAnyFlags(ModifyOperationFlags, ERCModifyOperationFlags::SkipPropertyChangeEvents);
+			if (bWithPropertyChangedEvents)
+			{
+				OwningObject->PreEditChange(ToStructProp);
+			}
 #endif // WITH_EDITOR
 
 			if (const FColor* ColorProp = ToStructProp->ContainerPtrToValuePtr<FColor>(OwningObject))
@@ -332,8 +365,11 @@ void FColorMaskingFactory::ApplyMaskedValues(const TSharedRef<FRCMaskingOperatio
 			}
 
 #if WITH_EDITOR
-			FPropertyChangedEvent ChangeEvent(ToStructProp, bIsInteractive ? EPropertyChangeType::Interactive : EPropertyChangeType::ValueSet);
-			OwningObject->PostEditChangeProperty(ChangeEvent);
+			if (bWithPropertyChangedEvents)
+			{
+				FPropertyChangedEvent ChangeEvent(ToStructProp, bIsInteractive ? EPropertyChangeType::Interactive : EPropertyChangeType::ValueSet);
+				OwningObject->PostEditChangeProperty(ChangeEvent);
+			}
 #endif // WITH_EDITOR
 		}
 	}
@@ -368,15 +404,18 @@ TSharedRef<IRemoteControlMaskingFactory> FLinearColorMaskingFactory::MakeInstanc
 	return MakeShared<FLinearColorMaskingFactory>();
 }
 
-void FLinearColorMaskingFactory::ApplyMaskedValues(const TSharedRef<FRCMaskingOperation>& InMaskingOperation, bool bIsInteractive)
+void FLinearColorMaskingFactory::ApplyMaskedValues(const TSharedRef<FRCMaskingOperation>& InMaskingOperation, bool bIsInteractive, const ERCModifyOperationFlags ModifyOperationFlags)
 {
 	if (FStructProperty* ToStructProp = CastField<FStructProperty>(InMaskingOperation->ObjectRef.Property.Get()))
 	{
 		if (UObject* OwningObject = InMaskingOperation->ObjectRef.Object.Get())
 		{
 #if WITH_EDITOR
-			OwningObject->PreEditChange(ToStructProp);
-			OwningObject->Modify();
+			const bool bWithPropertyChangedEvents = !EnumHasAnyFlags(ModifyOperationFlags, ERCModifyOperationFlags::SkipPropertyChangeEvents);
+			if (bWithPropertyChangedEvents)
+			{
+				OwningObject->PreEditChange(ToStructProp);
+			}
 #endif // WITH_EDITOR
 
 			if (const FLinearColor* LinearColorProp = ToStructProp->ContainerPtrToValuePtr<FLinearColor>(OwningObject))
@@ -392,8 +431,11 @@ void FLinearColorMaskingFactory::ApplyMaskedValues(const TSharedRef<FRCMaskingOp
 			}
 
 #if WITH_EDITOR
-			FPropertyChangedEvent ChangeEvent(ToStructProp, bIsInteractive ? EPropertyChangeType::Interactive : EPropertyChangeType::ValueSet);
-			OwningObject->PostEditChangeProperty(ChangeEvent);
+			if (bWithPropertyChangedEvents)
+			{
+				FPropertyChangedEvent ChangeEvent(ToStructProp, bIsInteractive ? EPropertyChangeType::Interactive : EPropertyChangeType::ValueSet);
+				OwningObject->PostEditChangeProperty(ChangeEvent);
+			}
 #endif // WITH_EDITOR
 		}
 	}
diff --git a/Engine/Plugins/VirtualProduction/RemoteControl/Source/RemoteControl/Private/Factories/RemoteControlMaskingFactories.h b/Engine/Plugins/VirtualProduction/RemoteControl/Source/RemoteControl/Private/Factories/RemoteControlMaskingFactories.h
--- a/Engine/Plugins/VirtualProduction/RemoteControl/Source/RemoteControl/Private/Factories/RemoteControlMaskingFactories.h
+++ b/Engine/Plugins/VirtualProduction/RemoteControl/Source/RemoteControl/Private/Factories/RemoteControlMaskingFactories.h
@@ -11,7 +11,7 @@ public:
 	static TSharedRef<IRemoteControlMaskingFactory> MakeInstance();
 
 	//~ Begin IRemoteControlMaskingFactory interface
-	virtual void ApplyMaskedValues(const TSharedRef<FRCMaskingOperation>& InMaskingOperation, bool bIsInteractive) override;
+	virtual void ApplyMaskedValues(const TSharedRef<FRCMaskingOperation>& InMaskingOperation, bool bIsInteractive, const ERCModifyOperationFlags ModifyOperationFlags = ERCModifyOperationFlags::None) override;
 	virtual void CacheRawValues(const TSharedRef<FRCMaskingOperation>& InMaskingOperation) override;
 	virtual bool SupportsExposedEntity(UScriptStruct* ScriptStruct) const override;
 	//~ End IRemoteControlMaskingFactory interface
@@ -23,7 +23,7 @@ public:
 	static TSharedRef<IRemoteControlMaskingFactory> MakeInstance();
 
 	//~ Begin IRemoteControlMaskingFactory interface
-	virtual void ApplyMaskedValues(const TSharedRef<FRCMaskingOperation>& InMaskingOperation, bool bIsInteractive) override;
+	virtual void ApplyMaskedValues(const TSharedRef<FRCMaskingOperation>& InMaskingOperation, bool bIsInteractive, const ERCModifyOperationFlags ModifyOperationFlags = ERCModifyOperationFlags::None) override;
 	virtual void CacheRawValues(const TSharedRef<FRCMaskingOperation>& InMaskingOperation) override;
 	virtual bool SupportsExposedEntity(UScriptStruct* ScriptStruct) const override;
 	//~ End IRemoteControlMaskingFactory interface
@@ -35,7 +35,7 @@ public:
 	static TSharedRef<IRemoteControlMaskingFactory> MakeInstance();
 
 	//~ Begin IRemoteControlMaskingFactory interface
-	virtual void ApplyMaskedValues(const TSharedRef<FRCMaskingOperation>& InMaskingOperation, bool bIsInteractive) override;
+	virtual void ApplyMaskedValues(const TSharedRef<FRCMaskingOperation>& InMaskingOperation, bool bIsInteractive, const ERCModifyOperationFlags ModifyOperationFlags = ERCModifyOperationFlags::None) override;
 	virtual void CacheRawValues(const TSharedRef<FRCMaskingOperation>& InMaskingOperation) override;
 	virtual bool SupportsExposedEntity(UScriptStruct* ScriptStruct) const override;
 	//~ End IRemoteControlMaskingFactory interface
@@ -47,7 +47,7 @@ public:
 	static TSharedRef<IRemoteControlMaskingFactory> MakeInstance();
 
 	//~ Begin IRemoteControlMaskingFactory interface
-	virtual void ApplyMaskedValues(const TSharedRef<FRCMaskingOperation>& InMaskingOperation, bool bIsInteractive) override;
+	virtual void ApplyMaskedValues(const TSharedRef<FRCMaskingOperation>& InMaskingOperation, bool bIsInteractive, const ERCModifyOperationFlags ModifyOperationFlags = ERCModifyOperationFlags::None) override;
 	virtual void CacheRawValues(const TSharedRef<FRCMaskingOperation>& InMaskingOperation) override;
 	virtual bool SupportsExposedEntity(UScriptStruct* ScriptStruct) const override;
 	//~ End IRemoteControlMaskingFactory interface
@@ -59,7 +59,7 @@ public:
 	static TSharedRef<IRemoteControlMaskingFactory> MakeInstance();
 
 	//~ Begin IRemoteControlMaskingFactory interface
-	virtual void ApplyMaskedValues(const TSharedRef<FRCMaskingOperation>& InMaskingOperation, bool bIsInteractive) override;
+	virtual void ApplyMaskedValues(const TSharedRef<FRCMaskingOperation>& InMaskingOperation, bool bIsInteractive, const ERCModifyOperationFlags ModifyOperationFlags = ERCModifyOperationFlags::None) override;
 	virtual void CacheRawValues(const TSharedRef<FRCMaskingOperation>& InMaskingOperation) override;
 	virtual bool SupportsExposedEntity(UScriptStruct* ScriptStruct) const override;
 	//~ End IRemoteControlMaskingFactory interface
@@ -71,7 +71,7 @@ public:
 	static TSharedRef<IRemoteControlMaskingFactory> MakeInstance();
 
 	//~ Begin IRemoteControlMaskingFactory interface
-	virtual void ApplyMaskedValues(const TSharedRef<FRCMaskingOperation>& InMaskingOperation, bool bIsInteractive) override;
+	virtual void ApplyMaskedValues(const TSharedRef<FRCMaskingOperation>& InMaskingOperation, bool bIsInteractive, const ERCModifyOperationFlags ModifyOperationFlags = ERCModifyOperationFlags::None) override;
 	virtual void CacheRawValues(const TSharedRef<FRCMaskingOperation>& InMaskingOperation) override;
 	virtual bool SupportsExposedEntity(UScriptStruct* ScriptStruct) const override;
 	//~ End IRemoteControlMaskingFactory interface
@@ -83,7 +83,7 @@ public:
 	static TSharedRef<IRemoteControlMaskingFactory> MakeInstance();
 
 	//~ Begin IRemoteControlMaskingFactory interface
-	virtual void ApplyMaskedValues(const TSharedRef<FRCMaskingOperation>& InMaskingOperation, bool bIsInteractive) override;
+	virtual void ApplyMaskedValues(const TSharedRef<FRCMaskingOperation>& InMaskingOperation, bool bIsInteractive, const ERCModifyOperationFlags ModifyOperationFlags = ERCModifyOperationFlags::None) override;
 	virtual void CacheRawValues(const TSharedRef<FRCMaskingOperation>& InMaskingOperation) override;
 	virtual bool SupportsExposedEntity(UScriptStruct* ScriptStruct) const override;
 	//~ End IRemoteControlMaskingFactory interface
diff --git a/Engine/Plugins/VirtualProduction/RemoteControl/Source/RemoteControl/Private/RemoteControlModule.cpp b/Engine/Plugins/VirtualProduction/RemoteControl/Source/RemoteControl/Private/RemoteControlModule.cpp
--- a/Engine/Plugins/VirtualProduction/RemoteControl/Source/RemoteControl/Private/RemoteControlModule.cpp
+++ b/Engine/Plugins/VirtualProduction/RemoteControl/Source/RemoteControl/Private/RemoteControlModule.cpp
@@ -11,16 +11,16 @@
 #include "Factories/RCDefaultValueFactories.h"
 #include "Factories/RemoteControlMaskingFactories.h"
 #include "Features/IModularFeatures.h"
-#include "PropertyIdHandler/BasePropertyIdHandler.h"
-#include "PropertyIdHandler/EnumPropertyIdHandler.h"
-#include "PropertyIdHandler/ObjectPropertyIdHandler.h"
-#include "PropertyIdHandler/StructPropertyIdHandler.h"
 #include "IRemoteControlInterceptionFeature.h"
 #include "IRemoteControlModule.h"
 #include "IStructDeserializerBackend.h"
 #include "IStructSerializerBackend.h"
 #include "Kismet/KismetSystemLibrary.h"
 #include "Misc/ScopeExit.h"
+#include "PropertyIdHandler/BasePropertyIdHandler.h"
+#include "PropertyIdHandler/EnumPropertyIdHandler.h"
+#include "PropertyIdHandler/ObjectPropertyIdHandler.h"
+#include "PropertyIdHandler/StructPropertyIdHandler.h"
 #include "RCPropertyUtilities.h"
 #include "RCVirtualProperty.h"
 #include "RCVirtualPropertyContainer.h"
@@ -1027,7 +1027,7 @@ void FRemoteControlModule::ResetToDefaultValue(UObject* InObject, FRCResetToDefa
 	}
 }
 
-void FRemoteControlModule::PerformMasking(const TSharedRef<FRCMaskingOperation>& InMaskingOperation)
+void FRemoteControlModule::PerformMasking(const TSharedRef<FRCMaskingOperation>& InMaskingOperation, const ERCModifyOperationFlags ModifyOperationFlags)
 {
 	TRACE_CPUPROFILER_EVENT_SCOPE(FRemoteControlModule::PerformMasking);
 
@@ -1050,7 +1050,7 @@ void FRemoteControlModule::PerformMasking(const TSharedRef<FRCMaskingOperation>&
 			{
 				constexpr bool bIsInteractive = true;
 
-				(*MaskingFactory)->ApplyMaskedValues(InMaskingOperation, bIsInteractive);
+				(*MaskingFactory)->ApplyMaskedValues(InMaskingOperation, bIsInteractive, ModifyOperationFlags);
 
 				ActiveMaskingOperations.Remove(InMaskingOperation);
 			}
@@ -1554,7 +1554,7 @@ void FRemoteControlModule::SnapshotOrEndTransaction(FRCObjectReference& ObjectRe
 }
 #endif
 
-bool FRemoteControlModule::SetObjectProperties(const FRCObjectReference& ObjectAccess, IStructDeserializerBackend& Backend, ERCPayloadType InPayloadType, const TArray<uint8>& InPayload, ERCModifyOperation Operation)
+bool FRemoteControlModule::SetObjectProperties(const FRCObjectReference& ObjectAccess, IStructDeserializerBackend& Backend, ERCPayloadType InPayloadType, const TArray<uint8>& InPayload, ERCModifyOperation Operation, const ERCModifyOperationFlags ModifyOperationFlags)
 {
 	TRACE_CPUPROFILER_EVENT_SCOPE(FRemoteControlModule::SetObjectProperties);
 	UE_LOG(LogRemoteControl, VeryVerbose, TEXT("Set Object Properties"));
@@ -1675,8 +1675,11 @@ bool FRemoteControlModule::SetObjectProperties(const FRCObjectReference& ObjectA
 		FRCObjectReference MutableObjectReference = ObjectAccess;
 
 #if WITH_EDITOR
-		bool bGeneratedTransaction;
-		if (!StartPropertyTransaction(MutableObjectReference, LOCTEXT("RemoteSetPropertyTransaction", "Remote Set Object Property"), bGeneratedTransaction))
+		const bool bWithPropertyChangedEvents = !EnumHasAnyFlags(ModifyOperationFlags, ERCModifyOperationFlags::SkipPropertyChangeEvents);
+		
+		bool bGeneratedTransaction = false;
+		if (bWithPropertyChangedEvents &&
+			!StartPropertyTransaction(MutableObjectReference, LOCTEXT("RemoteSetPropertyTransaction", "Remote Set Object Property"), bGeneratedTransaction))
 		{
 			return false;
 		}
@@ -1738,7 +1741,10 @@ bool FRemoteControlModule::SetObjectProperties(const FRCObjectReference& ObjectA
 		}
 
 #if WITH_EDITOR
-		SnapshotOrEndTransaction(MutableObjectReference, bGeneratedTransaction);
+		if (bWithPropertyChangedEvents)
+		{
+			SnapshotOrEndTransaction(MutableObjectReference, bGeneratedTransaction);
+		}
 #endif
 
 		for (const TPair<FName, TSharedPtr<IRemoteControlPropertyFactory>>& EntityFactoryPair : EntityFactories)
diff --git a/Engine/Plugins/VirtualProduction/RemoteControl/Source/RemoteControl/Private/RemoteControlModule.h b/Engine/Plugins/VirtualProduction/RemoteControl/Source/RemoteControl/Private/RemoteControlModule.h
--- a/Engine/Plugins/VirtualProduction/RemoteControl/Source/RemoteControl/Private/RemoteControlModule.h
+++ b/Engine/Plugins/VirtualProduction/RemoteControl/Source/RemoteControl/Private/RemoteControlModule.h
@@ -35,7 +35,7 @@ public:
 	virtual bool CanResetToDefaultValue(UObject* InObject, const FRCResetToDefaultArgs& InArgs) const override;
 	virtual bool HasDefaultValueCustomization(const UObject* InObject, const FProperty* InProperty) const override;
 	virtual void ResetToDefaultValue(UObject* InObject, FRCResetToDefaultArgs& InArgs) override;
-	virtual void PerformMasking(const TSharedRef<FRCMaskingOperation>& InMaskingOperation) override;
+	virtual void PerformMasking(const TSharedRef<FRCMaskingOperation>& InMaskingOperation, const ERCModifyOperationFlags ModifyOperationFlags = ERCModifyOperationFlags::None) override;
 	virtual void RegisterMaskingFactoryForType(UScriptStruct* RemoteControlPropertyType, const TSharedPtr<IRemoteControlMaskingFactory>& InMaskingFactory) override;
 	virtual void UnregisterMaskingFactoryForType(UScriptStruct* RemoteControlPropertyType) override;
 	virtual bool SupportsMasking(const UScriptStruct* InStruct) const override;
@@ -45,7 +45,7 @@ public:
 	virtual bool ResolveObject(ERCAccess AccessType, const FString& ObjectPath, const FString& PropertyName, FRCObjectReference& OutObjectRef, FString* OutErrorText = nullptr) override;
 	virtual bool ResolveObjectProperty(ERCAccess AccessType, UObject* Object, FRCFieldPathInfo PropertyPath, FRCObjectReference& OutObjectRef, FString* OutErrorText = nullptr) override;
 	virtual bool GetObjectProperties(const FRCObjectReference& ObjectAccess, IStructSerializerBackend& Backend) override;
-	virtual bool SetObjectProperties(const FRCObjectReference& ObjectAccess, IStructDeserializerBackend& Backend, ERCPayloadType InPayloadType, const TArray<uint8>& InPayload, ERCModifyOperation Operation) override;
+	virtual bool SetObjectProperties(const FRCObjectReference& ObjectAccess, IStructDeserializerBackend& Backend, ERCPayloadType InPayloadType, const TArray<uint8>& InPayload, ERCModifyOperation Operation, const ERCModifyOperationFlags ModifyOperationFlags = ERCModifyOperationFlags::None) override;
 	virtual bool ResetObjectProperties(const FRCObjectReference& ObjectAccess, const bool bAllowIntercept) override;
 	virtual bool InsertToObjectArrayProperty(int32 Index, const FRCObjectReference& ObjectAccess, IStructDeserializerBackend& Backend, ERCPayloadType InPayloadType, const TArray<uint8>& InInterceptPayload) override;
 	virtual bool RemoveFromObjectArrayProperty(int32 Index, const FRCObjectReference& ObjectAccess) override;
diff --git a/Engine/Plugins/VirtualProduction/RemoteControl/Source/RemoteControl/Private/RemoteControlProtocolBinding.cpp b/Engine/Plugins/VirtualProduction/RemoteControl/Source/RemoteControl/Private/RemoteControlProtocolBinding.cpp
--- a/Engine/Plugins/VirtualProduction/RemoteControl/Source/RemoteControl/Private/RemoteControlProtocolBinding.cpp
+++ b/Engine/Plugins/VirtualProduction/RemoteControl/Source/RemoteControl/Private/RemoteControlProtocolBinding.cpp
@@ -2,21 +2,19 @@
 
 #include "RemoteControlProtocolBinding.h"
 
+#include "Algo/MinElement.h"
+#include "Algo/Sort.h"
+#include "Backends/CborStructDeserializerBackend.h"
 #include "CborWriter.h"
 #include "Factories/IRemoteControlMaskingFactory.h"
 #include "IRemoteControlModule.h"
 #include "RemoteControlPreset.h"
 #include "RemoteControlSettings.h"
-#include "Algo/MinElement.h"
-#include "Algo/Sort.h"
-#include "Backends/CborStructDeserializerBackend.h"
 #include "Serialization/MemoryReader.h"
 #include "Serialization/MemoryWriter.h"
 #include "UObject/StructOnScope.h"
 #include "UObject/TextProperty.h"
 
-#define LOCTEXT_NAMESPACE "RemoteControl"
-
 namespace EntityInterpolation
 {
 	/** 
@@ -637,7 +635,7 @@ const FString& FRemoteControlProtocolEntity::GetRangePropertyMaxValue() const
 	return Empty;
 }
 
-bool FRemoteControlProtocolEntity::ApplyProtocolValueToProperty(double InProtocolValue)
+bool FRemoteControlProtocolEntity::ApplyProtocolValueToProperty(const double InProtocolValue, const ERCModifyOperationFlags ModifyOperationFlags)
 {
 	if (Mappings.Num() <= 1)
 	{
@@ -709,18 +707,20 @@ bool FRemoteControlProtocolEntity::ApplyProtocolValueToProperty(double InProtoco
 		}
 
 		// Cache the values.
-		IRemoteControlModule::Get().PerformMasking(MaskingOperation.ToSharedRef());
+		IRemoteControlModule::Get().PerformMasking(MaskingOperation.ToSharedRef(), ModifyOperationFlags);
 
 		// Set properties after interpolation
 		TArray<uint8> InterpolatedBuffer;
 		if (GetInterpolatedPropertyBuffer(Property, InProtocolValue, InterpolatedBuffer))
 		{
+			constexpr ERCModifyOperation ModifyOperation = ERCModifyOperation::EQUAL;
+
 			FMemoryReader MemoryReader(InterpolatedBuffer);
 			FCborStructDeserializerBackend CborStructDeserializerBackend(MemoryReader);
-			bSuccess &= IRemoteControlModule::Get().SetObjectProperties(ObjectRef, CborStructDeserializerBackend, ERCPayloadType::Cbor, InterpolatedBuffer);
+			bSuccess &= IRemoteControlModule::Get().SetObjectProperties(ObjectRef, CborStructDeserializerBackend, ERCPayloadType::Cbor, InterpolatedBuffer, ModifyOperation, ModifyOperationFlags);
 
 			// Apply the masked the values.
-			IRemoteControlModule::Get().PerformMasking(MaskingOperation.ToSharedRef());
+			IRemoteControlModule::Get().PerformMasking(MaskingOperation.ToSharedRef(), ModifyOperationFlags);
 		}
 	}
 
diff --git a/Engine/Plugins/VirtualProduction/RemoteControl/Source/RemoteControl/Public/Factories/IRemoteControlMaskingFactory.h b/Engine/Plugins/VirtualProduction/RemoteControl/Source/RemoteControl/Public/Factories/IRemoteControlMaskingFactory.h
--- a/Engine/Plugins/VirtualProduction/RemoteControl/Source/RemoteControl/Public/Factories/IRemoteControlMaskingFactory.h
+++ b/Engine/Plugins/VirtualProduction/RemoteControl/Source/RemoteControl/Public/Factories/IRemoteControlMaskingFactory.h
@@ -4,6 +4,7 @@
 
 #include "CoreMinimal.h"
 #include "IRemoteControlModule.h"
+#include "RCModifyOperationFlags.h"
 
 /**
  * Reference to a UObject or one of its properties for the purpose of masking.
@@ -72,9 +73,12 @@ public:
 
 	/**
 	 * Applies masked values to the given struct property.
-	 * @param InMaskingOperation Shared reference of the masking operation to perform.
+	 * 
+	 * @param InMaskingOperation			Shared reference of the masking operation to perform.
+	 * @param bIsInteractive				If bWithPropertyChangedEvents, defined if the property changed events are interactive or not.
+	 * @param ModifyOperationFlags			(optional) Flags that specify how the property is modified when the value is applied.
 	 */
-	virtual void ApplyMaskedValues(const TSharedRef<FRCMaskingOperation>& InMaskingOperation, bool bIsInteractive) = 0;
+	virtual void ApplyMaskedValues(const TSharedRef<FRCMaskingOperation>& InMaskingOperation, bool bIsInteractive, const ERCModifyOperationFlags ModifyOperationFlags = ERCModifyOperationFlags::None) = 0;
 
 	/**
 	 * Caches premasking values from the given struct property.
diff --git a/Engine/Plugins/VirtualProduction/RemoteControl/Source/RemoteControl/Public/IRemoteControlModule.h b/Engine/Plugins/VirtualProduction/RemoteControl/Source/RemoteControl/Public/IRemoteControlModule.h
--- a/Engine/Plugins/VirtualProduction/RemoteControl/Source/RemoteControl/Public/IRemoteControlModule.h
+++ b/Engine/Plugins/VirtualProduction/RemoteControl/Source/RemoteControl/Public/IRemoteControlModule.h
@@ -5,6 +5,7 @@
 #include "CoreMinimal.h"
 #include "Modules/ModuleInterface.h"
 #include "Modules/ModuleManager.h"
+#include "RCModifyOperationFlags.h"
 #include "RemoteControlField.h"
 #include "RemoteControlFieldPath.h"
 #include "UObject/StructOnScope.h"
@@ -314,9 +315,10 @@ public:
 
 	/**
 	 * Performs the given masking operation.
-	 * @param InMaskingOperation Masking operation to be performed.
+	 * @param InMaskingOperation	Masking operation to be performed.
+	 * @param ModifyOperationFlags	(optional) Flags that specify how the property is modified when the value is applied.
 	 */
-	virtual void PerformMasking(const TSharedRef<FRCMaskingOperation>& InMaskingOperation) = 0;
+	virtual void PerformMasking(const TSharedRef<FRCMaskingOperation>& InMaskingOperation, const ERCModifyOperationFlags ModifyOperationFlags = ERCModifyOperationFlags::None) = 0;
 
 	/**
 	 * Register a masking factory to handle that masks the supported properties.
@@ -399,9 +401,10 @@ public:
 	 * @param InPayloadType the payload type archive.
 	 * @param InInterceptPayload the payload reference archive for the interception.
 	 * @param Operation the type of operation to perform when setting the value.
+	 * @param ModifyOperationFlags (optional) Flags that specify how the property is modified when the value is applied.
 	 * @return true if the deserialization succeeded
 	 */
-	virtual bool SetObjectProperties(const FRCObjectReference& ObjectAccess, IStructDeserializerBackend& Backend, ERCPayloadType InPayloadType = ERCPayloadType::Json, const TArray<uint8>& InInterceptPayload = TArray<uint8>(), ERCModifyOperation Operation = ERCModifyOperation::EQUAL) = 0;
+	virtual bool SetObjectProperties(const FRCObjectReference& ObjectAccess, IStructDeserializerBackend& Backend, ERCPayloadType InPayloadType = ERCPayloadType::Json, const TArray<uint8>& InInterceptPayload = TArray<uint8>(), ERCModifyOperation Operation = ERCModifyOperation::EQUAL, const ERCModifyOperationFlags ModifyOperationFlags = ERCModifyOperationFlags::None) = 0;
 
 	/**
 	 * Reset the property or the object the Object Reference is pointing to
diff --git a/Engine/Plugins/VirtualProduction/RemoteControl/Source/RemoteControl/Public/RCModifyOperationFlags.h b/Engine/Plugins/VirtualProduction/RemoteControl/Source/RemoteControl/Public/RCModifyOperationFlags.h
new file mode 100644
--- /dev/null
+++ b/Engine/Plugins/VirtualProduction/RemoteControl/Source/RemoteControl/Public/RCModifyOperationFlags.h
@@ -0,0 +1,14 @@
+// Copyright Epic Games, Inc. All Rights Reserved.
+
+#pragma  once
+
+#include "GenericPlatform/GenericPlatform.h"
+#include "Misc/EnumClassFlags.h"
+
+enum class ERCModifyOperationFlags : FGenericPlatformTypes::uint8
+{
+	None = 0x00,
+	SkipPropertyChangeEvents = 0x01
+};
+
+ENUM_CLASS_FLAGS(ERCModifyOperationFlags)
diff --git a/Engine/Plugins/VirtualProduction/RemoteControl/Source/RemoteControl/Public/RemoteControlProtocolBinding.h b/Engine/Plugins/VirtualProduction/RemoteControl/Source/RemoteControl/Public/RemoteControlProtocolBinding.h
--- a/Engine/Plugins/VirtualProduction/RemoteControl/Source/RemoteControl/Public/RemoteControlProtocolBinding.h
+++ b/Engine/Plugins/VirtualProduction/RemoteControl/Source/RemoteControl/Public/RemoteControlProtocolBinding.h
@@ -4,10 +4,11 @@
 
 #include "CoreMinimal.h"
 
-#include "RemoteControlCommon.h"
+#include "Containers/UnrealString.h"
+#include "RCModifyOperationFlags.h"
 #include "RCPropertyUtilities.h"
 #include "RCTypeTraits.h"
-#include "Containers/UnrealString.h"
+#include "RemoteControlCommon.h"
 #include "UObject/Class.h"
 #include "UObject/StructOnScope.h"
 
@@ -385,10 +386,12 @@ public:
 
 	/**
 	 * Interpolate and apply protocol value to the property
-	 * @param InProtocolValue double value from the protocol
-	 * @return true of applied successfully
+	 * 
+	 * @param InProtocolValue				Double value from the protocol
+	 * @param ModifyOperationFlags			(optional) Flags that specify how the property is modified when the value is applied.
+	 * @return								True if the value was applied successfully
 	 */
-	bool ApplyProtocolValueToProperty(double InProtocolValue);
+	bool ApplyProtocolValueToProperty(const double InProtocolValue, const ERCModifyOperationFlags ModifyOperationFlags = ERCModifyOperationFlags::None);
 
 	/** 
 	 * Get bound range property. For example, the range could be bound to FFloatProperty or FIntProperty, etc.
diff --git a/Engine/Plugins/VirtualProduction/RemoteControl/Source/RemoteControlCommon/Public/RemoteControlSettings.h b/Engine/Plugins/VirtualProduction/RemoteControl/Source/RemoteControlCommon/Public/RemoteControlSettings.h
--- a/Engine/Plugins/VirtualProduction/RemoteControl/Source/RemoteControlCommon/Public/RemoteControlSettings.h
+++ b/Engine/Plugins/VirtualProduction/RemoteControl/Source/RemoteControlCommon/Public/RemoteControlSettings.h
@@ -273,8 +273,15 @@ public:
 	 * Disabling transactions improves performance but will prevent events from being transacted to Multi-User
 	 * unless using the Remote Control Interception feature.
 	 */
-	UPROPERTY(config, EditAnywhere, Category = "Remote Control")
-	bool bProtocolsGenerateTransactions = true;
+	UPROPERTY(config, EditAnywhere, AdvancedDisplay, Category = "Remote Control Protocol")
+	bool bProtocolsGenerateTransactions = false;
+
+	/** 
+	 * Experimental:  When set to true does not raise any property changed events when applying properties. Performance is greatly improved. 
+	 * Some properties may not update properly in this mode. Useful when controling many properties e.g. via DMX or OSC.
+	 */
+	UPROPERTY(config, EditAnywhere, AdvancedDisplay, Category = "Remote Control Protocol", meta = (DisplayName = "Protocol Performance Mode (Experimental)"))
+	bool bProtocolPerformanceMode = false;
 
 	/** The remote control web app bind address. */
 	UPROPERTY(config, EditAnywhere, Category = "Remote Control Web Interface", DisplayName = "Remote Control Web Interface bind address")
diff --git a/Engine/Plugins/VirtualProduction/RemoteControl/Source/RemoteControlProtocol/Private/RemoteControlProtocol.cpp b/Engine/Plugins/VirtualProduction/RemoteControl/Source/RemoteControlProtocol/Private/RemoteControlProtocol.cpp
--- a/Engine/Plugins/VirtualProduction/RemoteControl/Source/RemoteControlProtocol/Private/RemoteControlProtocol.cpp
+++ b/Engine/Plugins/VirtualProduction/RemoteControl/Source/RemoteControlProtocol/Private/RemoteControlProtocol.cpp
@@ -2,13 +2,64 @@
 
 #include "RemoteControlProtocol.h"
 
+#include "GameFramework/Actor.h"
+#include "Misc/CoreDelegates.h"
 #include "RemoteControlPreset.h"
 #include "RemoteControlProtocolBinding.h"
 #include "RemoteControlProtocolModule.h"
-
-#include "Misc/CoreDelegates.h"
+#include "RemoteControlSettings.h"
 #include "UObject/StructOnScope.h"
 
+namespace UE::RemoteControl::Protocol::Private
+{
+	/** Struct that updates actors on scope, useful when applying protocol values in performance mode */
+	struct FRCScopedUpdateActors
+	{
+		FRCScopedUpdateActors(const TMap<const FRemoteControlProtocolEntityPtr, double>& EntityToValueMap)
+		{
+			for (const TPair<const FRemoteControlProtocolEntityPtr, double>& EntityValuesToApplyPair : EntityToValueMap)
+			{
+				if (EntityValuesToApplyPair.Key.IsValid() && EntityValuesToApplyPair.Key->IsValid())
+				{
+					FRemoteControlProtocolEntity* ProtocolEntity = EntityValuesToApplyPair.Key->Get();
+					URemoteControlPreset* Preset = ProtocolEntity->GetOwner().Get();
+					if (!Preset)
+					{
+						continue;
+					}
+
+					const TSharedPtr<FRemoteControlProperty> RemoteControlProperty = Preset->GetExposedEntity<FRemoteControlProperty>(ProtocolEntity->GetPropertyId()).Pin();
+					if (!RemoteControlProperty.IsValid())
+					{
+						continue;
+					}
+
+					for (UObject* Object : RemoteControlProperty->GetBoundObjects())
+					{
+						if (AActor* Actor = Object->GetTypedOuter<AActor>())
+						{
+							Actors.AddUnique(Actor);
+						}
+					}
+				}
+			}
+		}
+
+		~FRCScopedUpdateActors()
+		{
+			for (AActor* Actor : Actors)
+			{
+				Actor->UpdateComponentTransforms();
+				Actor->MarkComponentsRenderStateDirty();
+			}
+		}
+
+	private:
+		/** The actors which need update */
+		TArray<AActor*> Actors;
+	};
+}
+
 FRemoteControlProtocol::FRemoteControlProtocol(const FName InProtocolName)
 	: ProtocolName(InProtocolName)
 {
@@ -45,6 +96,16 @@ void FRemoteControlProtocol::QueueValue(const FRemoteControlProtocolEntityPtr In
 
 void FRemoteControlProtocol::OnEndFrame()
 {
+	using namespace UE::RemoteControl::Protocol::Private;
+	
+	const bool bPerformancMode = GetDefault<URemoteControlSettings>()->bProtocolPerformanceMode;
+	TOptional<FRCScopedUpdateActors> OptionalScopedUpdateActors;
+	if (bPerformancMode)
+	{
+		OptionalScopedUpdateActors = FRCScopedUpdateActors(EntityValuesToApply);
+	}
+
+	const ERCModifyOperationFlags ModifyOperationFlags = bPerformancMode ? ERCModifyOperationFlags::SkipPropertyChangeEvents : ERCModifyOperationFlags::None;
 	for (const TPair<const FRemoteControlProtocolEntityPtr, double>& EntityValuesToApplyPair : EntityValuesToApply)
 	{
 		// Check is the Shared ptr and TStructOnScope is valid
@@ -57,7 +118,7 @@ void FRemoteControlProtocol::OnEndFrame()
 			// Check the value from previous frame
 			if (PreviousFrameValuePtr == nullptr || !FMath::IsNearlyEqual(ThisFrameValue, *PreviousFrameValuePtr))
 			{
-				if (!ProtocolEntity->ApplyProtocolValueToProperty(EntityValuesToApplyPair.Value))
+				if (!ProtocolEntity->ApplyProtocolValueToProperty(EntityValuesToApplyPair.Value, ModifyOperationFlags))
 				{
 					// Warn if the the value can't by applied
 					ensureMsgf(false, TEXT("Can't apply property for Protocol %s and PropertyId %s"),
