Remote Control: Fix undef LOCTEXT_NAMESPACE without a corresponding define LOCTEXT_NAMESPACE

#rnx

diff --git a/Engine/Plugins/VirtualProduction/RemoteControl/Source/RemoteControl/Private/RemoteControlProtocolBinding.cpp b/Engine/Plugins/VirtualProduction/RemoteControl/Source/RemoteControl/Private/RemoteControlProtocolBinding.cpp
--- a/Engine/Plugins/VirtualProduction/RemoteControl/Source/RemoteControl/Private/RemoteControlProtocolBinding.cpp
+++ b/Engine/Plugins/VirtualProduction/RemoteControl/Source/RemoteControl/Private/RemoteControlProtocolBinding.cpp
@@ -994,5 +994,3 @@ uint32 GetTypeHash(const FRemoteControlProtocolBinding& InProtocolBinding)
 {
 	return GetTypeHash(InProtocolBinding.Id);
 }
-
-#undef LOCTEXT_NAMESPACE
