Remote Control - Add option to disable console command execution through the web api

diff --git a/Engine/Plugins/VirtualProduction/RemoteControl/Source/RemoteControlCommon/Public/RemoteControlSettings.h b/Engine/Plugins/VirtualProduction/RemoteControl/Source/RemoteControlCommon/Public/RemoteControlSettings.h
--- a/Engine/Plugins/VirtualProduction/RemoteControl/Source/RemoteControlCommon/Public/RemoteControlSettings.h
+++ b/Engine/Plugins/VirtualProduction/RemoteControl/Source/RemoteControlCommon/Public/RemoteControlSettings.h
@@ -347,10 +347,14 @@ public:
 	UPROPERTY(Config, EditAnywhere, Category = "Remote Control | Security")
 	bool bRestrictServerAccess = false;
 
-	/** Enable remote python execution, enabling this could open you open to vulnerabilities if an outside actor has access to your server. */
+	/** Enable remote python execution, enabling this could open you up to vulnerabilities if an outside actor has access to your server. */
 	UPROPERTY(Config, EditAnywhere, Category = "Remote Control | Security", meta = (editCondition = bRestrictServerAccess))
 	bool bEnableRemotePythonExecution = false;
 
+	/** Enable calling 'ExecuteConsoleCommand' through the web api. Enabling this could open you up to vulnerabilities if an outside actor has access to your server. */
+	UPROPERTY(Config, EditAnywhere, Category = "Remote Control | Security", meta = (editCondition = bRestrictServerAccess))
+	bool bAllowConsoleCommandRemoteExecution = false;
+
 	/** List of IP Addresses that are allowed to access the Web API without authentication. */
 	UPROPERTY(config, EditAnywhere, Category = "Remote Control | Security", DisplayName = "Range of Allowlisted Clients", Meta = (EditCondition = bRestrictServerAccess, EditConditionHides))
 	TSet<FRCNetworkAddressRange> AllowlistedClients = { FRCNetworkAddressRange::AllowAllIPs() };
diff --git a/Engine/Plugins/VirtualProduction/RemoteControl/Source/WebRemoteControl/Private/WebRemoteControlInternalUtils.cpp b/Engine/Plugins/VirtualProduction/RemoteControl/Source/WebRemoteControl/Private/WebRemoteControlInternalUtils.cpp
--- a/Engine/Plugins/VirtualProduction/RemoteControl/Source/WebRemoteControl/Private/WebRemoteControlInternalUtils.cpp
+++ b/Engine/Plugins/VirtualProduction/RemoteControl/Source/WebRemoteControl/Private/WebRemoteControlInternalUtils.cpp
@@ -541,37 +541,54 @@ bool WebRemoteControlInternalUtils::ValidateFunctionCall(const FRCCall& InRCCall
 	{
 		return false;
 	}
+
+	UFunction* ExecuteControlCommandMethod = UKismetSystemLibrary::StaticClass()->FindFunctionByName(GET_FUNCTION_NAME_CHECKED(UKismetSystemLibrary, ExecuteConsoleCommand));
 	
-	if (!GetDefault<URemoteControlSettings>()->bEnableRemotePythonExecution
-		&& InRCCall.CallRef.Object->IsA(UKismetSystemLibrary::StaticClass())
-		&& InRCCall.CallRef.Function == UKismetSystemLibrary::StaticClass()->FindFunctionByName(GET_FUNCTION_NAME_CHECKED(UKismetSystemLibrary, ExecuteConsoleCommand)))
+	if (InRCCall.CallRef.Function == ExecuteControlCommandMethod)
 	{
-		// Make sure python is not called through KismetSystemLibrary.
-		if (const FProperty* Property = InRCCall.CallRef.Function->FindPropertyByName("Command"))
+		if (!GetDefault<URemoteControlSettings>()->bAllowConsoleCommandRemoteExecution)
 		{
-			FString FullCommand;
-			Property->GetValue_InContainer(InRCCall.ParamStruct.GetStructMemory(), &FullCommand);
+			if (OutErrorText)
+			{
+				*OutErrorText = TEXT("Executing console commands remotely is not enabled in the remote control settings.");
+			}
 
-			TArray<FString> SplitCommands;
-			FullCommand.ParseIntoArray(SplitCommands, TEXT("|"));
+			return false;
+		}
 
-			for (FString& Command : SplitCommands)
+		// Make sure python is not called through KismetSystemLibrary.
+		if (!GetDefault<URemoteControlSettings>()->bEnableRemotePythonExecution)
+		{
+			if (const FProperty* Property = InRCCall.CallRef.Function->FindPropertyByName("Command"))
 			{
-				Command.TrimStartAndEndInline();
-				if (Command.StartsWith(TEXT("py ")) || Command.StartsWith(TEXT("python ")))
+				FString FullCommand;
+				Property->GetValue_InContainer(InRCCall.ParamStruct.GetStructMemory(), &FullCommand);
+
+				TArray<FString> SplitCommands;
+				FullCommand.ParseIntoArray(SplitCommands, TEXT("|"));
+
+				for (FString& Command : SplitCommands)
 				{
-					if (OutErrorText)
+					Command.TrimStartAndEndInline();
+					if (Command.StartsWith(TEXT("py ")) || Command.StartsWith(TEXT("python ")))
 					{
-						*OutErrorText = TEXT("Executing Python remotely is not enabled in the remote control settings.");
+						if (OutErrorText)
+						{
+							*OutErrorText = TEXT("Executing Python remotely is not enabled in the remote control settings.");
+						}
+						return false;
 					}
-					return false;
 				}
 			}
-		}
-		else
-		{
-			ensureMsgf(false, TEXT("Could not find the Command parameter on ExecuteConsoleCommand"));
-			return false;
+			else
+			{
+				if (OutErrorText)
+				{
+					*OutErrorText = TEXT("Could not find the Command parameter on ExecuteConsoleCommand");
+				}
+
+				return false;
+			}
 		}
 	}
 
