Remote Control: The Protocol Generate Transactions option is now available per Preset instead of being a global option in Project Settings

diff --git a/Engine/Plugins/VirtualProduction/RemoteControl/Source/RemoteControl/Private/RemoteControlPreset.cpp b/Engine/Plugins/VirtualProduction/RemoteControl/Source/RemoteControl/Private/RemoteControlPreset.cpp
--- a/Engine/Plugins/VirtualProduction/RemoteControl/Source/RemoteControl/Private/RemoteControlPreset.cpp
+++ b/Engine/Plugins/VirtualProduction/RemoteControl/Source/RemoteControl/Private/RemoteControlPreset.cpp
@@ -39,6 +39,7 @@
 #include "Editor.h"
 #include "EngineAnalytics.h"
 #include "Engine/Blueprint.h"
+#include "RemoteControlSettings.h"
 #include "TimerManager.h"
 #include "UObject/PackageReload.h"
 #endif
@@ -1900,6 +1901,21 @@ void URemoteControlPreset::Serialize(FArchive& Ar)
 		{
 			PresetId = FGuid::NewGuid();
 		}
+		
+#if WITH_EDITOR
+		// 5.5 - Upgrade to per perset protocol modify operation flags
+		if (Ar.CustomVer(FRemoteControlObjectVersion::GUID) < FRemoteControlObjectVersion::AddedPerPresetModifyOperationFlags)
+		{
+			const URemoteControlSettings* Settings = GetDefault<URemoteControlSettings>();
+
+			PRAGMA_DISABLE_DEPRECATION_WARNINGS
+			if (Settings->bProtocolsGenerateTransactions_DEPRECATED)
+			{
+				ModifyOperationFlagsForProtocols = ERCModifyOperationFlags::None;
+			}
+			PRAGMA_ENABLE_DEPRECATION_WARNINGS
+		}
+#endif
 	}
 }
 
diff --git a/Engine/Plugins/VirtualProduction/RemoteControl/Source/RemoteControl/Private/RemoteControlProtocolBinding.cpp b/Engine/Plugins/VirtualProduction/RemoteControl/Source/RemoteControl/Private/RemoteControlProtocolBinding.cpp
--- a/Engine/Plugins/VirtualProduction/RemoteControl/Source/RemoteControl/Private/RemoteControlProtocolBinding.cpp
+++ b/Engine/Plugins/VirtualProduction/RemoteControl/Source/RemoteControl/Private/RemoteControlProtocolBinding.cpp
@@ -9,7 +9,6 @@
 #include "Factories/IRemoteControlMaskingFactory.h"
 #include "IRemoteControlModule.h"
 #include "RemoteControlPreset.h"
-#include "RemoteControlSettings.h"
 #include "Serialization/MemoryReader.h"
 #include "Serialization/MemoryWriter.h"
 #include "UObject/StructOnScope.h"
@@ -635,7 +634,7 @@ const FString& FRemoteControlProtocolEntity::GetRangePropertyMaxValue() const
 	return Empty;
 }
 
-bool FRemoteControlProtocolEntity::ApplyProtocolValueToProperty(const double InProtocolValue, const ERCModifyOperationFlags ModifyOperationFlags)
+bool FRemoteControlProtocolEntity::ApplyProtocolValueToProperty(const double InProtocolValue)
 {
 	if (Mappings.Num() <= 1)
 	{
@@ -681,9 +680,9 @@ bool FRemoteControlProtocolEntity::ApplyProtocolValueToProperty(const double InP
 	FRCObjectReference ObjectRef;
 	ObjectRef.Property = Property;
 	ObjectRef.Access = ERCAccess::WRITE_ACCESS;
-
-	const URemoteControlSettings* RemoteControlSettings = GetDefault<URemoteControlSettings>();
-	if (RemoteControlSettings->bProtocolsGenerateTransactions)
+	
+	const ERCModifyOperationFlags ModifyOperationFlags = Preset->GetModifyOperationFlagsForProtocols();
+	if (EnumHasAnyFlags(ModifyOperationFlags, ERCModifyOperationFlags::None))
 	{
 		ObjectRef.Access = ERCAccess::WRITE_TRANSACTION_ACCESS;
 	}
diff --git a/Engine/Plugins/VirtualProduction/RemoteControl/Source/RemoteControl/Public/RCModifyOperationFlags.h b/Engine/Plugins/VirtualProduction/RemoteControl/Source/RemoteControl/Public/RCModifyOperationFlags.h
--- a/Engine/Plugins/VirtualProduction/RemoteControl/Source/RemoteControl/Public/RCModifyOperationFlags.h
+++ b/Engine/Plugins/VirtualProduction/RemoteControl/Source/RemoteControl/Public/RCModifyOperationFlags.h
@@ -5,10 +5,19 @@
 #include "GenericPlatform/GenericPlatform.h"
 #include "Misc/EnumClassFlags.h"
 
-enum class ERCModifyOperationFlags : FGenericPlatformTypes::uint8
+#include "RCModifyOperationFlags.generated.h"
+
+UENUM()
+enum class ERCModifyOperationFlags : uint8
 {
+	/** Properties are modified with Transactions and Property Changed events */
 	None = 0x00,
-	SkipPropertyChangeEvents = 0x01
+
+	/** Properties are modified without Property Change events and implicitly without transactions */
+	SkipPropertyChangeEvents = 0x01,
+
+	/** Properties are modified without Transactions */
+	SkipTransactions = 0x02
 };
 
 ENUM_CLASS_FLAGS(ERCModifyOperationFlags)
diff --git a/Engine/Plugins/VirtualProduction/RemoteControl/Source/RemoteControl/Public/RemoteControlObjectVersion.h b/Engine/Plugins/VirtualProduction/RemoteControl/Source/RemoteControl/Public/RemoteControlObjectVersion.h
--- a/Engine/Plugins/VirtualProduction/RemoteControl/Source/RemoteControl/Public/RemoteControlObjectVersion.h
+++ b/Engine/Plugins/VirtualProduction/RemoteControl/Source/RemoteControl/Public/RemoteControlObjectVersion.h
@@ -27,6 +27,9 @@ struct FRemoteControlObjectVersion
 		// Added property/function flags to RemoteControlField
 		AddedFieldFlags,
 		
+		// Instead of the global 'Protocol Generates Transactions' setting, each preset can now specify its modify operation flags
+		AddedPerPresetModifyOperationFlags,
+
 		// -----<new versions can be added above this line>-------------------------------------------------
 		VersionPlusOne,
 		LatestVersion = VersionPlusOne - 1
diff --git a/Engine/Plugins/VirtualProduction/RemoteControl/Source/RemoteControl/Public/RemoteControlPreset.h b/Engine/Plugins/VirtualProduction/RemoteControl/Source/RemoteControl/Public/RemoteControlPreset.h
--- a/Engine/Plugins/VirtualProduction/RemoteControl/Source/RemoteControl/Public/RemoteControlPreset.h
+++ b/Engine/Plugins/VirtualProduction/RemoteControl/Source/RemoteControl/Public/RemoteControlPreset.h
@@ -3,6 +3,7 @@
 
 #include "Algo/Transform.h"
 #include "CoreTypes.h"
+#include "RCModifyOperationFlags.h"
 #include "RemoteControlEntity.h"
 #include "RemoteControlField.h"
 #include "RemoteControlPropertyIdRegistry.h"
@@ -889,6 +890,12 @@ public:
 	 */
 	TObjectPtr<URemoteControlPropertyIdRegistry> GetPropertyIdRegistry() const { return PropertyIdRegistry; }
 
+	/** Sets how this preset modifies properties when changed by protocols */
+	void SetModifyOperationFlagsForProtocols(ERCModifyOperationFlags InOperationFlags) { ModifyOperationFlagsForProtocols = InOperationFlags; }
+
+	/** Returns how this preset modifies properties when changed by protocols */
+	ERCModifyOperationFlags GetModifyOperationFlagsForProtocols() const { return ModifyOperationFlagsForProtocols; }
+
 private:
 
 	/** Find a binding that has the same boundobjectmap but that currently points to the object passed as argument. */
@@ -935,6 +942,10 @@ private:
 	UPROPERTY(AssetRegistrySearchable)
 	FGuid PresetId;
 
+	/** Defines how this preset modifies properties when changed by protocols */
+	UPROPERTY()
+	ERCModifyOperationFlags ModifyOperationFlagsForProtocols = ERCModifyOperationFlags::SkipPropertyChangeEvents | ERCModifyOperationFlags::SkipTransactions;
+
 	/** The cache for information about an exposed field. */
 	UPROPERTY(Transient)
 	TMap<FGuid, FRCCachedFieldData> FieldCache;
diff --git a/Engine/Plugins/VirtualProduction/RemoteControl/Source/RemoteControl/Public/RemoteControlProtocolBinding.h b/Engine/Plugins/VirtualProduction/RemoteControl/Source/RemoteControl/Public/RemoteControlProtocolBinding.h
--- a/Engine/Plugins/VirtualProduction/RemoteControl/Source/RemoteControl/Public/RemoteControlProtocolBinding.h
+++ b/Engine/Plugins/VirtualProduction/RemoteControl/Source/RemoteControl/Public/RemoteControlProtocolBinding.h
@@ -391,7 +391,7 @@ public:
 	 * @param ModifyOperationFlags			(optional) Flags that specify how the property is modified when the value is applied.
 	 * @return								True if the value was applied successfully
 	 */
-	bool ApplyProtocolValueToProperty(const double InProtocolValue, const ERCModifyOperationFlags ModifyOperationFlags = ERCModifyOperationFlags::None);
+	bool ApplyProtocolValueToProperty(const double InProtocolValue);
 
 	/** 
 	 * Get bound range property. For example, the range could be bound to FFloatProperty or FIntProperty, etc.
diff --git a/Engine/Plugins/VirtualProduction/RemoteControl/Source/RemoteControlCommon/Public/RemoteControlSettings.h b/Engine/Plugins/VirtualProduction/RemoteControl/Source/RemoteControlCommon/Public/RemoteControlSettings.h
--- a/Engine/Plugins/VirtualProduction/RemoteControl/Source/RemoteControlCommon/Public/RemoteControlSettings.h
+++ b/Engine/Plugins/VirtualProduction/RemoteControl/Source/RemoteControlCommon/Public/RemoteControlSettings.h
@@ -269,19 +269,15 @@ public:
 	}
 
 	/**
+	 *  DEPRECATED 5.5
+	 * 
 	 * Should transactions be generated for events received through protocols (ie. MIDI, DMX etc.)
 	 * Disabling transactions improves performance but will prevent events from being transacted to Multi-User
 	 * unless using the Remote Control Interception feature.
 	 */
-	UPROPERTY(config, EditAnywhere, AdvancedDisplay, Category = "Remote Control Protocol")
-	bool bProtocolsGenerateTransactions = false;
-
-	/** 
-	 * Experimental:  When set to true does not raise any property changed events when applying properties. Performance is greatly improved. 
-	 * Some properties may not update properly in this mode. Useful when controling many properties e.g. via DMX or OSC.
-	 */
-	UPROPERTY(config, EditAnywhere, AdvancedDisplay, Category = "Remote Control Protocol", meta = (DisplayName = "Protocol Performance Mode (Experimental)"))
-	bool bProtocolPerformanceMode = false;
+	UE_DEPRECATED(5.5, "Protocol Generate Transactions is now a per Preset option. See URemoteControlPreset::GetProtocolModifyOperationFlags.")
+	UPROPERTY(config)
+	bool bProtocolsGenerateTransactions_DEPRECATED = true;
 
 	/** The remote control web app bind address. */
 	UPROPERTY(config, EditAnywhere, Category = "Remote Control Web Interface", DisplayName = "Remote Control Web Interface bind address")
diff --git a/Engine/Plugins/VirtualProduction/RemoteControl/Source/RemoteControlProtocol/Private/RemoteControlProtocol.cpp b/Engine/Plugins/VirtualProduction/RemoteControl/Source/RemoteControlProtocol/Private/RemoteControlProtocol.cpp
--- a/Engine/Plugins/VirtualProduction/RemoteControl/Source/RemoteControlProtocol/Private/RemoteControlProtocol.cpp
+++ b/Engine/Plugins/VirtualProduction/RemoteControl/Source/RemoteControlProtocol/Private/RemoteControlProtocol.cpp
@@ -4,62 +4,13 @@
 
 #include "GameFramework/Actor.h"
 #include "Misc/CoreDelegates.h"
+#include "RCModifyOperationFlags.h"
 #include "RemoteControlPreset.h"
 #include "RemoteControlProtocolBinding.h"
 #include "RemoteControlProtocolModule.h"
 #include "RemoteControlSettings.h"
 #include "UObject/StructOnScope.h"
 
-namespace UE::RemoteControl::Protocol::Private
-{
-	/** Struct that updates actors on scope, useful when applying protocol values in performance mode */
-	struct FRCScopedUpdateActors
-	{
-		FRCScopedUpdateActors(const TMap<const FRemoteControlProtocolEntityPtr, double>& EntityToValueMap)
-		{
-			for (const TPair<const FRemoteControlProtocolEntityPtr, double>& EntityValuesToApplyPair : EntityToValueMap)
-			{
-				if (EntityValuesToApplyPair.Key.IsValid() && EntityValuesToApplyPair.Key->IsValid())
-				{
-					FRemoteControlProtocolEntity* ProtocolEntity = EntityValuesToApplyPair.Key->Get();
-					URemoteControlPreset* Preset = ProtocolEntity->GetOwner().Get();
-					if (!Preset)
-					{
-						continue;
-					}
-
-					const TSharedPtr<FRemoteControlProperty> RemoteControlProperty = Preset->GetExposedEntity<FRemoteControlProperty>(ProtocolEntity->GetPropertyId()).Pin();
-					if (!RemoteControlProperty.IsValid())
-					{
-						continue;
-					}
-
-					for (UObject* Object : RemoteControlProperty->GetBoundObjects())
-					{
-						if (AActor* Actor = Object->GetTypedOuter<AActor>())
-						{
-							Actors.AddUnique(Actor);
-						}
-					}
-				}
-			}
-		}
-
-		~FRCScopedUpdateActors()
-		{
-			for (AActor* Actor : Actors)
-			{
-				Actor->UpdateComponentTransforms();
-				Actor->MarkComponentsRenderStateDirty();
-			}
-		}
-
-	private:
-		/** The actors which need update */
-		TArray<AActor*> Actors;
-	};
-}
-
 FRemoteControlProtocol::FRemoteControlProtocol(const FName InProtocolName)
 	: ProtocolName(InProtocolName)
 {
@@ -96,29 +47,41 @@ void FRemoteControlProtocol::QueueValue(const FRemoteControlProtocolEntityPtr In
 
 void FRemoteControlProtocol::OnEndFrame()
 {
-	using namespace UE::RemoteControl::Protocol::Private;
-	
-	const bool bPerformancMode = GetDefault<URemoteControlSettings>()->bProtocolPerformanceMode;
-	TOptional<FRCScopedUpdateActors> OptionalScopedUpdateActors;
-	if (bPerformancMode)
-	{
-		OptionalScopedUpdateActors = FRCScopedUpdateActors(EntityValuesToApply);
-	}
-
-	const ERCModifyOperationFlags ModifyOperationFlags = bPerformancMode ? ERCModifyOperationFlags::SkipPropertyChangeEvents : ERCModifyOperationFlags::None;
+	TArray<AActor*> ActorsWithoutPropertyChangedEvents;
 	for (const TPair<const FRemoteControlProtocolEntityPtr, double>& EntityValuesToApplyPair : EntityValuesToApply)
 	{
 		// Check is the Shared ptr and TStructOnScope is valid
 		if (EntityValuesToApplyPair.Key.IsValid() && EntityValuesToApplyPair.Key->IsValid())
 		{
 			FRemoteControlProtocolEntity* ProtocolEntity = EntityValuesToApplyPair.Key->Get();
+
 			const double ThisFrameValue = EntityValuesToApplyPair.Value;
 			const double* PreviousFrameValuePtr = PreviousTickValuesToApply.Find(EntityValuesToApplyPair.Key);
-			
+
 			// Check the value from previous frame
 			if (PreviousFrameValuePtr == nullptr || !FMath::IsNearlyEqual(ThisFrameValue, *PreviousFrameValuePtr))
 			{
-				if (!ProtocolEntity->ApplyProtocolValueToProperty(EntityValuesToApplyPair.Value, ModifyOperationFlags))
+				// Remember actors that will not get property change events
+				URemoteControlPreset* Preset = ProtocolEntity->GetOwner().Get();
+				if (Preset && EnumHasAnyFlags(Preset->GetModifyOperationFlagsForProtocols(), ERCModifyOperationFlags::SkipPropertyChangeEvents))
+				{
+					const TSharedPtr<FRemoteControlProperty> RemoteControlProperty = Preset->GetExposedEntity<FRemoteControlProperty>(ProtocolEntity->GetPropertyId()).Pin();
+					if (!RemoteControlProperty.IsValid())
+					{
+						continue;
+					}
+
+					for (UObject* Object : RemoteControlProperty->GetBoundObjects())
+					{
+						if (AActor* Actor = Object->GetTypedOuter<AActor>())
+						{
+							ActorsWithoutPropertyChangedEvents.AddUnique(Actor);
+						}
+					}
+				}
+
+				// Apply the protocol value
+				if (!ProtocolEntity->ApplyProtocolValueToProperty(EntityValuesToApplyPair.Value))
 				{
 					// Warn if the the value can't by applied
 					ensureMsgf(false, TEXT("Can't apply property for Protocol %s and PropertyId %s"),
@@ -128,6 +91,13 @@ void FRemoteControlProtocol::OnEndFrame()
 		}
 	}
 
+	// Update actors that did not get property changed events
+	for (AActor* Actor : ActorsWithoutPropertyChangedEvents)
+	{
+		Actor->UpdateComponentTransforms();
+		Actor->MarkComponentsRenderStateDirty();
+	}
+
 	// Move the values from this frame to cached map
 	PreviousTickValuesToApply = MoveTemp(EntityValuesToApply);
 }
diff --git a/Engine/Plugins/VirtualProduction/RemoteControl/Source/RemoteControlProtocol/Public/IRemoteControlProtocol.h b/Engine/Plugins/VirtualProduction/RemoteControl/Source/RemoteControlProtocol/Public/IRemoteControlProtocol.h
--- a/Engine/Plugins/VirtualProduction/RemoteControl/Source/RemoteControlProtocol/Public/IRemoteControlProtocol.h
+++ b/Engine/Plugins/VirtualProduction/RemoteControl/Source/RemoteControlProtocol/Public/IRemoteControlProtocol.h
@@ -3,6 +3,7 @@
 #pragma once
 
 #include "CoreMinimal.h"
+#include "RCModifyOperationFlags.h"
 
 #include "UObject/StructOnScope.h"
 
diff --git a/Engine/Plugins/VirtualProduction/RemoteControl/Source/RemoteControlUI/Private/Commands/RemoteControlCommands.cpp b/Engine/Plugins/VirtualProduction/RemoteControl/Source/RemoteControlUI/Private/Commands/RemoteControlCommands.cpp
--- a/Engine/Plugins/VirtualProduction/RemoteControl/Source/RemoteControlUI/Private/Commands/RemoteControlCommands.cpp
+++ b/Engine/Plugins/VirtualProduction/RemoteControl/Source/RemoteControlUI/Private/Commands/RemoteControlCommands.cpp
@@ -102,6 +102,20 @@ void FRemoteControlCommands::RegisterCommands()
 		"Update the selected UI item value with the one in the fields list",
 		EUserInterfaceActionType::Button,
 		FInputChord(EKeys::U, EModifierKey::Control));
+
+	// Protocols generate property changed events
+	UI_COMMAND(ProtocolsGeneratePropertyChangeEvents,
+		"Protocols Generate Property Change Events",
+		"When checked, protocols generate Property change events whenever they change a property.\nEnabling this option may impact performance significantly, but may be required for certain editor controls.",
+		EUserInterfaceActionType::ToggleButton,
+		FInputChord());
+
+	// Protocols generate transactions
+	UI_COMMAND(ProtocolsGenerateTransactions,
+		"Protocols Generate Transactions",
+		"When checked, protocols generate Undo Redo events whenever they change a property.\nEnabling this option is not recommended as protocols will raise many transactions and impact performance significantly.",
+		EUserInterfaceActionType::ToggleButton,
+		FInputChord());
 }
 
 #undef LOCTEXT_NAMESPACE
diff --git a/Engine/Plugins/VirtualProduction/RemoteControl/Source/RemoteControlUI/Private/Commands/RemoteControlCommands.h b/Engine/Plugins/VirtualProduction/RemoteControl/Source/RemoteControlUI/Private/Commands/RemoteControlCommands.h
--- a/Engine/Plugins/VirtualProduction/RemoteControl/Source/RemoteControlUI/Private/Commands/RemoteControlCommands.h
+++ b/Engine/Plugins/VirtualProduction/RemoteControl/Source/RemoteControlUI/Private/Commands/RemoteControlCommands.h
@@ -82,7 +82,17 @@ public:
 	TSharedPtr<FUICommandInfo> DuplicateItem;
 
 	/**
-	 * UI Command for updating tha action in the action list with the value in the field list. Currently used for Logic panel
+	 * UI Command for updating the action in the action list with the value in the field list. Currently used for Logic panel
 	 */
 	TSharedPtr<FUICommandInfo> UpdateValue;
+
+	/**
+	 * UI Command to Toggle if protocols generate property changed events
+	 */
+	TSharedPtr<FUICommandInfo> ProtocolsGeneratePropertyChangeEvents;
+
+	/**
+	 * UI Command to Toggle if protocols generate transactions
+	 */
+	TSharedPtr<FUICommandInfo> ProtocolsGenerateTransactions;
 };
diff --git a/Engine/Plugins/VirtualProduction/RemoteControl/Source/RemoteControlUI/Private/UI/SRCPanelExposedEntitiesList.cpp b/Engine/Plugins/VirtualProduction/RemoteControl/Source/RemoteControlUI/Private/UI/SRCPanelExposedEntitiesList.cpp
--- a/Engine/Plugins/VirtualProduction/RemoteControl/Source/RemoteControlUI/Private/UI/SRCPanelExposedEntitiesList.cpp
+++ b/Engine/Plugins/VirtualProduction/RemoteControl/Source/RemoteControlUI/Private/UI/SRCPanelExposedEntitiesList.cpp
@@ -25,6 +25,7 @@
 #include "SRCPanelExposedEntitiesGroup.h"
 #include "SRCPanelExposedField.h"
 #include "SRCPanelFieldGroup.h"
+#include "SRCPanelPerPresetProtocolSettings.h"
 #include "SSearchToggleButton.h"
 #include "Styling/RemoteControlStyles.h"
 #include "UI/IRCExposedEntitiesPanelExtender.h"
@@ -38,7 +39,6 @@
 #include "Widgets/Views/STableRow.h"
 #include "Widgets/Views/STreeView.h"
 
-
 #define LOCTEXT_NAMESPACE "RemoteControlPanelEntitiesList"
 
 TSet<FName> SRCPanelExposedEntitiesList::DefaultProtocolColumns = {
@@ -525,12 +525,21 @@ void SRCPanelExposedEntitiesList::Construct(const FArguments& InArgs, URemoteCon
 			GetGroupMenuContentWidget()
 		];
 
+	// Per Preset Protocol Settings
+	const TSharedRef<SRCPanelPerPresetProtocolSettings> PerPresetProtocolSettings =
+		SNew(SRCPanelPerPresetProtocolSettings, Preset.Get())
+		.Visibility_Lambda([this]()
+			{
+				return bIsInProtocolsMode.Get() ? EVisibility::Visible : EVisibility::Collapsed;
+			});
+
 	// Expose Button
 	ExposeDockPanel->AddHeaderToolbarItem(EToolbar::Center, SearchBoxPtr.ToSharedRef());
 	ExposeDockPanel->AddHeaderToolbarItem(EToolbar::Right, SNew(SSearchToggleButton, SearchBoxPtr.ToSharedRef()));
 	ExposeDockPanel->AddHeaderToolbarItem(EToolbar::Right, ComboButtonGroupButton.ToSharedRef());
 	ExposeDockPanel->AddHeaderToolbarItem(EToolbar::Right, FilterComboButton.ToSharedRef());
 	ExposeDockPanel->AddHeaderToolbarItem(EToolbar::Right, ModeSwitcher.ToSharedRef());
+	ExposeDockPanel->AddHeaderToolbarItem(EToolbar::Right, PerPresetProtocolSettings);
 	ExposeDockPanel->AddHeaderToolbarItem(EToolbar::Left, InArgs._ExposeActorsComboButton.Get().ToSharedRef());
 	ExposeDockPanel->AddHeaderToolbarItem(EToolbar::Left, InArgs._ExposeFunctionsComboButton.Get().ToSharedRef());
 
diff --git a/Engine/Plugins/VirtualProduction/RemoteControl/Source/RemoteControlUI/Private/UI/SRCPanelPerPresetProtocolSettings.cpp b/Engine/Plugins/VirtualProduction/RemoteControl/Source/RemoteControlUI/Private/UI/SRCPanelPerPresetProtocolSettings.cpp
new file mode 100644
--- /dev/null
+++ b/Engine/Plugins/VirtualProduction/RemoteControl/Source/RemoteControlUI/Private/UI/SRCPanelPerPresetProtocolSettings.cpp
@@ -0,0 +1,133 @@
+// Copyright Epic Games, Inc. All Rights Reserved.
+
+#include "SRCPanelPerPresetProtocolSettings.h"
+
+#include "Commands/RemoteControlCommands.h"
+#include "Framework/Commands/UICommandList.h"
+#include "Framework/MultiBox/MultiBoxBuilder.h"
+#include "RCModifyOperationFlags.h"
+#include "RemoteControlPreset.h"
+#include "Widgets/Images/SImage.h"
+#include "Widgets/Input/SComboButton.h"
+
+#define LOCTEXT_NAMESPACE "SRCPanelPerPresetProtocolSettings"
+
+void SRCPanelPerPresetProtocolSettings::Construct(const FArguments& InArgs, TWeakObjectPtr<URemoteControlPreset> InPreset)
+{
+	WeakPreset = InPreset;
+
+	SetupCommandList();
+
+	ChildSlot
+	[
+		SNew(SComboButton)
+		.ButtonStyle(FAppStyle::Get(), "SimpleButton")
+		.HasDownArrow(false)
+		.OnGetMenuContent(this, &SRCPanelPerPresetProtocolSettings::CreateSettingsMenu)
+		.ButtonContent()
+		[
+			SNew(SImage)
+			.Image(FAppStyle::Get().GetBrush("Icons.Settings"))
+			.ColorAndOpacity(FSlateColor::UseForeground())
+		]
+	];
+}
+
+void SRCPanelPerPresetProtocolSettings::SetupCommandList()
+{
+	CommandList = MakeShared<FUICommandList>();
+
+	CommandList->MapAction(FRemoteControlCommands::Get().ProtocolsGeneratePropertyChangeEvents,
+		FExecuteAction::CreateSP(this, &SRCPanelPerPresetProtocolSettings::ToggleProtocolsGeneratePropertyChangeEvents),
+		FCanExecuteAction(),
+		FIsActionChecked::CreateSP(this, &SRCPanelPerPresetProtocolSettings::AreProtocolsGeneratingPropertyChangedEvents)
+	);
+
+	CommandList->MapAction(FRemoteControlCommands::Get().ProtocolsGenerateTransactions,
+		FExecuteAction::CreateSP(this, &SRCPanelPerPresetProtocolSettings::ToggleProtocolsGenerateTransactions),
+		FCanExecuteAction::CreateSP(this, &SRCPanelPerPresetProtocolSettings::AreProtocolsGeneratingPropertyChangedEvents),
+		FIsActionChecked::CreateSP(this, &SRCPanelPerPresetProtocolSettings::AreProtocolsGeneratingTransactions)
+	);
+}
+
+TSharedRef<SWidget> SRCPanelPerPresetProtocolSettings::CreateSettingsMenu() const
+{
+	constexpr bool bShouldCloseWindowAfterMenuSelection = true;
+	const FName NoExtensionHook = NAME_None;
+
+	FMenuBuilder MenuBuilder(bShouldCloseWindowAfterMenuSelection, CommandList);
+
+	MenuBuilder.BeginSection("ProtocolSettingsSection", LOCTEXT("ProtocolsSectionLabel", "Protocol Performance"));
+	{
+		MenuBuilder.AddMenuEntry(FRemoteControlCommands::Get().ProtocolsGeneratePropertyChangeEvents);
+		MenuBuilder.AddMenuEntry(FRemoteControlCommands::Get().ProtocolsGenerateTransactions);
+	}
+	MenuBuilder.EndSection();
+
+	return MenuBuilder.MakeWidget();
+}
+
+void SRCPanelPerPresetProtocolSettings::ToggleProtocolsGeneratePropertyChangeEvents()
+{
+	if (URemoteControlPreset* Preset = WeakPreset.Get())
+	{
+		ERCModifyOperationFlags Flags = Preset->GetModifyOperationFlagsForProtocols();
+
+		if (EnumHasAnyFlags(Preset->GetModifyOperationFlagsForProtocols(), ERCModifyOperationFlags::SkipPropertyChangeEvents))
+		{
+			EnumRemoveFlags(Flags, ERCModifyOperationFlags::SkipPropertyChangeEvents);
+			Preset->SetModifyOperationFlagsForProtocols(Flags);
+		}
+		else
+		{
+			EnumAddFlags(Flags, ERCModifyOperationFlags::SkipPropertyChangeEvents);
+			Preset->SetModifyOperationFlagsForProtocols(Flags);
+		}
+	}
+}
+
+bool SRCPanelPerPresetProtocolSettings::AreProtocolsGeneratingPropertyChangedEvents()
+{
+	if (const URemoteControlPreset* Preset = WeakPreset.Get())
+	{
+		return !EnumHasAnyFlags(Preset->GetModifyOperationFlagsForProtocols(), ERCModifyOperationFlags::SkipPropertyChangeEvents);
+	}
+
+	return false;
+}
+
+void SRCPanelPerPresetProtocolSettings::ToggleProtocolsGenerateTransactions()
+{
+	if (URemoteControlPreset* Preset = WeakPreset.Get())
+	{
+		ERCModifyOperationFlags Flags = Preset->GetModifyOperationFlagsForProtocols();
+
+		if (EnumHasAnyFlags(Preset->GetModifyOperationFlagsForProtocols(), ERCModifyOperationFlags::SkipTransactions))
+		{
+			EnumRemoveFlags(Flags, ERCModifyOperationFlags::SkipTransactions);
+			Preset->SetModifyOperationFlagsForProtocols(Flags);
+		}
+		else
+		{
+			EnumAddFlags(Flags, ERCModifyOperationFlags::SkipTransactions);
+			Preset->SetModifyOperationFlagsForProtocols(Flags);
+		}
+	}
+}
+
+bool SRCPanelPerPresetProtocolSettings::AreProtocolsGeneratingTransactions()
+{
+	if (!AreProtocolsGeneratingPropertyChangedEvents())
+	{
+		return false;
+	}
+
+	if (const URemoteControlPreset* Preset = WeakPreset.Get())
+	{
+		return !EnumHasAnyFlags(Preset->GetModifyOperationFlagsForProtocols(), ERCModifyOperationFlags::SkipTransactions);
+	}
+
+	return false;
+}
+
+#undef LOCTEXT_NAMESPACE
diff --git a/Engine/Plugins/VirtualProduction/RemoteControl/Source/RemoteControlUI/Private/UI/SRCPanelPerPresetProtocolSettings.h b/Engine/Plugins/VirtualProduction/RemoteControl/Source/RemoteControlUI/Private/UI/SRCPanelPerPresetProtocolSettings.h
new file mode 100644
--- /dev/null
+++ b/Engine/Plugins/VirtualProduction/RemoteControl/Source/RemoteControlUI/Private/UI/SRCPanelPerPresetProtocolSettings.h
@@ -0,0 +1,47 @@
+// Copyright Epic Games, Inc. All Rights Reserved.
+#pragma once
+
+#include "Widgets/SCompoundWidget.h"
+
+class FUICommandList;
+class URemoteControlPreset;
+
+/**
+ * Widget that displays a picker for blueprint functions.
+ */
+class SRCPanelPerPresetProtocolSettings : public SCompoundWidget
+{
+public:
+	SLATE_BEGIN_ARGS(SRCPanelPerPresetProtocolSettings)
+	{}
+
+	SLATE_END_ARGS()
+
+	/** Constructs this widget */
+	void Construct(const FArguments& InArgs, TWeakObjectPtr<URemoteControlPreset> InPreset);
+
+private:
+	/** Creates the command list to use for this widget */
+	void SetupCommandList();
+
+	/** Creates the settings menu */
+	TSharedRef<SWidget> CreateSettingsMenu() const;
+
+	/** Called when the Protocols Generate Property Changed Events setting changed */
+	void ToggleProtocolsGeneratePropertyChangeEvents();
+
+	/** Returns true if protocols are generating Property Changed Events */
+	bool AreProtocolsGeneratingPropertyChangedEvents();
+
+	/** Called when the Protocols Generate Transactions setting changed */
+	void ToggleProtocolsGenerateTransactions();
+
+	/** Returns true if protocols are generating transactions */
+	bool AreProtocolsGeneratingTransactions();
+
+	/** The command list for this widget */
+	TSharedPtr<FUICommandList> CommandList;
+
+	/** The preset for which the settings are made */
+	TWeakObjectPtr<URemoteControlPreset> WeakPreset;
+};
