Remote Control: Get Protocol Names to return in alphabetical order

diff --git a/Engine/Plugins/VirtualProduction/RemoteControl/Source/RemoteControlProtocol/Private/RemoteControlProtocolModule.cpp b/Engine/Plugins/VirtualProduction/RemoteControl/Source/RemoteControlProtocol/Private/RemoteControlProtocolModule.cpp
--- a/Engine/Plugins/VirtualProduction/RemoteControl/Source/RemoteControlProtocol/Private/RemoteControlProtocolModule.cpp
+++ b/Engine/Plugins/VirtualProduction/RemoteControl/Source/RemoteControlProtocol/Private/RemoteControlProtocolModule.cpp
@@ -40,6 +40,7 @@ TArray<FName> FRemoteControlProtocolModule::GetProtocolNames() const
 {
 	TArray<FName> ProtocolNames;
 	Protocols.GenerateKeyArray(ProtocolNames);
+	ProtocolNames.Sort(FNameLexicalLess());
 	return ProtocolNames;
 }
 
