Remote Control: Added function to check whether a Property from an Object is exposed in RC. This is used primarily for Signatures to not double expose the same property but rather unexpose the existing and expose a new one with the given actions of the protocol.

diff --git a/Engine/Plugins/VirtualProduction/RemoteControl/Source/RemoteControl/Private/RCSignature.cpp b/Engine/Plugins/VirtualProduction/RemoteControl/Source/RemoteControl/Private/RCSignature.cpp
--- a/Engine/Plugins/VirtualProduction/RemoteControl/Source/RemoteControl/Private/RCSignature.cpp
+++ b/Engine/Plugins/VirtualProduction/RemoteControl/Source/RemoteControl/Private/RCSignature.cpp
@@ -133,7 +133,7 @@ int32 FRCSignature::ApplySignature(URemoteControlPreset* InPreset, TConstArrayVi
 		return false;
 	}
 
-	int32 ExposeCount = 0;
+	int32 AffectedCount = 0;
 
 	FRemoteControlPresetExposeArgs ExposeArgs;
 
@@ -164,14 +164,21 @@ int32 FRCSignature::ApplySignature(URemoteControlPreset* InPreset, TConstArrayVi
 			FRCSignatureActionContext ActionContext;
 			ActionContext.Preset = InPreset;
 			ActionContext.Object = Object;
-			ActionContext.Property = InPreset->ExposeProperty(Context, Field.FieldPath, ExposeArgs).Pin();
+			ActionContext.Property = InPreset->FindExposedProperty(Context, Field.FieldPath);
+
+			// Expose the existing property if not existing already.
+			if (!ActionContext.Property.IsValid())
+			{
+				ActionContext.Property = InPreset->ExposeProperty(Context, Field.FieldPath, ExposeArgs).Pin();
+			}
+
 			if (!ActionContext.Property.IsValid())
 			{
 				continue;
 			}
 
-			// Property was exposed successfully, increase expose count
-			++ExposeCount;
+			// Property was exposed successfully, increase affected count
+			++AffectedCount;
 
 			// Execute the actions for the newly exposed property
 			for (const FRCSignatureActionInstance& Action : Field.Actions)
@@ -181,5 +188,5 @@ int32 FRCSignature::ApplySignature(URemoteControlPreset* InPreset, TConstArrayVi
 		}
 	}
 
-	return ExposeCount;
+	return AffectedCount;
 }
diff --git a/Engine/Plugins/VirtualProduction/RemoteControl/Source/RemoteControl/Private/RemoteControlPreset.cpp b/Engine/Plugins/VirtualProduction/RemoteControl/Source/RemoteControl/Private/RemoteControlPreset.cpp
--- a/Engine/Plugins/VirtualProduction/RemoteControl/Source/RemoteControl/Private/RemoteControlPreset.cpp
+++ b/Engine/Plugins/VirtualProduction/RemoteControl/Source/RemoteControl/Private/RemoteControlPreset.cpp
@@ -775,6 +775,27 @@ TWeakPtr<FRemoteControlActor> URemoteControlPreset::ExposeActor(AActor* Actor, F
 	return StaticCastSharedPtr<FRemoteControlActor>(Expose(MoveTemp(RCActor), FRemoteControlActor::StaticStruct(), Args.GroupId));
 }
 
+TSharedPtr<FRemoteControlProperty> URemoteControlPreset::FindExposedProperty(UObject* InOuterObject, const FRCFieldPathInfo& InFieldPath) const
+{
+	if (!InOuterObject || !Registry)
+	{
+		return nullptr;
+	}
+
+	const FString FieldPathStr = InFieldPath.ToString();
+
+	const TArray<UObject*> OuterObjects = { InOuterObject };
+
+	for (TSharedPtr<FRemoteControlProperty> Property : Registry->GetExposedEntities<FRemoteControlProperty>())
+	{
+		if (Property->FieldPathInfo.ToString() == FieldPathStr && Property->ContainsBoundObjects(OuterObjects))
+		{
+			return Property;
+		}
+	}
+
+	return nullptr;
+}
 
 FName URemoteControlPreset::GenerateUniqueLabel(const FName InDesiredName) const
 {
diff --git a/Engine/Plugins/VirtualProduction/RemoteControl/Source/RemoteControl/Public/RemoteControlPreset.h b/Engine/Plugins/VirtualProduction/RemoteControl/Source/RemoteControl/Public/RemoteControlPreset.h
--- a/Engine/Plugins/VirtualProduction/RemoteControl/Source/RemoteControl/Public/RemoteControlPreset.h
+++ b/Engine/Plugins/VirtualProduction/RemoteControl/Source/RemoteControl/Public/RemoteControlPreset.h
@@ -409,6 +409,14 @@ public:
 	 */
 	TWeakPtr<FRemoteControlActor> ExposeActor(AActor* Actor, FRemoteControlPresetExposeArgs Args = FRemoteControlPresetExposeArgs());
 
+	/**
+	 * Gets a given object's property exposed property guid in this preset
+	 * @param InOuterObject the object to check
+	 * @param InFieldPath the path to the property
+	 * @return The exposed property if found, nullptr otherwise
+	 */
+	TSharedPtr<FRemoteControlProperty> FindExposedProperty(UObject* InOuterObject, const FRCFieldPathInfo& InFieldPath) const;
+
 	/**
 	 * Expose a property on this preset.
 	 * @param Object the object that holds the property.
diff --git a/Engine/Plugins/VirtualProduction/RemoteControl/Source/RemoteControlLogic/Private/Signature/RCSignatureControllerAction.cpp b/Engine/Plugins/VirtualProduction/RemoteControl/Source/RemoteControlLogic/Private/Signature/RCSignatureControllerAction.cpp
--- a/Engine/Plugins/VirtualProduction/RemoteControl/Source/RemoteControlLogic/Private/Signature/RCSignatureControllerAction.cpp
+++ b/Engine/Plugins/VirtualProduction/RemoteControl/Source/RemoteControlLogic/Private/Signature/RCSignatureControllerAction.cpp
@@ -91,3 +91,12 @@ URCController* FRCSignatureControllerAction::FindOrAddController(const FRCSignat
 
 	return Controller;
 }
+
+URCAction* FRCSignatureControllerAction::FindActionInBehavior(const FRCSignatureActionContext& InContext, const URCBehaviour& InBehavior) const
+{
+	if (InBehavior.ActionContainer)
+	{
+		return InBehavior.ActionContainer->FindActionByFieldId(InContext.Property->GetId());
+	}
+	return nullptr;
+}
diff --git a/Engine/Plugins/VirtualProduction/RemoteControl/Source/RemoteControlLogic/Private/Signature/RCSignatureControllerAction.h b/Engine/Plugins/VirtualProduction/RemoteControl/Source/RemoteControlLogic/Private/Signature/RCSignatureControllerAction.h
--- a/Engine/Plugins/VirtualProduction/RemoteControl/Source/RemoteControlLogic/Private/Signature/RCSignatureControllerAction.h
+++ b/Engine/Plugins/VirtualProduction/RemoteControl/Source/RemoteControlLogic/Private/Signature/RCSignatureControllerAction.h
@@ -5,6 +5,8 @@
 #include "RCSignatureAction.h"
 #include "RCSignatureControllerAction.generated.h"
 
+class URCAction;
+class URCBehaviour;
 class URCController;
 struct FPropertyBagPropertyDesc;
 
@@ -65,6 +67,8 @@ struct FRCSignatureControllerAction : public FRCSignatureAction
 
 	URCController* FindOrAddController(const FRCSignatureActionContext& InContext, bool* bOutControllerAdded = nullptr) const;
 
+	URCAction* FindActionInBehavior(const FRCSignatureActionContext& InContext, const URCBehaviour& InBehavior) const;
+
 	UPROPERTY(EditAnywhere, Category = "Controller")
 	FName ControllerName;
 
diff --git a/Engine/Plugins/VirtualProduction/RemoteControl/Source/RemoteControlLogic/Private/Signature/RCSignatureControllerBindAction.cpp b/Engine/Plugins/VirtualProduction/RemoteControl/Source/RemoteControlLogic/Private/Signature/RCSignatureControllerBindAction.cpp
--- a/Engine/Plugins/VirtualProduction/RemoteControl/Source/RemoteControlLogic/Private/Signature/RCSignatureControllerBindAction.cpp
+++ b/Engine/Plugins/VirtualProduction/RemoteControl/Source/RemoteControlLogic/Private/Signature/RCSignatureControllerBindAction.cpp
@@ -114,7 +114,14 @@ bool FRCSignatureControllerBindAction::ExecuteControllerAction(const FRCSignatur
 		return false;
 	}
 
-	URCAction* BindAction = BindBehavior->AddPropertyBindAction(RCProperty);
+	URCAction* BindAction = FindActionInBehavior(InContext, *BindBehavior);
+	if (!BindAction)
+	{
+		BindAction = BindBehavior->AddPropertyBindAction(RCProperty);
+	}
+
+	check(BindAction);
 	BindAction->Execute();
+
 	return true;
 }
diff --git a/Engine/Plugins/VirtualProduction/RemoteControl/Source/RemoteControlUI/Private/UI/Signature/Items/RCSignatureTreeSignatureItem.cpp b/Engine/Plugins/VirtualProduction/RemoteControl/Source/RemoteControlUI/Private/UI/Signature/Items/RCSignatureTreeSignatureItem.cpp
--- a/Engine/Plugins/VirtualProduction/RemoteControl/Source/RemoteControlUI/Private/UI/Signature/Items/RCSignatureTreeSignatureItem.cpp
+++ b/Engine/Plugins/VirtualProduction/RemoteControl/Source/RemoteControlUI/Private/UI/Signature/Items/RCSignatureTreeSignatureItem.cpp
@@ -1,6 +1,7 @@
 // Copyright Epic Games, Inc. All Rights Reserved.
 
 #include "RCSignatureTreeSignatureItem.h"
+#include "Framework/Notifications/NotificationManager.h"
 #include "IRemoteControlUIModule.h"
 #include "RCSignatureRegistry.h"
 #include "RCSignatureTreeFieldItem.h"
@@ -8,6 +9,7 @@
 #include "ScopedTransaction.h"
 #include "UI/SRemoteControlPanel.h"
 #include "UI/Signature/SRCSignatureTree.h"
+#include "Widgets/Notifications/SNotificationList.h"
 
 #define LOCTEXT_NAMESPACE "RCSignatureTreeSignatureItem"
 
@@ -96,9 +98,19 @@ void FRCSignatureTreeSignatureItem::ApplySignature(TConstArrayView<TWeakObjectPt
 
 	FScopedTransaction Transaction(LOCTEXT("ApplySignature", "Apply Signature"));
 
-	int32 ExposedCount = Signature->ApplySignature(Preset, InObjects);
+	int32 AffectedCount = Signature->ApplySignature(Preset, InObjects);
 
-	if (ExposedCount == 0)
+	if (AffectedCount > 0)
+	{
+		const FText MessageFormat = LOCTEXT("SignatureAppliedMessage", "Signature applied to {0} property entities.");
+
+		FNotificationInfo NotificationInfo(FText::Format(MessageFormat, FText::AsNumber(AffectedCount)));
+		NotificationInfo.ExpireDuration = 3.f;
+		NotificationInfo.bFireAndForget = true;
+
+		FSlateNotificationManager::Get().AddNotification(NotificationInfo);
+	}
+	else
 	{
 		// Nothing was exposed, cancel transaction
 		Transaction.Cancel();
