Remote Control Signature: Changed Masking Channel Type to be a Checkbox rather than an Enum for readability as the default use case is the major use case and the extra option is a  very minor one. In Additionally, renamed Scalar to Multiply.

diff --git a/Engine/Plugins/VirtualProduction/RemoteControl/Source/RemoteControlProtocol/Private/Signature/RCSignatureProtocolAction.cpp b/Engine/Plugins/VirtualProduction/RemoteControl/Source/RemoteControlProtocol/Private/Signature/RCSignatureProtocolAction.cpp
--- a/Engine/Plugins/VirtualProduction/RemoteControl/Source/RemoteControlProtocol/Private/Signature/RCSignatureProtocolAction.cpp
+++ b/Engine/Plugins/VirtualProduction/RemoteControl/Source/RemoteControlProtocol/Private/Signature/RCSignatureProtocolAction.cpp
@@ -203,7 +203,7 @@ namespace UE::RemoteControlProtocol::Private
 				InOutMapping.SetMappingValueAsPrimitive(FRelativeOperation<InCppType>::CalculateAdditive(PropertyValue, MappingValue));
 				break;
 
-			case ERCSignatureProtocolActionMappingSpace::Scalar:
+			case ERCSignatureProtocolActionMappingSpace::Multiply:
 				InOutMapping.SetMappingValueAsPrimitive(FRelativeOperation<InCppType>::CalculateScalar(PropertyValue, MappingValue));
 				break;
 
@@ -268,13 +268,20 @@ bool FRCSignatureProtocolAction::Execute(const FRCSignatureActionContext& InCont
 
 	if (PropertyDimension == 1)
 	{
+		InContext.Property->ProtocolBindings.Empty(1);
 		CreateProtocolEntity(InContext, *Protocol, 0xFF);
 		return true;
 	}
 
-	switch (MaskingType)
+	if (bSingleProtocolChannel)
 	{
-	case ERCSignatureProtocolActionMaskingType::SingleChannel:
+		InContext.Property->ProtocolBindings.Empty(1);
+		CreateProtocolEntity(InContext, *Protocol, OverrideMask);
+	}
+	else
+	{
+		InContext.Property->ProtocolBindings.Empty(PropertyDimension);
+
 		// Add Mask Single Channel per Entity instead of all Channels
 		for (uint8 Dimension = 0; Dimension < PropertyDimension; ++Dimension)
 		{
@@ -284,15 +291,8 @@ bool FRCSignatureProtocolAction::Execute(const FRCSignatureActionContext& InCont
 				CreateProtocolEntity(InContext, *Protocol, Mask);
 			}
 		}
-		break;
-
-	case ERCSignatureProtocolActionMaskingType::MultiChannel:
-		CreateProtocolEntity(InContext, *Protocol, OverrideMask);
-		break;
-
-	default:
-		break;
 	}
+
 	return true;
 }
 
diff --git a/Engine/Plugins/VirtualProduction/RemoteControl/Source/RemoteControlProtocol/Public/Signature/RCSignatureProtocolAction.h b/Engine/Plugins/VirtualProduction/RemoteControl/Source/RemoteControlProtocol/Public/Signature/RCSignatureProtocolAction.h
--- a/Engine/Plugins/VirtualProduction/RemoteControl/Source/RemoteControlProtocol/Public/Signature/RCSignatureProtocolAction.h
+++ b/Engine/Plugins/VirtualProduction/RemoteControl/Source/RemoteControlProtocol/Public/Signature/RCSignatureProtocolAction.h
@@ -11,18 +11,11 @@
 class IRemoteControlProtocol;
 struct FRemoteControlProtocolEntity;
 
-UENUM()
-enum class ERCSignatureProtocolActionMaskingType : uint8
-{
-	SingleChannel,
-	MultiChannel,
-};
-
 UENUM()
 enum class ERCSignatureProtocolActionMappingSpace : uint8
 {
 	Additive,
-	Scalar,
+	Multiply,
 	Absolute
 };
 
@@ -60,10 +53,11 @@ struct FRCSignatureProtocolAction : public FRCSignatureAction
 	TInstancedStruct<FRemoteControlProtocolEntity> ProtocolEntity;
 
 	UPROPERTY(EditAnywhere, Category="Protocol", meta=(EditCondition="PropertyDimension > 1", EditConditionHides))
-	ERCSignatureProtocolActionMaskingType MaskingType = ERCSignatureProtocolActionMaskingType::SingleChannel;
+	uint8 OverrideMask = 0xFF;
 
+	/** Whether to combine all the Masks into a Single Protocol Channel rather than each Mask bit assigned to their own channel */
 	UPROPERTY(EditAnywhere, Category="Protocol", meta=(EditCondition="PropertyDimension > 1", EditConditionHides))
-	uint8 OverrideMask = 0xFF;
+	bool bSingleProtocolChannel = false;
 
 	UPROPERTY(EditAnywhere, Category="Protocol")
 	ERCSignatureProtocolActionMappingSpace MappingSpace = ERCSignatureProtocolActionMappingSpace::Additive;
diff --git a/Engine/Plugins/VirtualProduction/RemoteControl/Source/RemoteControlProtocolWidgets/Private/Customizations/RCSignatureProtocolActionCustomization.cpp b/Engine/Plugins/VirtualProduction/RemoteControl/Source/RemoteControlProtocolWidgets/Private/Customizations/RCSignatureProtocolActionCustomization.cpp
--- a/Engine/Plugins/VirtualProduction/RemoteControl/Source/RemoteControlProtocolWidgets/Private/Customizations/RCSignatureProtocolActionCustomization.cpp
+++ b/Engine/Plugins/VirtualProduction/RemoteControl/Source/RemoteControlProtocolWidgets/Private/Customizations/RCSignatureProtocolActionCustomization.cpp
@@ -60,13 +60,13 @@ void FRCSignatureProtocolActionCustomization::CustomizeDetails(IDetailLayoutBuil
 {
 	const TSharedRef<IPropertyHandle> ProtocolNameHandle = InDetailBuilder.GetProperty(GET_MEMBER_NAME_CHECKED(FRCSignatureProtocolAction, ProtocolName));
 	const TSharedRef<IPropertyHandle> ProtocolEntityHandle = InDetailBuilder.GetProperty(GET_MEMBER_NAME_CHECKED(FRCSignatureProtocolAction, ProtocolEntity));
-	const TSharedRef<IPropertyHandle> MaskingTypeHandle = InDetailBuilder.GetProperty(GET_MEMBER_NAME_CHECKED(FRCSignatureProtocolAction, MaskingType));
 	const TSharedRef<IPropertyHandle> OverrideMaskHandle = InDetailBuilder.GetProperty(GET_MEMBER_NAME_CHECKED(FRCSignatureProtocolAction, OverrideMask));
+	const TSharedRef<IPropertyHandle> SingleProtocolChannelHandle = InDetailBuilder.GetProperty(GET_MEMBER_NAME_CHECKED(FRCSignatureProtocolAction, bSingleProtocolChannel));
 
 	InDetailBuilder.HideProperty(ProtocolNameHandle);
 	InDetailBuilder.HideProperty(ProtocolEntityHandle);
-	InDetailBuilder.HideProperty(MaskingTypeHandle);
 	InDetailBuilder.HideProperty(OverrideMaskHandle);
+	InDetailBuilder.HideProperty(SingleProtocolChannelHandle);
 
 	IDetailCategoryBuilder& ProtocolCategory = InDetailBuilder.EditCategory(TEXT("Protocol"));
 
@@ -82,8 +82,6 @@ void FRCSignatureProtocolActionCustomization::CustomizeDetails(IDetailLayoutBuil
 		];
 
 	ProtocolCategory.AddProperty(ProtocolEntityHandle);
-	ProtocolCategory.AddProperty(MaskingTypeHandle);
-
 	ProtocolCategory.AddProperty(OverrideMaskHandle)
 		.CustomWidget()
 		.NameContent()
@@ -94,6 +92,7 @@ void FRCSignatureProtocolActionCustomization::CustomizeDetails(IDetailLayoutBuil
 		[
 			SNew(SRCSignatureProtocolMaskProperty, OverrideMaskHandle, GetProtocolActions(InDetailBuilder))
 		];
+	ProtocolCategory.AddProperty(SingleProtocolChannelHandle);
 }
 
 TArray<const FRCSignatureProtocolAction*, TInlineAllocator<1>> FRCSignatureProtocolActionCustomization::GetProtocolActions(IDetailLayoutBuilder& InDetailBuilder) const
