[WebRemoteControl] Fix bool property updates not being sent when a single bool property has changed

diff --git a/Engine/Plugins/VirtualProduction/RemoteControl/Source/WebRemoteControl/Private/WebSocketMessageHandler.cpp b/Engine/Plugins/VirtualProduction/RemoteControl/Source/WebRemoteControl/Private/WebSocketMessageHandler.cpp
--- a/Engine/Plugins/VirtualProduction/RemoteControl/Source/WebRemoteControl/Private/WebSocketMessageHandler.cpp
+++ b/Engine/Plugins/VirtualProduction/RemoteControl/Source/WebRemoteControl/Private/WebSocketMessageHandler.cpp
@@ -1030,7 +1030,7 @@ void FWebSocketMessageHandler::ProcessChangedProperties()
 				//Check if multiple booleans properties want to be sent and send them since multiple booleans have problem with the common workflow.
 				if (ClassToEventsPair.Key == FBoolProperty::StaticClass())
 				{
-					TrySendMultipleBoolProperties(Preset, ClientToEventsPair.Key, ClassToEventsPair.Value, SequenceNumber);
+					SendBoolProperties(Preset, ClientToEventsPair.Key, ClassToEventsPair.Value, SequenceNumber);
 					continue;
 				}
 
@@ -1870,12 +1870,10 @@ bool FWebSocketMessageHandler::WritePropertyChangeEventPayload(URemoteControlPre
 	return bHasProperty;
 }
 
-bool FWebSocketMessageHandler::TrySendMultipleBoolProperties(URemoteControlPreset* InPreset,
+void FWebSocketMessageHandler::SendBoolProperties(URemoteControlPreset* InPreset,
 	const FGuid& InTargetClientId, const TSet<FGuid>& InModifiedPropertyIds, int64 InSequenceNumber)
 {
-	bool bFound = false;
-	int32 NumberSent = 0;
-	if(InModifiedPropertyIds.Num() > 1)
+	if (!InModifiedPropertyIds.IsEmpty())
 	{
 		if (TSharedPtr<FRemoteControlProperty> RCProperty = InPreset->GetExposedEntity<FRemoteControlProperty>(InModifiedPropertyIds.Array()[0]).Pin())
 		{
@@ -1883,7 +1881,6 @@ bool FWebSocketMessageHandler::TrySendMultipleBoolProperties(URemoteControlPrese
 			{
 				if (Property->IsA<FBoolProperty>())
 				{
-					bFound = true;
 					for (FGuid ModifiedPropertyId : InModifiedPropertyIds)
 					{
 						TArray<uint8> BoolsWorkingBuffer;
@@ -1892,14 +1889,12 @@ bool FWebSocketMessageHandler::TrySendMultipleBoolProperties(URemoteControlPrese
 							TArray<uint8> Payload;
 							WebRemoteControlUtils::ConvertToUTF8(BoolsWorkingBuffer, Payload);
 							Server->Send(InTargetClientId, Payload);
-							++NumberSent;
 						}
 					}
 				}
 			}
 		}
 	}
-	return bFound && (NumberSent == InModifiedPropertyIds.Num());
 }
 
 
diff --git a/Engine/Plugins/VirtualProduction/RemoteControl/Source/WebRemoteControl/Private/WebSocketMessageHandler.h b/Engine/Plugins/VirtualProduction/RemoteControl/Source/WebRemoteControl/Private/WebSocketMessageHandler.h
--- a/Engine/Plugins/VirtualProduction/RemoteControl/Source/WebRemoteControl/Private/WebSocketMessageHandler.h
+++ b/Engine/Plugins/VirtualProduction/RemoteControl/Source/WebRemoteControl/Private/WebSocketMessageHandler.h
@@ -212,14 +212,13 @@ private:
 	bool WritePropertyChangeEventPayload(URemoteControlPreset* InPreset, const TSet<FGuid>& InModifiedPropertyIds, int64 InSequenceNumber, TArray<uint8>& OutBuffer);
 
 	/**
-	 * Check if multiple booleans properties want to be sent over to a client, and send a message for each of those properties. This is needed because multiple booleans have problems with the common workflow.
+	 * Send a message for each boolean property. This is needed because multiple booleans have problems with the common workflow.
 	 * @param InPreset The current preset
 	 * @param InTargetClientId Client Id to send the message
 	 * @param InModifiedPropertyIds Set of modified properties already divided by type
 	 * @param InSequenceNumber SequenceNumber of the client
-	 * @return True if multiple booleans property are found and sent, false if they are not found or if not all of them are sent.
 	 */
-	bool TrySendMultipleBoolProperties(URemoteControlPreset* InPreset, const FGuid& InTargetClientId, const TSet<FGuid>& InModifiedPropertyIds, int64 InSequenceNumber);
+	void SendBoolProperties(URemoteControlPreset* InPreset, const FGuid& InTargetClientId, const TSet<FGuid>& InModifiedPropertyIds, int64 InSequenceNumber);
 	
 	/**
 	 * Write the provided list of controller events to a buffer.
