Remote Control Signature: Made Min/Max mapping to retain all the meta data except displayname and tooltip.

diff --git a/Engine/Plugins/VirtualProduction/RemoteControl/Source/RemoteControlProtocol/Private/Signature/RCSignatureProtocolAction.cpp b/Engine/Plugins/VirtualProduction/RemoteControl/Source/RemoteControlProtocol/Private/Signature/RCSignatureProtocolAction.cpp
--- a/Engine/Plugins/VirtualProduction/RemoteControl/Source/RemoteControlProtocol/Private/Signature/RCSignatureProtocolAction.cpp
+++ b/Engine/Plugins/VirtualProduction/RemoteControl/Source/RemoteControlProtocol/Private/Signature/RCSignatureProtocolAction.cpp
@@ -228,13 +228,22 @@ void FRCSignatureProtocolAction::Initialize(const FRCSignatureField& InField)
 	MaxMappingDesc.Name = TEXT("Max");
 
 #if WITH_EDITORONLY_DATA
-	// Reset MetaData that could affect Display Name, or Min/Max clamping.
-	// The mapping descs are meant to just match the type
-	MinMappingDesc.MetaData.Reset();
-	MinMappingDesc.MetaClass = nullptr;
+	// Remove Meta-data like DisplayName, Tooltip
+	// Other meta-data like Min/Max clamping, whether to hide Alpha Channel should be retained
+	static const TSet<FName, DefaultKeyFuncs<FName>, TFixedSetAllocator<2>> MetaDataToRemove =
+	{
+		TEXT("DisplayName"),
+		TEXT("Tooltip"),
+	};
+
+	MinMappingDesc.MetaData.RemoveAll(
+		[](const FPropertyBagPropertyDescMetaData& InMetaData)
+		{
+			return MetaDataToRemove.Contains(InMetaData.Key);
+		});
 
-	MaxMappingDesc.MetaData.Reset();
-	MaxMappingDesc.MetaClass = nullptr;
+	// Both Min/Max refer to the same property so should have the same meta-data
+	MaxMappingDesc.MetaData = MinMappingDesc.MetaData;
 #endif
 
 	PropertyDimension = UE::RemoteControlProtocol::Private::GetPropertyDimension(GetPropertyStruct());
diff --git a/Engine/Plugins/VirtualProduction/RemoteControl/Source/RemoteControlUI/Private/UI/Signature/SRCSignatureTree.cpp b/Engine/Plugins/VirtualProduction/RemoteControl/Source/RemoteControlUI/Private/UI/Signature/SRCSignatureTree.cpp
--- a/Engine/Plugins/VirtualProduction/RemoteControl/Source/RemoteControlUI/Private/UI/Signature/SRCSignatureTree.cpp
+++ b/Engine/Plugins/VirtualProduction/RemoteControl/Source/RemoteControlUI/Private/UI/Signature/SRCSignatureTree.cpp
@@ -207,6 +207,8 @@ void SRCSignatureTree::DeleteSelectedPanelItems()
 
 	FScopedTransaction Transaction(LOCTEXT("RemoveSelectedItems", "Remove Selected Items"));
 	DeleteItemsFromLogicPanel<FRCSignatureTreeItemBase>(RootItem->GetChildrenMutable(), SelectedItems);
+
+	RootItem->GetSelection()->ClearSelection();
 }
 
 void SRCSignatureTree::Reset()
