Remote Control: Fix an issue with a static cast done without properly checking for type

diff --git a/Engine/Plugins/VirtualProduction/RemoteControl/Source/RemoteControlUI/Private/UI/Filters/RCFilter.cpp b/Engine/Plugins/VirtualProduction/RemoteControl/Source/RemoteControlUI/Private/UI/Filters/RCFilter.cpp
--- a/Engine/Plugins/VirtualProduction/RemoteControl/Source/RemoteControlUI/Private/UI/Filters/RCFilter.cpp
+++ b/Engine/Plugins/VirtualProduction/RemoteControl/Source/RemoteControlUI/Private/UI/Filters/RCFilter.cpp
@@ -25,20 +25,24 @@ void FRCFilter::AddTypeFilter(UClass* InAssetTypeFilter)
 
 bool FRCFilter::DoesPassFilters(FEntityFilterType InEntityItem)
 {
+	const SRCPanelTreeNode::ENodeType EntityRCType = InEntityItem->GetRCType();
+
 	// Do not process Field Child.
-	if (InEntityItem->GetRCType() == SRCPanelTreeNode::FieldChild)
+	if (EntityRCType == SRCPanelTreeNode::FieldChild)
 	{
 		return false;
 	}
 
 	// Return early if the RC type of entity is Actor.
-	if (InEntityItem->GetRCType() == SRCPanelTreeNode::Actor)
+	if (EntityRCType == SRCPanelTreeNode::Actor)
 	{
 		return CustomTypeFilters.Contains(RemoteControlTypes::NAME_RCActors);
 	}
 
-	if (TSharedPtr<const SRCPanelExposedField> ExposedField = StaticCastSharedRef<const SRCPanelExposedField>(InEntityItem))
+	if (EntityRCType == SRCPanelTreeNode::Field)
 	{
+		const TSharedRef<const SRCPanelExposedField> ExposedField = StaticCastSharedRef<const SRCPanelExposedField>(InEntityItem);
+
 		// Return early if the field type of entity is Function.
 		if (ExposedField->GetFieldType() == EExposedFieldType::Function)	
 		{
