Remote Control: Use an ordered type to store entities, avoid per session order in UI and code

diff --git a/Engine/Plugins/VirtualProduction/RemoteControl/Source/RemoteControl/Private/RemoteControlExposeRegistry.cpp b/Engine/Plugins/VirtualProduction/RemoteControl/Source/RemoteControl/Private/RemoteControlExposeRegistry.cpp
--- a/Engine/Plugins/VirtualProduction/RemoteControl/Source/RemoteControl/Private/RemoteControlExposeRegistry.cpp
+++ b/Engine/Plugins/VirtualProduction/RemoteControl/Source/RemoteControl/Private/RemoteControlExposeRegistry.cpp
@@ -1,17 +1,22 @@
 // Copyright Epic Games, Inc. All Rights Reserved.
 
 #include "RemoteControlExposeRegistry.h"
+
+#include "Algo/Find.h"
+#include "Algo/RemoveIf.h"
+#include "Algo/Transform.h"
+#include "HAL/UnrealMemory.h"
 #include "Misc/Guid.h"
+#include "RemoteControlObjectVersion.h"
 #include "Serialization/Archive.h"
 #include "UObject/Class.h"
 #include "UObject/SoftObjectPath.h"
-#include "HAL/UnrealMemory.h"
 
 TArray<TSharedPtr<const FRemoteControlEntity>> URemoteControlExposeRegistry::GetExposedEntities() const
 {
 	TArray<TSharedPtr<const FRemoteControlEntity>> Entities;
 
-	Algo::TransformIf(ExposedEntities, Entities,
+	Algo::TransformIf(ExposedEntitiesArray, Entities,
 		[](const FRCEntityWrapper& Wrapper)
 		{
 			return Wrapper.IsValid();
@@ -29,7 +34,7 @@ TArray<TSharedPtr<const FRemoteControlEntity>> URemoteControlExposeRegistry::Get
 	check(EntityType->IsChildOf(FRemoteControlEntity::StaticStruct()));
 	TArray<TSharedPtr<const FRemoteControlEntity>> Entities;
 
-	for (const FRCEntityWrapper& Wrapper : ExposedEntities)
+	for (const FRCEntityWrapper& Wrapper : ExposedEntitiesArray)
 	{
 		if (Wrapper.IsValid() && Wrapper.GetType()->IsChildOf(EntityType))
 		{
@@ -45,7 +50,7 @@ TArray<TSharedPtr<FRemoteControlEntity>> URemoteControlExposeRegistry::GetExpose
 	check(EntityType->IsChildOf(FRemoteControlEntity::StaticStruct()));
 	TArray<TSharedPtr<FRemoteControlEntity>> Entities;
 
-	for (FRCEntityWrapper& Wrapper : ExposedEntities)
+	for (FRCEntityWrapper& Wrapper : ExposedEntitiesArray)
 	{
 		if (Wrapper.IsValid() && Wrapper.GetType()->IsChildOf(EntityType))
 		{
@@ -63,7 +68,7 @@ TSharedPtr<const FRemoteControlEntity> URemoteControlExposeRegistry::GetExposedE
 
 TSharedPtr<FRemoteControlEntity> URemoteControlExposeRegistry::GetExposedEntity(const FGuid& ExposedEntityId, const UScriptStruct* EntityType)
 {
-	if (FRCEntityWrapper* Wrapper = ExposedEntities.FindByHash(GetTypeHash(ExposedEntityId), ExposedEntityId))
+	if (FRCEntityWrapper* Wrapper = Algo::Find(ExposedEntitiesArray, ExposedEntityId))
 	{
 		if (Wrapper->IsValid() && Wrapper->GetType()->IsChildOf(EntityType))
 		{
@@ -76,7 +81,7 @@ TSharedPtr<FRemoteControlEntity> URemoteControlExposeRegistry::GetExposedEntity(
 
 const UScriptStruct* URemoteControlExposeRegistry::GetExposedEntityType(const FGuid& ExposedEntityId) const
 {
-	if (const FRCEntityWrapper* Wrapper = ExposedEntities.FindByHash(GetTypeHash(ExposedEntityId), ExposedEntityId))
+	if (const FRCEntityWrapper* Wrapper = Algo::Find(ExposedEntitiesArray, ExposedEntityId))
 	{
 		return Wrapper->GetType();
 	}
@@ -86,7 +91,7 @@ const UScriptStruct* URemoteControlExposeRegistry::GetExposedEntityType(const FG
 
 const bool URemoteControlExposeRegistry::IsEmpty() const
 {
-    return ExposedEntities.IsEmpty();
+    return ExposedEntitiesArray.IsEmpty();
 }
 
 TSharedPtr<FRemoteControlEntity> URemoteControlExposeRegistry::AddExposedEntity(FRemoteControlEntity&& EntityToExpose, UScriptStruct* EntityType)
@@ -94,7 +99,7 @@ TSharedPtr<FRemoteControlEntity> URemoteControlExposeRegistry::AddExposedEntity(
 	LabelToIdCache.Add(EntityToExpose.GetLabel(), EntityToExpose.GetId());
 	FRCEntityWrapper Wrapper{ MoveTemp(EntityToExpose), EntityType};
 	TSharedPtr<FRemoteControlEntity> Entity = Wrapper.Get();
-	ExposedEntities.Add(MoveTemp(Wrapper));
+	ExposedEntitiesArray.Add(MoveTemp(Wrapper));
 	ExposedTypes.Add(EntityType);
 	return Entity;
 }
@@ -102,13 +107,17 @@ TSharedPtr<FRemoteControlEntity> URemoteControlExposeRegistry::AddExposedEntity(
 void URemoteControlExposeRegistry::RemoveExposedEntity(const FGuid& Id)
 {
 	uint32 Hash = GetTypeHash(Id);
-	if (const FRCEntityWrapper* Wrapper = ExposedEntities.FindByHash(Hash, Id))
+	if (const FRCEntityWrapper* Wrapper = Algo::Find(ExposedEntitiesArray, Id))
 	{
 		if (TSharedPtr<const FRemoteControlEntity> Entity = Wrapper->Get())
 		{
 			LabelToIdCache.Remove(Entity->GetLabel());
 		}
-		ExposedEntities.RemoveByHash(Hash, Id);
+		ExposedEntitiesArray.SetNum(Algo::StableRemoveIf(ExposedEntitiesArray,
+			[&Id](const FRCEntityWrapper& EntityWrapper)
+			{
+				return EntityWrapper.Get()->GetId() == Id;
+			}));
 	}
 }
 
@@ -156,6 +165,29 @@ FName URemoteControlExposeRegistry::GenerateUniqueLabel(FName BaseName) const
 	return NAME_None;
 }
 
+void URemoteControlExposeRegistry::Serialize(FArchive& Ar)
+{
+	Super::Serialize(Ar);
+
+	Ar.UsingCustomVersion(FRemoteControlObjectVersion::GUID);
+
+#if WITH_EDITORONLY_DATA
+	if (Ar.IsLoading())
+	{
+		if (Ar.CustomVer(FRemoteControlObjectVersion::GUID) < FRemoteControlObjectVersion::ChangedExposedEntitiesSetToAnOrderedTArray)
+		{
+PRAGMA_DISABLE_DEPRECATION_WARNINGS
+			Algo::Transform(ExposedEntities_DEPRECATED, ExposedEntitiesArray,
+			[](const FRCEntityWrapper& EntityWrapper)
+			{
+				return EntityWrapper;
+			});
+PRAGMA_ENABLE_DEPRECATION_WARNINGS
+		}
+	}
+#endif // WITH_EDITORONLY_DATA
+}
+
 void URemoteControlExposeRegistry::PostLoad()
 {
 	Super::PostLoad();
@@ -168,25 +200,10 @@ void URemoteControlExposeRegistry::PostDuplicate(bool bDuplicateForPIE)
 	CacheLabels();
 }
 
-void URemoteControlExposeRegistry::Rehash()
-{
-	TSet<FRCEntityWrapper> RehashedEntities;
-	RehashedEntities.Reserve(ExposedEntities.Num());
-	
-	for (FRCEntityWrapper& Wrapper : ExposedEntities)
-	{
-		RehashedEntities.Add(FRCEntityWrapper(Wrapper));
-	}
-	
-	ExposedEntities = MoveTemp(RehashedEntities);
-
-	CacheLabels();
-}
-
 TSharedPtr<FRemoteControlEntity> URemoteControlExposeRegistry::GetEntity(const FGuid& EntityId)
 {
 	/** Get a raw pointer to an entity using its id. */
-	if (FRCEntityWrapper* Wrapper = ExposedEntities.FindByHash(GetTypeHash(EntityId), EntityId))
+	if (FRCEntityWrapper* Wrapper = Algo::Find(ExposedEntitiesArray, EntityId))
 	{
 		return Wrapper->Get();
 	}
@@ -202,7 +219,7 @@ TSharedPtr<const FRemoteControlEntity> URemoteControlExposeRegistry::GetEntity(c
 void URemoteControlExposeRegistry::CacheLabels()
 {
 	LabelToIdCache.Reset();
-	for (const FRCEntityWrapper& Wrapper : ExposedEntities)
+	for (const FRCEntityWrapper& Wrapper : ExposedEntitiesArray)
 	{
 		if (TSharedPtr<const FRemoteControlEntity> Entity = Wrapper.Get())
 		{
diff --git a/Engine/Plugins/VirtualProduction/RemoteControl/Source/RemoteControl/Private/RemoteControlExposeRegistry.h b/Engine/Plugins/VirtualProduction/RemoteControl/Source/RemoteControl/Private/RemoteControlExposeRegistry.h
--- a/Engine/Plugins/VirtualProduction/RemoteControl/Source/RemoteControl/Private/RemoteControlExposeRegistry.h
+++ b/Engine/Plugins/VirtualProduction/RemoteControl/Source/RemoteControl/Private/RemoteControlExposeRegistry.h
@@ -195,15 +195,15 @@ public:
 	 */
 	FName GenerateUniqueLabel(FName BaseName) const;
 
+	//~ Begin UObject interface
+	virtual void Serialize(FArchive& Ar) override;
 	virtual void PostLoad() override;
-
 	virtual void PostDuplicate(bool bDuplicateForPIE) override;
+	//~ End UObject interface
+
+	/** Cache labels for all exposed entities. */
+	void CacheLabels();
 
-	/**
-	 * @brief Rehash the entities after ids have been modified.
-	 */
-	void Rehash();
-	
 private:
 	/** Get a raw pointer to an entity using its id. */
 	TSharedPtr<FRemoteControlEntity> GetEntity(const FGuid& EntityId);
@@ -211,13 +211,17 @@ private:
 	/** Get a raw pointer to an entity using its id. */
 	TSharedPtr<const FRemoteControlEntity> GetEntity(const FGuid& EntityId) const;
 
-	/** Cache labels for all expose sets. */
-	void CacheLabels();
-
 private:
 	/** Holds the exposed entities. */
 	UPROPERTY()
-	TSet<FRCEntityWrapper> ExposedEntities;
+	TArray<FRCEntityWrapper> ExposedEntitiesArray;
+
+#if WITH_EDITORONLY_DATA
+	/** DEPRECATED 5.5 Holds the exposed entities. Instead refer to the new TArray<FRCEntityWrapper> URemoteControlExposeRegistry::ExposedEntities */
+	UE_DEPRECATED(5.5, "Changed to an ordered TArray.")
+	UPROPERTY()
+	TSet<FRCEntityWrapper> ExposedEntities_DEPRECATED;
+#endif //
 
 	/** Cache of label to ids. */
 	UPROPERTY(Transient)
diff --git a/Engine/Plugins/VirtualProduction/RemoteControl/Source/RemoteControl/Private/RemoteControlPreset.cpp b/Engine/Plugins/VirtualProduction/RemoteControl/Source/RemoteControl/Private/RemoteControlPreset.cpp
--- a/Engine/Plugins/VirtualProduction/RemoteControl/Source/RemoteControl/Private/RemoteControlPreset.cpp
+++ b/Engine/Plugins/VirtualProduction/RemoteControl/Source/RemoteControl/Private/RemoteControlPreset.cpp
@@ -1834,8 +1834,10 @@ void URemoteControlPreset::RenewEntityIds()
 		}
 	}
 
+	// Cache labels for all exposed entities
+	Registry->CacheLabels();
+
 	// Rehash the registries
-	Registry->Rehash();
 	PropertyIdRegistry->UpdateEntityIds(EntityIdMap);
 	Layout.UpdateEntityIds(EntityIdMap);
 	
diff --git a/Engine/Plugins/VirtualProduction/RemoteControl/Source/RemoteControl/Public/RemoteControlObjectVersion.h b/Engine/Plugins/VirtualProduction/RemoteControl/Source/RemoteControl/Public/RemoteControlObjectVersion.h
--- a/Engine/Plugins/VirtualProduction/RemoteControl/Source/RemoteControl/Public/RemoteControlObjectVersion.h
+++ b/Engine/Plugins/VirtualProduction/RemoteControl/Source/RemoteControl/Public/RemoteControlObjectVersion.h
@@ -30,6 +30,9 @@ struct FRemoteControlObjectVersion
 		// Instead of the global 'Protocol Generates Transactions' setting, each preset can now specify its modify operation flags
 		AddedPerPresetModifyOperationFlags,
 
+		// Change FRCEntityWrapper::ExposedEntities' Set to an Array to retain ordering across sessions
+		ChangedExposedEntitiesSetToAnOrderedTArray,
+
 		// -----<new versions can be added above this line>-------------------------------------------------
 		VersionPlusOne,
 		LatestVersion = VersionPlusOne - 1
