[Remote Control] Removing superfluous transaction.

#rnx

diff --git a/Engine/Plugins/VirtualProduction/RemoteControl/Source/RemoteControl/Private/RCVirtualPropertyContainer.cpp b/Engine/Plugins/VirtualProduction/RemoteControl/Source/RemoteControl/Private/RCVirtualPropertyContainer.cpp
--- a/Engine/Plugins/VirtualProduction/RemoteControl/Source/RemoteControl/Private/RCVirtualPropertyContainer.cpp
+++ b/Engine/Plugins/VirtualProduction/RemoteControl/Source/RemoteControl/Private/RCVirtualPropertyContainer.cpp
@@ -412,10 +412,6 @@ void URCVirtualPropertyContainerBase::PostEditUndo()
 
 void URCVirtualPropertyContainerBase::OnModifyPropertyValue(const FPropertyChangedEvent& PropertyChangedEvent)
 {
-	const FScopedTransaction Transaction(LOCTEXT("OnModifyPropertyValue", "On Modify Property Value"));
-
-	Modify();
-
 	MarkPackageDirty();
 }
 #endif
