RemoteControl: Fix ensure condition preventing from packaging with the RemoteControl plugin enabled

#rnx

diff --git a/Engine/Plugins/VirtualProduction/RemoteControl/Source/RemoteControl/Private/RemoteControlExposeRegistry.cpp b/Engine/Plugins/VirtualProduction/RemoteControl/Source/RemoteControl/Private/RemoteControlExposeRegistry.cpp
--- a/Engine/Plugins/VirtualProduction/RemoteControl/Source/RemoteControl/Private/RemoteControlExposeRegistry.cpp
+++ b/Engine/Plugins/VirtualProduction/RemoteControl/Source/RemoteControl/Private/RemoteControlExposeRegistry.cpp
@@ -260,7 +260,7 @@ bool FRCEntityWrapper::Serialize(FArchive& Ar)
 {
 	if (Ar.IsSaving())
 	{
-		if (ensure(IsValid()))
+		if (IsValid())
 		{
 			FSoftObjectPath Path{ EntityType };
 			Path.Serialize(Ar);
