Remote Control: Fix issues when controlling struct properties with Remote Control Protocols, fix issues specific to controlling struct properties that support masking such as FVector or FColor

diff --git a/Engine/Plugins/VirtualProduction/RemoteControl/Source/RemoteControl/Private/Factories/RemoteControlMaskingFactories.cpp b/Engine/Plugins/VirtualProduction/RemoteControl/Source/RemoteControl/Private/Factories/RemoteControlMaskingFactories.cpp
--- a/Engine/Plugins/VirtualProduction/RemoteControl/Source/RemoteControl/Private/Factories/RemoteControlMaskingFactories.cpp
+++ b/Engine/Plugins/VirtualProduction/RemoteControl/Source/RemoteControl/Private/Factories/RemoteControlMaskingFactories.cpp
@@ -2,6 +2,9 @@
 
 #include "RemoteControlMaskingFactories.h"
 
+// Entirely deprecated 5.5
+PRAGMA_DISABLE_DEPRECATION_WARNINGS
+
 /**
  * FVectorMaskingFactory
  */
@@ -10,18 +13,14 @@ TSharedRef<IRemoteControlMaskingFactory> FVectorMaskingFactory::MakeInstance()
 	return MakeShared<FVectorMaskingFactory>();
 }
 
-void FVectorMaskingFactory::ApplyMaskedValues(const TSharedRef<FRCMaskingOperation>& InMaskingOperation, bool bIsInteractive, const ERCModifyOperationFlags ModifyOperationFlags)
+void FVectorMaskingFactory::ApplyMaskedValues(const TSharedRef<FRCMaskingOperation>& InMaskingOperation, bool bIsInteractive)
 {
 	if (FStructProperty* ToStructProp = CastField<FStructProperty>(InMaskingOperation->ObjectRef.Property.Get()))
 	{
 		if (UObject* OwningObject = InMaskingOperation->ObjectRef.Object.Get())
 		{
 #if WITH_EDITOR
-			const bool bWithPropertyChangedEvents = !EnumHasAnyFlags(ModifyOperationFlags, ERCModifyOperationFlags::SkipPropertyChangeEvents);
-			if (bWithPropertyChangedEvents)
-			{
-				OwningObject->PreEditChange(ToStructProp);
-			}
+			OwningObject->PreEditChange(ToStructProp);
 #endif
 
 			if (const FVector* VectorProp = ToStructProp->ContainerPtrToValuePtr<FVector>(OwningObject))
@@ -36,11 +35,8 @@ void FVectorMaskingFactory::ApplyMaskedValues(const TSharedRef<FRCMaskingOperati
 			}
 
 #if WITH_EDITOR
-			if (bWithPropertyChangedEvents)
-			{
-				FPropertyChangedEvent ChangeEvent(ToStructProp, bIsInteractive ? EPropertyChangeType::Interactive : EPropertyChangeType::ValueSet);
-				OwningObject->PostEditChangeProperty(ChangeEvent);
-			}
+			FPropertyChangedEvent ChangeEvent(ToStructProp, bIsInteractive ? EPropertyChangeType::Interactive : EPropertyChangeType::ValueSet);
+			OwningObject->PostEditChangeProperty(ChangeEvent);
 #endif
 		}
 	}
@@ -75,18 +71,14 @@ TSharedRef<IRemoteControlMaskingFactory> FVector4MaskingFactory::MakeInstance()
 	return MakeShared<FVector4MaskingFactory>();
 }
 
-void FVector4MaskingFactory::ApplyMaskedValues(const TSharedRef<FRCMaskingOperation>& InMaskingOperation, bool bIsInteractive, const ERCModifyOperationFlags ModifyOperationFlags)
+void FVector4MaskingFactory::ApplyMaskedValues(const TSharedRef<FRCMaskingOperation>& InMaskingOperation, bool bIsInteractive)
 {
 	if (FStructProperty* ToStructProp = CastField<FStructProperty>(InMaskingOperation->ObjectRef.Property.Get()))
 	{
 		if (UObject* OwningObject = InMaskingOperation->ObjectRef.Object.Get())
 		{
 #if WITH_EDITOR
-			const bool bWithPropertyChangedEvents = !EnumHasAnyFlags(ModifyOperationFlags, ERCModifyOperationFlags::SkipPropertyChangeEvents);
-			if (bWithPropertyChangedEvents)
-			{
-				OwningObject->PreEditChange(ToStructProp);
-			}
+			OwningObject->PreEditChange(ToStructProp);
 #endif // WITH_EDITOR
 
 			if (const FVector4* Vector4Prop = ToStructProp->ContainerPtrToValuePtr<FVector4>(OwningObject))
@@ -102,11 +94,8 @@ void FVector4MaskingFactory::ApplyMaskedValues(const TSharedRef<FRCMaskingOperat
 			}
 
 #if WITH_EDITOR
-			if (bWithPropertyChangedEvents)
-			{
-				FPropertyChangedEvent ChangeEvent(ToStructProp, bIsInteractive ? EPropertyChangeType::Interactive : EPropertyChangeType::ValueSet);
-				OwningObject->PostEditChangeProperty(ChangeEvent);
-			}
+			FPropertyChangedEvent ChangeEvent(ToStructProp, bIsInteractive ? EPropertyChangeType::Interactive : EPropertyChangeType::ValueSet);
+			OwningObject->PostEditChangeProperty(ChangeEvent);
 #endif // WITH_EDITOR
 		}
 	}
@@ -142,18 +131,14 @@ TSharedRef<IRemoteControlMaskingFactory> FIntVectorMaskingFactory::MakeInstance(
 	return MakeShared<FIntVectorMaskingFactory>();
 }
 
-void FIntVectorMaskingFactory::ApplyMaskedValues(const TSharedRef<FRCMaskingOperation>& InMaskingOperation, bool bIsInteractive, const ERCModifyOperationFlags ModifyOperationFlags)
+void FIntVectorMaskingFactory::ApplyMaskedValues(const TSharedRef<FRCMaskingOperation>& InMaskingOperation, bool bIsInteractive)
 {
 	if (FStructProperty* ToStructProp = CastField<FStructProperty>(InMaskingOperation->ObjectRef.Property.Get()))
 	{
 		if (UObject* OwningObject = InMaskingOperation->ObjectRef.Object.Get())
 		{
 #if WITH_EDITOR
-			const bool bWithPropertyChangedEvents = !EnumHasAnyFlags(ModifyOperationFlags, ERCModifyOperationFlags::SkipPropertyChangeEvents);
-			if (bWithPropertyChangedEvents)
-			{
-				OwningObject->PreEditChange(ToStructProp);
-			}
+			OwningObject->PreEditChange(ToStructProp);
 #endif // WITH_EDITOR
 
 			if (const FIntVector* IntVectorProp = ToStructProp->ContainerPtrToValuePtr<FIntVector>(OwningObject))
@@ -168,11 +153,8 @@ void FIntVectorMaskingFactory::ApplyMaskedValues(const TSharedRef<FRCMaskingOper
 			}
 
 #if WITH_EDITOR
-			if (bWithPropertyChangedEvents)
-			{
-				FPropertyChangedEvent ChangeEvent(ToStructProp, bIsInteractive ? EPropertyChangeType::Interactive : EPropertyChangeType::ValueSet);
-				OwningObject->PostEditChangeProperty(ChangeEvent);
-			}
+			FPropertyChangedEvent ChangeEvent(ToStructProp, bIsInteractive ? EPropertyChangeType::Interactive : EPropertyChangeType::ValueSet);
+			OwningObject->PostEditChangeProperty(ChangeEvent);
 #endif // WITH_EDITOR
 		}
 	}
@@ -207,18 +189,14 @@ TSharedRef<IRemoteControlMaskingFactory> FIntVector4MaskingFactory::MakeInstance
 	return MakeShared<FIntVector4MaskingFactory>();
 }
 
-void FIntVector4MaskingFactory::ApplyMaskedValues(const TSharedRef<FRCMaskingOperation>& InMaskingOperation, bool bIsInteractive, const ERCModifyOperationFlags ModifyOperationFlags)
+void FIntVector4MaskingFactory::ApplyMaskedValues(const TSharedRef<FRCMaskingOperation>& InMaskingOperation, bool bIsInteractive)
 {
 	if (FStructProperty* ToStructProp = CastField<FStructProperty>(InMaskingOperation->ObjectRef.Property.Get()))
 	{
 		if (UObject* OwningObject = InMaskingOperation->ObjectRef.Object.Get())
 		{
 #if WITH_EDITOR
-			const bool bWithPropertyChangedEvents = !EnumHasAnyFlags(ModifyOperationFlags, ERCModifyOperationFlags::SkipPropertyChangeEvents);
-			if (bWithPropertyChangedEvents)
-			{
-				OwningObject->PreEditChange(ToStructProp);
-			}
+			OwningObject->PreEditChange(ToStructProp);
 #endif // WITH_EDITOR
 
 			if (const FIntVector4* IntVector4Prop = ToStructProp->ContainerPtrToValuePtr<FIntVector4>(OwningObject))
@@ -234,11 +212,8 @@ void FIntVector4MaskingFactory::ApplyMaskedValues(const TSharedRef<FRCMaskingOpe
 			}
 
 #if WITH_EDITOR
-			if (bWithPropertyChangedEvents)
-			{
-				FPropertyChangedEvent ChangeEvent(ToStructProp, bIsInteractive ? EPropertyChangeType::Interactive : EPropertyChangeType::ValueSet);
-				OwningObject->PostEditChangeProperty(ChangeEvent);
-			}
+			FPropertyChangedEvent ChangeEvent(ToStructProp, bIsInteractive ? EPropertyChangeType::Interactive : EPropertyChangeType::ValueSet);
+			OwningObject->PostEditChangeProperty(ChangeEvent);
 #endif // WITH_EDITOR
 		}
 	}
@@ -274,18 +249,14 @@ TSharedRef<IRemoteControlMaskingFactory> FRotatorMaskingFactory::MakeInstance()
 	return MakeShared<FRotatorMaskingFactory>();
 }
 
-void FRotatorMaskingFactory::ApplyMaskedValues(const TSharedRef<FRCMaskingOperation>& InMaskingOperation, bool bIsInteractive, const ERCModifyOperationFlags ModifyOperationFlags)
+void FRotatorMaskingFactory::ApplyMaskedValues(const TSharedRef<FRCMaskingOperation>& InMaskingOperation, bool bIsInteractive)
 {
 	if (FStructProperty* ToStructProp = CastField<FStructProperty>(InMaskingOperation->ObjectRef.Property.Get()))
 	{
 		if (UObject* OwningObject = InMaskingOperation->ObjectRef.Object.Get())
 		{
 #if WITH_EDITOR
-			const bool bWithPropertyChangedEvents = !EnumHasAnyFlags(ModifyOperationFlags, ERCModifyOperationFlags::SkipPropertyChangeEvents);
-			if (bWithPropertyChangedEvents)
-			{
-				OwningObject->PreEditChange(ToStructProp);
-			}
+			OwningObject->PreEditChange(ToStructProp);
 #endif // WITH_EDITOR
 
 			if (const FRotator* RotatorProp = ToStructProp->ContainerPtrToValuePtr<FRotator>(OwningObject))
@@ -300,11 +271,8 @@ void FRotatorMaskingFactory::ApplyMaskedValues(const TSharedRef<FRCMaskingOperat
 			}
 
 #if WITH_EDITOR			
-			if (bWithPropertyChangedEvents)
-			{
-				FPropertyChangedEvent ChangeEvent(ToStructProp, bIsInteractive ? EPropertyChangeType::Interactive : EPropertyChangeType::ValueSet);
-				OwningObject->PostEditChangeProperty(ChangeEvent);
-			}
+			FPropertyChangedEvent ChangeEvent(ToStructProp, bIsInteractive ? EPropertyChangeType::Interactive : EPropertyChangeType::ValueSet);
+			OwningObject->PostEditChangeProperty(ChangeEvent);
 #endif // WITH_EDITOR
 		}
 	}
@@ -339,18 +307,14 @@ TSharedRef<IRemoteControlMaskingFactory> FColorMaskingFactory::MakeInstance()
 	return MakeShared<FColorMaskingFactory>();
 }
 
-void FColorMaskingFactory::ApplyMaskedValues(const TSharedRef<FRCMaskingOperation>& InMaskingOperation, bool bIsInteractive, const ERCModifyOperationFlags ModifyOperationFlags)
+void FColorMaskingFactory::ApplyMaskedValues(const TSharedRef<FRCMaskingOperation>& InMaskingOperation, bool bIsInteractive)
 {
 	if (FStructProperty* ToStructProp = CastField<FStructProperty>(InMaskingOperation->ObjectRef.Property.Get()))
 	{
 		if (UObject* OwningObject = InMaskingOperation->ObjectRef.Object.Get())
 		{
 #if WITH_EDITOR
-			const bool bWithPropertyChangedEvents = !EnumHasAnyFlags(ModifyOperationFlags, ERCModifyOperationFlags::SkipPropertyChangeEvents);
-			if (bWithPropertyChangedEvents)
-			{
-				OwningObject->PreEditChange(ToStructProp);
-			}
+			OwningObject->PreEditChange(ToStructProp);
 #endif // WITH_EDITOR
 
 			if (const FColor* ColorProp = ToStructProp->ContainerPtrToValuePtr<FColor>(OwningObject))
@@ -365,11 +329,8 @@ void FColorMaskingFactory::ApplyMaskedValues(const TSharedRef<FRCMaskingOperatio
 			}
 
 #if WITH_EDITOR
-			if (bWithPropertyChangedEvents)
-			{
-				FPropertyChangedEvent ChangeEvent(ToStructProp, bIsInteractive ? EPropertyChangeType::Interactive : EPropertyChangeType::ValueSet);
-				OwningObject->PostEditChangeProperty(ChangeEvent);
-			}
+			FPropertyChangedEvent ChangeEvent(ToStructProp, bIsInteractive ? EPropertyChangeType::Interactive : EPropertyChangeType::ValueSet);
+			OwningObject->PostEditChangeProperty(ChangeEvent);
 #endif // WITH_EDITOR
 		}
 	}
@@ -404,18 +365,14 @@ TSharedRef<IRemoteControlMaskingFactory> FLinearColorMaskingFactory::MakeInstanc
 	return MakeShared<FLinearColorMaskingFactory>();
 }
 
-void FLinearColorMaskingFactory::ApplyMaskedValues(const TSharedRef<FRCMaskingOperation>& InMaskingOperation, bool bIsInteractive, const ERCModifyOperationFlags ModifyOperationFlags)
+void FLinearColorMaskingFactory::ApplyMaskedValues(const TSharedRef<FRCMaskingOperation>& InMaskingOperation, bool bIsInteractive)
 {
 	if (FStructProperty* ToStructProp = CastField<FStructProperty>(InMaskingOperation->ObjectRef.Property.Get()))
 	{
 		if (UObject* OwningObject = InMaskingOperation->ObjectRef.Object.Get())
 		{
 #if WITH_EDITOR
-			const bool bWithPropertyChangedEvents = !EnumHasAnyFlags(ModifyOperationFlags, ERCModifyOperationFlags::SkipPropertyChangeEvents);
-			if (bWithPropertyChangedEvents)
-			{
-				OwningObject->PreEditChange(ToStructProp);
-			}
+			OwningObject->PreEditChange(ToStructProp);
 #endif // WITH_EDITOR
 
 			if (const FLinearColor* LinearColorProp = ToStructProp->ContainerPtrToValuePtr<FLinearColor>(OwningObject))
@@ -431,11 +388,8 @@ void FLinearColorMaskingFactory::ApplyMaskedValues(const TSharedRef<FRCMaskingOp
 			}
 
 #if WITH_EDITOR
-			if (bWithPropertyChangedEvents)
-			{
-				FPropertyChangedEvent ChangeEvent(ToStructProp, bIsInteractive ? EPropertyChangeType::Interactive : EPropertyChangeType::ValueSet);
-				OwningObject->PostEditChangeProperty(ChangeEvent);
-			}
+			FPropertyChangedEvent ChangeEvent(ToStructProp, bIsInteractive ? EPropertyChangeType::Interactive : EPropertyChangeType::ValueSet);
+			OwningObject->PostEditChangeProperty(ChangeEvent);
 #endif // WITH_EDITOR
 		}
 	}
@@ -462,3 +416,5 @@ bool FLinearColorMaskingFactory::SupportsExposedEntity(UScriptStruct* ScriptStru
 {
 	return ScriptStruct == TBaseStructure<FLinearColor>::Get();
 }
+
+PRAGMA_ENABLE_DEPRECATION_WARNINGS
diff --git a/Engine/Plugins/VirtualProduction/RemoteControl/Source/RemoteControl/Private/Factories/RemoteControlMaskingFactories.h b/Engine/Plugins/VirtualProduction/RemoteControl/Source/RemoteControl/Private/Factories/RemoteControlMaskingFactories.h
--- a/Engine/Plugins/VirtualProduction/RemoteControl/Source/RemoteControl/Private/Factories/RemoteControlMaskingFactories.h
+++ b/Engine/Plugins/VirtualProduction/RemoteControl/Source/RemoteControl/Private/Factories/RemoteControlMaskingFactories.h
@@ -5,86 +5,142 @@
 #include "CoreMinimal.h"
 #include "Factories/IRemoteControlMaskingFactory.h"
 
+class UE_DEPRECATED(5.5, "Deprecated in favor of RemoteControlMaskingUtil.") FVectorMaskingFactory;
+PRAGMA_DISABLE_DEPRECATION_WARNINGS
 class FVectorMaskingFactory : public IRemoteControlMaskingFactory
 {
 public:
 	static TSharedRef<IRemoteControlMaskingFactory> MakeInstance();
 
 	//~ Begin IRemoteControlMaskingFactory interface
-	virtual void ApplyMaskedValues(const TSharedRef<FRCMaskingOperation>& InMaskingOperation, bool bIsInteractive, const ERCModifyOperationFlags ModifyOperationFlags = ERCModifyOperationFlags::None) override;
+	virtual void ApplyMaskedValues(const TSharedRef<FRCMaskingOperation>& InMaskingOperation, bool bIsInteractive) override;
 	virtual void CacheRawValues(const TSharedRef<FRCMaskingOperation>& InMaskingOperation) override;
 	virtual bool SupportsExposedEntity(UScriptStruct* ScriptStruct) const override;
 	//~ End IRemoteControlMaskingFactory interface
+
+	// Workaround for clang deprecation warnings for any deprecated members in implicitly-defined special member functions
+	FVectorMaskingFactory() = default;
+	FVectorMaskingFactory(const FVectorMaskingFactory&) = default;
+	FVectorMaskingFactory(FVectorMaskingFactory&&) = default;
 };
+PRAGMA_ENABLE_DEPRECATION_WARNINGS
 
+class UE_DEPRECATED(5.5, "Deprecated in favor of RemoteControlMaskingUtil.") FVector4MaskingFactory;
+PRAGMA_DISABLE_DEPRECATION_WARNINGS
 class FVector4MaskingFactory : public IRemoteControlMaskingFactory
 {
 public:
 	static TSharedRef<IRemoteControlMaskingFactory> MakeInstance();
 
 	//~ Begin IRemoteControlMaskingFactory interface
-	virtual void ApplyMaskedValues(const TSharedRef<FRCMaskingOperation>& InMaskingOperation, bool bIsInteractive, const ERCModifyOperationFlags ModifyOperationFlags = ERCModifyOperationFlags::None) override;
+	virtual void ApplyMaskedValues(const TSharedRef<FRCMaskingOperation>& InMaskingOperation, bool bIsInteractive) override;
 	virtual void CacheRawValues(const TSharedRef<FRCMaskingOperation>& InMaskingOperation) override;
 	virtual bool SupportsExposedEntity(UScriptStruct* ScriptStruct) const override;
 	//~ End IRemoteControlMaskingFactory interface
+
+	// Workaround for clang deprecation warnings for any deprecated members in implicitly-defined special member functions
+	FVector4MaskingFactory() = default;
+	FVector4MaskingFactory(const FVector4MaskingFactory&) = default;
+	FVector4MaskingFactory(FVector4MaskingFactory&&) = default;
 };
+PRAGMA_ENABLE_DEPRECATION_WARNINGS
 
+class UE_DEPRECATED(5.5, "Deprecated in favor of RemoteControlMaskingUtil.") FIntVectorMaskingFactory;
+PRAGMA_DISABLE_DEPRECATION_WARNINGS
 class FIntVectorMaskingFactory : public IRemoteControlMaskingFactory
 {
 public:
 	static TSharedRef<IRemoteControlMaskingFactory> MakeInstance();
 
 	//~ Begin IRemoteControlMaskingFactory interface
-	virtual void ApplyMaskedValues(const TSharedRef<FRCMaskingOperation>& InMaskingOperation, bool bIsInteractive, const ERCModifyOperationFlags ModifyOperationFlags = ERCModifyOperationFlags::None) override;
+	virtual void ApplyMaskedValues(const TSharedRef<FRCMaskingOperation>& InMaskingOperation, bool bIsInteractive) override;
 	virtual void CacheRawValues(const TSharedRef<FRCMaskingOperation>& InMaskingOperation) override;
 	virtual bool SupportsExposedEntity(UScriptStruct* ScriptStruct) const override;
 	//~ End IRemoteControlMaskingFactory interface
+
+	// Workaround for clang deprecation warnings for any deprecated members in implicitly-defined special member functions
+	FIntVectorMaskingFactory() = default;
+	FIntVectorMaskingFactory(const FIntVectorMaskingFactory&) = default;
+	FIntVectorMaskingFactory(FIntVectorMaskingFactory&&) = default;
 };
+PRAGMA_ENABLE_DEPRECATION_WARNINGS
 
+class UE_DEPRECATED(5.5, "Deprecated in favor of RemoteControlMaskingUtil.") FIntVector4MaskingFactory;
+PRAGMA_DISABLE_DEPRECATION_WARNINGS
 class FIntVector4MaskingFactory : public IRemoteControlMaskingFactory
 {
 public:
 	static TSharedRef<IRemoteControlMaskingFactory> MakeInstance();
 
 	//~ Begin IRemoteControlMaskingFactory interface
-	virtual void ApplyMaskedValues(const TSharedRef<FRCMaskingOperation>& InMaskingOperation, bool bIsInteractive, const ERCModifyOperationFlags ModifyOperationFlags = ERCModifyOperationFlags::None) override;
+	virtual void ApplyMaskedValues(const TSharedRef<FRCMaskingOperation>& InMaskingOperation, bool bIsInteractive) override;
 	virtual void CacheRawValues(const TSharedRef<FRCMaskingOperation>& InMaskingOperation) override;
 	virtual bool SupportsExposedEntity(UScriptStruct* ScriptStruct) const override;
 	//~ End IRemoteControlMaskingFactory interface
+
+	// Workaround for clang deprecation warnings for any deprecated members in implicitly-defined special member functions
+	FIntVector4MaskingFactory() = default;
+	FIntVector4MaskingFactory(const FIntVector4MaskingFactory&) = default;
+	FIntVector4MaskingFactory(FIntVector4MaskingFactory&&) = default;
 };
+PRAGMA_ENABLE_DEPRECATION_WARNINGS
 
+class UE_DEPRECATED(5.5, "Deprecated in favor of RemoteControlMaskingUtil.") FRotatorMaskingFactory;
+PRAGMA_DISABLE_DEPRECATION_WARNINGS
 class FRotatorMaskingFactory : public IRemoteControlMaskingFactory
 {
 public:
 	static TSharedRef<IRemoteControlMaskingFactory> MakeInstance();
 
 	//~ Begin IRemoteControlMaskingFactory interface
-	virtual void ApplyMaskedValues(const TSharedRef<FRCMaskingOperation>& InMaskingOperation, bool bIsInteractive, const ERCModifyOperationFlags ModifyOperationFlags = ERCModifyOperationFlags::None) override;
+	virtual void ApplyMaskedValues(const TSharedRef<FRCMaskingOperation>& InMaskingOperation, bool bIsInteractive) override;
 	virtual void CacheRawValues(const TSharedRef<FRCMaskingOperation>& InMaskingOperation) override;
 	virtual bool SupportsExposedEntity(UScriptStruct* ScriptStruct) const override;
 	//~ End IRemoteControlMaskingFactory interface
+
+	// Workaround for clang deprecation warnings for any deprecated members in implicitly-defined special member functions
+	FRotatorMaskingFactory() = default;
+	FRotatorMaskingFactory(const FRotatorMaskingFactory&) = default;
+	FRotatorMaskingFactory(FRotatorMaskingFactory&&) = default;
 };
+PRAGMA_ENABLE_DEPRECATION_WARNINGS
 
+class UE_DEPRECATED(5.5, "Deprecated in favor of RemoteControlMaskingUtil.") FColorMaskingFactory;
+PRAGMA_DISABLE_DEPRECATION_WARNINGS
 class FColorMaskingFactory : public IRemoteControlMaskingFactory
 {
 public:
 	static TSharedRef<IRemoteControlMaskingFactory> MakeInstance();
 
 	//~ Begin IRemoteControlMaskingFactory interface
-	virtual void ApplyMaskedValues(const TSharedRef<FRCMaskingOperation>& InMaskingOperation, bool bIsInteractive, const ERCModifyOperationFlags ModifyOperationFlags = ERCModifyOperationFlags::None) override;
+	virtual void ApplyMaskedValues(const TSharedRef<FRCMaskingOperation>& InMaskingOperation, bool bIsInteractive) override;
 	virtual void CacheRawValues(const TSharedRef<FRCMaskingOperation>& InMaskingOperation) override;
 	virtual bool SupportsExposedEntity(UScriptStruct* ScriptStruct) const override;
 	//~ End IRemoteControlMaskingFactory interface
+
+	// Workaround for clang deprecation warnings for any deprecated members in implicitly-defined special member functions
+	FColorMaskingFactory() = default;
+	FColorMaskingFactory(const FColorMaskingFactory&) = default;
+	FColorMaskingFactory(FColorMaskingFactory&&) = default;
 };
+PRAGMA_ENABLE_DEPRECATION_WARNINGS
 
+class UE_DEPRECATED(5.5, "Deprecated in favor of RemoteControlMaskingUtil.") FLinearColorMaskingFactory;
+PRAGMA_DISABLE_DEPRECATION_WARNINGS
 class FLinearColorMaskingFactory : public IRemoteControlMaskingFactory
 {
 public:
 	static TSharedRef<IRemoteControlMaskingFactory> MakeInstance();
 
 	//~ Begin IRemoteControlMaskingFactory interface
-	virtual void ApplyMaskedValues(const TSharedRef<FRCMaskingOperation>& InMaskingOperation, bool bIsInteractive, const ERCModifyOperationFlags ModifyOperationFlags = ERCModifyOperationFlags::None) override;
+	virtual void ApplyMaskedValues(const TSharedRef<FRCMaskingOperation>& InMaskingOperation, bool bIsInteractive) override;
 	virtual void CacheRawValues(const TSharedRef<FRCMaskingOperation>& InMaskingOperation) override;
 	virtual bool SupportsExposedEntity(UScriptStruct* ScriptStruct) const override;
 	//~ End IRemoteControlMaskingFactory interface
+
+	// Workaround for clang deprecation warnings for any deprecated members in implicitly-defined special member functions
+	FLinearColorMaskingFactory() = default;
+	FLinearColorMaskingFactory(const FLinearColorMaskingFactory&) = default;
+	FLinearColorMaskingFactory(FLinearColorMaskingFactory&&) = default;
 };
+PRAGMA_ENABLE_DEPRECATION_WARNINGS
diff --git a/Engine/Plugins/VirtualProduction/RemoteControl/Source/RemoteControl/Private/RemoteControlModule.cpp b/Engine/Plugins/VirtualProduction/RemoteControl/Source/RemoteControl/Private/RemoteControlModule.cpp
--- a/Engine/Plugins/VirtualProduction/RemoteControl/Source/RemoteControl/Private/RemoteControlModule.cpp
+++ b/Engine/Plugins/VirtualProduction/RemoteControl/Source/RemoteControl/Private/RemoteControlModule.cpp
@@ -29,6 +29,7 @@
 #include "RemoteControlInterceptionHelpers.h"
 #include "RemoteControlInterceptionProcessor.h"
 #include "RemoteControlPreset.h"
+#include "RemoteControlProtocolEntityProcessor.h"
 #include "RemoteControlSettings.h"
 #include "SceneInterface.h"
 #include "Serialization/PropertyMapStructDeserializerBackendWrapper.h"
@@ -794,8 +795,10 @@ void FRemoteControlModule::StartupModule()
 	// Register Property Factories
 	RegisterEntityFactory(FRemoteControlInstanceMaterial::StaticStruct()->GetFName(), FRemoteControlInstanceMaterialFactory::MakeInstance());
 
-	// Register Masking Factories
+	// DEPRECATED 5.5, here to keep support of the old implementation while it cannot be removed yet.
+	PRAGMA_DISABLE_DEPRECATION_WARNINGS
 	RegisterMaskingFactories();
+	PRAGMA_ENABLE_DEPRECATION_WARNINGS
 
 	// Register PropertyIdHandler
 	RegisterPropertyIdHandler();
@@ -827,9 +830,13 @@ void FRemoteControlModule::ShutdownModule()
 		// Unregister Property Factories
 		UnregisterEntityFactory(FRemoteControlInstanceMaterial::StaticStruct()->GetFName());
 		
-		// Unregister Default Value & Masking factories.
+		// Unregister Default Value
 		DefaultValueFactories.Empty();
+
+		// DEPRECATED 5.5, here to keep support of the old implementation while it cannot be removed yet.
+		PRAGMA_DISABLE_DEPRECATION_WARNINGS
 		MaskingFactories.Empty();
+		PRAGMA_ENABLE_DEPRECATION_WARNINGS
 	}
 }
 
@@ -1027,8 +1034,10 @@ void FRemoteControlModule::ResetToDefaultValue(UObject* InObject, FRCResetToDefa
 	}
 }
 
-void FRemoteControlModule::PerformMasking(const TSharedRef<FRCMaskingOperation>& InMaskingOperation, const ERCModifyOperationFlags ModifyOperationFlags)
+void FRemoteControlModule::PerformMasking(const TSharedRef<FRCMaskingOperation>& InMaskingOperation)
 {
+	PRAGMA_DISABLE_DEPRECATION_WARNINGS
+	// DEPRECATED 5.5, here to keep support of the old implementation while it cannot be removed yet.
 	TRACE_CPUPROFILER_EVENT_SCOPE(FRemoteControlModule::PerformMasking);
 
 	if (!InMaskingOperation->IsValid())
@@ -1050,7 +1059,7 @@ void FRemoteControlModule::PerformMasking(const TSharedRef<FRCMaskingOperation>&
 			{
 				constexpr bool bIsInteractive = true;
 
-				(*MaskingFactory)->ApplyMaskedValues(InMaskingOperation, bIsInteractive, ModifyOperationFlags);
+				(*MaskingFactory)->ApplyMaskedValues(InMaskingOperation, bIsInteractive);
 
 				ActiveMaskingOperations.Remove(InMaskingOperation);
 			}
@@ -1062,10 +1071,12 @@ void FRemoteControlModule::PerformMasking(const TSharedRef<FRCMaskingOperation>&
 			}
 		}
 	}
+	PRAGMA_ENABLE_DEPRECATION_WARNINGS
 }
 
 void FRemoteControlModule::RegisterMaskingFactoryForType(UScriptStruct* RemoteControlPropertyType, const TSharedPtr<IRemoteControlMaskingFactory>& InMaskingFactory)
 {
+	// DEPRECATED 5.5, here to keep support of the old implementation while it cannot be removed yet.
 	if (!MaskingFactories.Contains(RemoteControlPropertyType))
 	{
 		MaskingFactories.Add(RemoteControlPropertyType, InMaskingFactory);
@@ -1073,22 +1084,21 @@ void FRemoteControlModule::RegisterMaskingFactoryForType(UScriptStruct* RemoteCo
 }
 
 void FRemoteControlModule::UnregisterMaskingFactoryForType(UScriptStruct* RemoteControlPropertyType)
-{
+{	
+	// DEPRECATED 5.5, here to keep support of the old implementation while it cannot be removed yet.
 	MaskingFactories.Remove(RemoteControlPropertyType);
 }
 
 bool FRemoteControlModule::SupportsMasking(const UScriptStruct* InStruct) const
 {
-	return MaskingFactories.Contains(InStruct);
+	using namespace UE::RemoteControl;
+	return ProtocolEntityProcessor::DoesScriptStructSupportMasking(InStruct);
 }
 
 bool FRemoteControlModule::SupportsMasking(const FProperty* InProperty) const
 {
-	if (const FStructProperty* StructProperty = CastField<FStructProperty>(InProperty))
-	{
-		return SupportsMasking(StructProperty->Struct);
-	}
-	return false;
+	using namespace UE::RemoteControl;
+	return ProtocolEntityProcessor::DoesPropertySupportMasking(InProperty);
 }
 
 bool FRemoteControlModule::ResolveCall(const FString& ObjectPath, const FString& FunctionName, FRCCallReference& OutCallRef, FString* OutErrorText)
@@ -1211,14 +1221,14 @@ bool FRemoteControlModule::InvokeCall(FRCCall& InCall, ERCPayloadType InPayloadT
 		}
 
 		const bool bIsManualTransaction = InCall.TransactionMode == ERCTransactionMode::MANUAL;
-		if ((bIsNewTransaction || bIsManualTransaction) && ensureAlways(InCall.CallRef.Object.IsValid()))
+		if ((bIsNewTransaction || bIsManualTransaction) && ensure(InCall.CallRef.Object.IsValid()))
 		{
 			InCall.CallRef.Object->Modify();
 		}
 			
 #endif
 		FEditorScriptExecutionGuard ScriptGuard;
-		if (ensureAlways(InCall.CallRef.Object.IsValid()))
+		if (ensure(InCall.CallRef.Object.IsValid()))
 		{
 			if (InCall.CallRef.PropertyWithSetter.IsValid())
 			{
@@ -2712,6 +2722,8 @@ void FRemoteControlModule::RegisterDefaultValueFactories()
 
 void FRemoteControlModule::RegisterMaskingFactories()
 {
+	// DEPRECATED 5.5, here to keep support of the old implementation while it cannot be removed yet.
+	PRAGMA_DISABLE_DEPRECATION_WARNINGS
 	RegisterMaskingFactoryForType(TBaseStructure<FVector>::Get(), FVectorMaskingFactory::MakeInstance());
 	RegisterMaskingFactoryForType(TBaseStructure<FVector4>::Get(), FVector4MaskingFactory::MakeInstance());
 	RegisterMaskingFactoryForType(TBaseStructure<FIntVector>::Get(), FIntVectorMaskingFactory::MakeInstance());
@@ -2719,6 +2731,7 @@ void FRemoteControlModule::RegisterMaskingFactories()
 	RegisterMaskingFactoryForType(TBaseStructure<FRotator>::Get(), FRotatorMaskingFactory::MakeInstance());
 	RegisterMaskingFactoryForType(TBaseStructure<FColor>::Get(), FColorMaskingFactory::MakeInstance());
 	RegisterMaskingFactoryForType(TBaseStructure<FLinearColor>::Get(), FLinearColorMaskingFactory::MakeInstance());
+	PRAGMA_ENABLE_DEPRECATION_WARNINGS
 }
 
 void FRemoteControlModule::RegisterPropertyIdHandler()
diff --git a/Engine/Plugins/VirtualProduction/RemoteControl/Source/RemoteControl/Private/RemoteControlModule.h b/Engine/Plugins/VirtualProduction/RemoteControl/Source/RemoteControl/Private/RemoteControlModule.h
--- a/Engine/Plugins/VirtualProduction/RemoteControl/Source/RemoteControl/Private/RemoteControlModule.h
+++ b/Engine/Plugins/VirtualProduction/RemoteControl/Source/RemoteControl/Private/RemoteControlModule.h
@@ -35,7 +35,8 @@ public:
 	virtual bool CanResetToDefaultValue(UObject* InObject, const FRCResetToDefaultArgs& InArgs) const override;
 	virtual bool HasDefaultValueCustomization(const UObject* InObject, const FProperty* InProperty) const override;
 	virtual void ResetToDefaultValue(UObject* InObject, FRCResetToDefaultArgs& InArgs) override;
-	virtual void PerformMasking(const TSharedRef<FRCMaskingOperation>& InMaskingOperation, const ERCModifyOperationFlags ModifyOperationFlags = ERCModifyOperationFlags::None) override;
+	UE_DEPRECATED(5.5, "PerformMasking is deprecated, masking is now applied where appropriate automatically.")
+	virtual void PerformMasking(const TSharedRef<FRCMaskingOperation>& InMaskingOperation) override;
 	virtual void RegisterMaskingFactoryForType(UScriptStruct* RemoteControlPropertyType, const TSharedPtr<IRemoteControlMaskingFactory>& InMaskingFactory) override;
 	virtual void UnregisterMaskingFactoryForType(UScriptStruct* RemoteControlPropertyType) override;
 	virtual bool SupportsMasking(const UScriptStruct* InStruct) const override;
diff --git a/Engine/Plugins/VirtualProduction/RemoteControl/Source/RemoteControl/Private/RemoteControlPreset.cpp b/Engine/Plugins/VirtualProduction/RemoteControl/Source/RemoteControl/Private/RemoteControlPreset.cpp
--- a/Engine/Plugins/VirtualProduction/RemoteControl/Source/RemoteControl/Private/RemoteControlPreset.cpp
+++ b/Engine/Plugins/VirtualProduction/RemoteControl/Source/RemoteControl/Private/RemoteControlPreset.cpp
@@ -2520,7 +2520,7 @@ void URemoteControlPreset::OnActorDeleted(AActor* Actor)
 	{
 		for (auto It = Bindings.CreateIterator(); It; ++It)
 		{
-			UObject* ResolvedObject = (*It)->Resolve();
+			UObject* ResolvedObject = *It ? (*It)->Resolve() : nullptr;
 			if (ResolvedObject && (Actor == ResolvedObject || Actor == ResolvedObject->GetTypedOuter<AActor>()))
 			{
 				// Defer binding clean up to next frame in case the actor deletion is actually an actor being moved to a different sub level.
diff --git a/Engine/Plugins/VirtualProduction/RemoteControl/Source/RemoteControl/Private/RemoteControlProtocolBinding.cpp b/Engine/Plugins/VirtualProduction/RemoteControl/Source/RemoteControl/Private/RemoteControlProtocolBinding.cpp
--- a/Engine/Plugins/VirtualProduction/RemoteControl/Source/RemoteControl/Private/RemoteControlProtocolBinding.cpp
+++ b/Engine/Plugins/VirtualProduction/RemoteControl/Source/RemoteControl/Private/RemoteControlProtocolBinding.cpp
@@ -9,156 +9,17 @@
 #include "Factories/IRemoteControlMaskingFactory.h"
 #include "IRemoteControlModule.h"
 #include "RemoteControlPreset.h"
+#include "RemoteControlProtocolEntityInterpolator.h"
 #include "Serialization/MemoryReader.h"
 #include "Serialization/MemoryWriter.h"
 #include "UObject/StructOnScope.h"
 #include "UObject/TextProperty.h"
 
-namespace EntityInterpolation
-{
-	/** 
-	 * The function is taking the map with the range buffer pointer and convert it to the map with value pointers instead. 
-	 * Range keys stay the same.
-	 * For example, we have an input map TArray<TPair<int32, uint8*>> where uint8* points to the float buffer (4 bytes)
-	 * Then we need to convert it to value buffer where we convert uint8* to value pointer. In this case, is float*
-	 * 
-	 * But in cases when we are dealing with Struct or any containers (Array, Map, Structs) inner elements we need to convert the container pointer to value pointer
-	 * And that means if there no FProperty Outer wh should convert uint8* to ValueType* directly, 
-	 * but in case with Outer != nullptr we need to convert with InProperty->ContainerPtrToValuePtr<ValueType>
-	 */
-	template <typename ValueType, typename PropertyType, typename ProtocolValueType>
-	TArray<FRemoteControlProtocolEntity::TRangeMappingData<ProtocolValueType, ValueType>> ContainerPtrMapToValuePtrMap(PropertyType* InProperty, FProperty* Outer, TArray<FRemoteControlProtocolEntity::FRangeMappingData>& InRangeMappingBuffers, int32 InArrayIndex)
-	{
-		TArray<FRemoteControlProtocolEntity::TRangeMappingData<ProtocolValueType, ValueType>> ValueMap;
-		ValueMap.Reserve(InRangeMappingBuffers.Num());
-
-		int32 MappingIndex = 0;
-		for (FRemoteControlProtocolEntity::FRangeMappingData& RangePair : InRangeMappingBuffers)
-		{
-			FRemoteControlProtocolEntity::TRangeMappingData<ProtocolValueType, ValueType> TypedPair;
-			TypedPair.Range = *reinterpret_cast<ProtocolValueType*>(RangePair.Range.GetData());
-
-			if(Outer)
-			{
-				if(Outer->GetClass() == FArrayProperty::StaticClass())
-				{
-					FScriptArrayHelper Helper(CastField<FArrayProperty>(Outer), RangePair.Mapping.GetData());
-					TypedPair.Mapping = *reinterpret_cast<ValueType*>(Helper.GetRawPtr(InArrayIndex));
-				}
-				else
-				{
-					TypedPair.Mapping = *InProperty->template ContainerPtrToValuePtr<ValueType>(RangePair.Mapping.GetData(), InArrayIndex);	
-				}
-			}
-			else
-			{
-				TArray<uint8> Buffer;
-				RemoteControlPropertyUtilities::FRCPropertyVariant Src{InProperty, RangePair.Mapping};
-				RemoteControlPropertyUtilities::FRCPropertyVariant Dst{InProperty, Buffer};
-				RemoteControlPropertyUtilities::Deserialize<PropertyType>(Src, Dst);
-				TypedPair.Mapping = *Dst.GetPropertyValue<ValueType>();
-			}
-
-			ValueMap.Add(TypedPair);
-			MappingIndex++;
-		}
-
-		return ValueMap;
-	}
-
-	/** Wraps FMath::Lerp, allowing Protocol Binding specific specialization */
-	template <class T, class U>
-	static T Lerp(const T& A, const T& B, const U& Alpha)
-	{
-		return FMath::Lerp(A, B, Alpha);
-	}
-
-	/** Specialization for bool, toggles at 0.5 Alpha instead of 1.0 */
-	template <>
-	bool Lerp(const bool& A, const bool& B, const float& Alpha)
-	{
-		return Alpha > 0.5f ? B : A;
-	}
-
-	/** Specialization for FString, toggles at 0.5 Alpha instead of 1.0 */
-	template <>
-	FString Lerp(const FString& A, const FString& B, const float& Alpha)
-	{
-		return Alpha > 0.5f ? B : A;
-	}
-
-	/** Specialization for FName, toggles at 0.5 Alpha instead of 1.0 */
-	template <>
-	FName Lerp(const FName& A, const FName& B, const float& Alpha)
-	{
-		return Alpha > 0.5f ? B : A;
-	}
-
-	/** Specialization for FText, toggles at 0.5 Alpha instead of 1.0 */
-	template <>
-	FText Lerp(const FText& A, const FText& B, const float& Alpha)
-	{
-		return Alpha > 0.5f ? B : A;
-	}
-
-	template <typename ValueType, typename PropertyType, typename ProtocolValueType>
-	bool InterpolateValue(PropertyType* InProperty, FProperty* Outer, TArray<FRemoteControlProtocolEntity::FRangeMappingData>& InRangeMappingBuffers, ProtocolValueType InProtocolValue, ValueType& OutResultValue, int32 ArrayIndex)
-	{
-		TArray<FRemoteControlProtocolEntity::TRangeMappingData<ProtocolValueType, ValueType>> ValueMap = ContainerPtrMapToValuePtrMap<ValueType, PropertyType, ProtocolValueType>(InProperty, Outer, InRangeMappingBuffers, ArrayIndex);
-
-		// sort by input protocol value
-		Algo::SortBy(ValueMap, [](const FRemoteControlProtocolEntity::TRangeMappingData<ProtocolValueType, ValueType>& Item) -> ProtocolValueType
-		{
-			return Item.Range;
-		});
-
-		// clamp to min and max mapped values
-		ProtocolValueType ClampProtocolValue = FMath::Clamp(InProtocolValue, ValueMap[0].Range, ValueMap.Last().Range);
-
-		FRemoteControlProtocolEntity::TRangeMappingData<ProtocolValueType*, ValueType*> RangeMin{nullptr, nullptr};
-		FRemoteControlProtocolEntity::TRangeMappingData<ProtocolValueType*, ValueType*> RangeMax{nullptr, nullptr};
-
-		for (int32 RangeIdx = 0; RangeIdx < ValueMap.Num(); ++RangeIdx)
-		{
-			FRemoteControlProtocolEntity::TRangeMappingData<ProtocolValueType, ValueType>& Range = ValueMap[RangeIdx];
-			if (ClampProtocolValue > Range.Range || RangeMin.Mapping == nullptr)
-			{
-				RangeMin.Range = &Range.Range;
-				RangeMin.Mapping = &Range.Mapping;
-			}
-			else if (ClampProtocolValue <= Range.Range)
-			{
-				RangeMax.Range = &Range.Range;
-				RangeMax.Mapping = &Range.Mapping;
-				// Max found, no need to continue
-				break;
-			}
-		}
-
-		if (RangeMax.Mapping == nullptr || RangeMin.Mapping == nullptr)
-		{
-			return ensure(false);
-		}
-		if (RangeMax.Range == nullptr || RangeMin.Range == nullptr)
-		{
-			return ensure(false);
-		}
-		else if (*RangeMax.Range == *RangeMin.Range)
-		{
-			UE_LOG(LogRemoteControl, Warning, TEXT("Range input values are the same."));
-			return true;
-		}
-
-		// Normalize value to range for interpolation
-		const float Percentage = static_cast<float>(ClampProtocolValue - *RangeMin.Range) / static_cast<float>(*RangeMax.Range - *RangeMin.Range);
-
-		OutResultValue = Lerp(*RangeMin.Mapping, *RangeMax.Mapping, Percentage);
-
-		return true;
-	}
-
+namespace UE::RemoteControl::Private
+{	
 	// Writes a property value to the serialization output.
 	template <typename ValueType>
+	UE_DEPRECATED(5.5, "Instead refer to RemoteControlProtocolEntityProcessor")
 	void WritePropertyValue(FCborWriter& InCborWriter, FProperty* InProperty, const ValueType& Value, bool bWriteName = true)
 	{
 		if (bWriteName)
@@ -170,6 +31,7 @@ namespace EntityInterpolation
 
 	// Specialization for FName that converts to FString
 	template <>
+	UE_DEPRECATED(5.5, "Instead refer to RemoteControlProtocolEntityProcessor")
 	void WritePropertyValue<FName>(FCborWriter& InCborWriter, FProperty* InProperty, const FName& Value, bool bWriteName)
 	{
 		if (bWriteName)
@@ -181,6 +43,7 @@ namespace EntityInterpolation
 
 	// Specialization for FText that converts to FString
 	template <>
+	UE_DEPRECATED(5.5, "Instead refer to RemoteControlProtocolEntityProcessor")
 	void WritePropertyValue<FText>(FCborWriter& InCborWriter, FProperty* InProperty, const FText& Value, bool bWriteName)
 	{
 		if (bWriteName)
@@ -191,8 +54,12 @@ namespace EntityInterpolation
 	}
 
 	template <typename ProtocolValueType>
-	bool WriteProperty(FProperty* InProperty, FProperty* OuterProperty, TArray<FRemoteControlProtocolEntity::FRangeMappingData>& InRangeMappingBuffers, ProtocolValueType InProtocolValue, FCborWriter& InCborWriter, int32 InArrayIndex = 0)
+	UE_DEPRECATED(5.5, "Instead refer to RemoteControlProtocolEntityProcessor")
+	bool WriteProperty(FRemoteControlProtocolEntity* Entity, FProperty* InProperty, FProperty* OuterProperty, TArray<FRemoteControlProtocolEntity::FRangeMappingData>& InRangeMappingBuffers, ProtocolValueType InProtocolValue, FCborWriter& InCborWriter, int32 InArrayIndex = 0)
 	{
+		using namespace UE::RemoteControl;
+
+		PRAGMA_DISABLE_DEPRECATION_WARNINGS
 		// Value nested in Array/Set (except single element) or map as array or as root
 		const bool bIsInArray = OuterProperty != nullptr
 								&& (InProperty->ArrayDim > 1
@@ -205,7 +72,7 @@ namespace EntityInterpolation
 		if (FBoolProperty* BoolProperty = CastField<FBoolProperty>(InProperty))
 		{
 			bool bBoolValue = false;
-			bSuccess = InterpolateValue(BoolProperty, OuterProperty, InRangeMappingBuffers, InProtocolValue, bBoolValue, InArrayIndex);			
+			bSuccess = ProtocolEntityInterpolator::InterpolateValue(Entity, BoolProperty, OuterProperty, InRangeMappingBuffers, InProtocolValue, bBoolValue, InArrayIndex);			
 			bBoolValue = StaticCast<uint8>(bBoolValue) > 0; // Ensure 0 or 1, can be different if property was packed.
 			WritePropertyValue(InCborWriter, InProperty, bBoolValue, !bIsInArray);
 		}
@@ -214,13 +81,13 @@ namespace EntityInterpolation
 			if (CastField<FFloatProperty>(InProperty))
 			{
 				float FloatValue = 0.f;
-				bSuccess = InterpolateValue(NumericProperty, OuterProperty, InRangeMappingBuffers, InProtocolValue, FloatValue, InArrayIndex);
+				bSuccess = ProtocolEntityInterpolator::InterpolateValue(Entity, NumericProperty, OuterProperty, InRangeMappingBuffers, InProtocolValue, FloatValue, InArrayIndex);
 				WritePropertyValue(InCborWriter, InProperty, FloatValue, !bIsInArray);
 			}
 			else if (CastField<FDoubleProperty>(InProperty))
 			{
 				double DoubleValue = 0.0;
-				bSuccess = InterpolateValue(NumericProperty, OuterProperty, InRangeMappingBuffers, InProtocolValue, DoubleValue, InArrayIndex);
+				bSuccess = ProtocolEntityInterpolator::InterpolateValue(Entity, NumericProperty, OuterProperty, InRangeMappingBuffers, InProtocolValue, DoubleValue, InArrayIndex);
 				WritePropertyValue(InCborWriter, InProperty, DoubleValue, !bIsInArray);
 			}
 			else if (NumericProperty->IsInteger() && !NumericProperty->IsEnum())
@@ -228,49 +95,49 @@ namespace EntityInterpolation
 				if (FByteProperty* ByteProperty = CastField<FByteProperty>(InProperty))
 				{
 					uint8 IntValue = 0;
-					bSuccess = InterpolateValue(NumericProperty, OuterProperty, InRangeMappingBuffers, InProtocolValue, IntValue, InArrayIndex);
+					bSuccess = ProtocolEntityInterpolator::InterpolateValue(Entity, NumericProperty, OuterProperty, InRangeMappingBuffers, InProtocolValue, IntValue, InArrayIndex);
 					WritePropertyValue(InCborWriter, InProperty, static_cast<int64>(IntValue), !bIsInArray);
 				}
 				else if (FIntProperty* IntProperty = CastField<FIntProperty>(InProperty))
 				{
 					int IntValue = 0;
-					bSuccess = InterpolateValue(NumericProperty, OuterProperty, InRangeMappingBuffers, InProtocolValue, IntValue, InArrayIndex);
+					bSuccess = ProtocolEntityInterpolator::InterpolateValue(Entity, NumericProperty, OuterProperty, InRangeMappingBuffers, InProtocolValue, IntValue, InArrayIndex);
 					WritePropertyValue(InCborWriter, InProperty, static_cast<int64>(IntValue), !bIsInArray);
 				}
 				else if (FUInt32Property* UInt32Property = CastField<FUInt32Property>(InProperty))
 				{
 					uint32 IntValue = 0;
-					bSuccess = InterpolateValue(NumericProperty, OuterProperty, InRangeMappingBuffers, InProtocolValue, IntValue, InArrayIndex);
+					bSuccess = ProtocolEntityInterpolator::InterpolateValue(Entity, NumericProperty, OuterProperty, InRangeMappingBuffers, InProtocolValue, IntValue, InArrayIndex);
 					WritePropertyValue(InCborWriter, InProperty, static_cast<int64>(IntValue), !bIsInArray);
 				}
 				else if (FInt16Property* Int16Property = CastField<FInt16Property>(InProperty))
 				{
 					int16 IntValue = 0;
-					bSuccess = InterpolateValue(NumericProperty, OuterProperty, InRangeMappingBuffers, InProtocolValue, IntValue, InArrayIndex);
+					bSuccess = ProtocolEntityInterpolator::InterpolateValue(Entity, NumericProperty, OuterProperty, InRangeMappingBuffers, InProtocolValue, IntValue, InArrayIndex);
 					WritePropertyValue(InCborWriter, InProperty, static_cast<int64>(IntValue), !bIsInArray);
 				}
 				else if (FUInt16Property* FInt16Property = CastField<FUInt16Property>(InProperty))
 				{
 					uint16 IntValue = 0;
-					bSuccess = InterpolateValue(NumericProperty, OuterProperty, InRangeMappingBuffers, InProtocolValue, IntValue, InArrayIndex);
+					bSuccess = ProtocolEntityInterpolator::InterpolateValue(Entity, NumericProperty, OuterProperty, InRangeMappingBuffers, InProtocolValue, IntValue, InArrayIndex);
 					WritePropertyValue(InCborWriter, InProperty, static_cast<int64>(IntValue), !bIsInArray);
 				}
 				else if (FInt64Property* Int64Property = CastField<FInt64Property>(InProperty))
 				{
 					int64 IntValue = 0;
-					bSuccess = InterpolateValue(NumericProperty, OuterProperty, InRangeMappingBuffers, InProtocolValue, IntValue, InArrayIndex);
+					bSuccess = ProtocolEntityInterpolator::InterpolateValue(Entity, NumericProperty, OuterProperty, InRangeMappingBuffers, InProtocolValue, IntValue, InArrayIndex);
 					WritePropertyValue(InCborWriter, InProperty, static_cast<int64>(IntValue), !bIsInArray);
 				}
 				else if (FUInt64Property* FInt64Property = CastField<FUInt64Property>(InProperty))
 				{
 					uint64 IntValue = 0;
-					bSuccess = InterpolateValue(NumericProperty, OuterProperty, InRangeMappingBuffers, InProtocolValue, IntValue, InArrayIndex);
+					bSuccess = ProtocolEntityInterpolator::InterpolateValue(Entity, NumericProperty, OuterProperty, InRangeMappingBuffers, InProtocolValue, IntValue, InArrayIndex);
 					WritePropertyValue(InCborWriter, InProperty, static_cast<int64>(IntValue), !bIsInArray);
 				}
 				else if (FInt8Property* Int8Property = CastField<FInt8Property>(InProperty))
 				{
 					int8 IntValue = 0;
-					bSuccess = InterpolateValue(NumericProperty, OuterProperty, InRangeMappingBuffers, InProtocolValue, IntValue, InArrayIndex);
+					bSuccess = ProtocolEntityInterpolator::InterpolateValue(Entity, NumericProperty, OuterProperty, InRangeMappingBuffers, InProtocolValue, IntValue, InArrayIndex);
 					WritePropertyValue(InCborWriter, InProperty, static_cast<int64>(IntValue), !bIsInArray);
 				}
 			}
@@ -298,7 +165,7 @@ namespace EntityInterpolation
 					RangeMappingBuffers.Emplace(RangePair.Range, DataInStruct, InnerProperty->GetSize(), 1);
 				}
 
-				bStructSuccess &= WriteProperty(InnerProperty, StructProperty, RangeMappingBuffers, InProtocolValue, InCborWriter, InArrayIndex);
+				bStructSuccess &= WriteProperty(Entity, InnerProperty, StructProperty, RangeMappingBuffers, InProtocolValue, InCborWriter, InArrayIndex);
 			}
 
 			bSuccess = bStructSuccess;
@@ -308,19 +175,19 @@ namespace EntityInterpolation
 		else if (FStrProperty* StrProperty = CastField<FStrProperty>(InProperty))
 		{
 			FString StringValue;
-			bSuccess = InterpolateValue(StrProperty, OuterProperty, InRangeMappingBuffers, InProtocolValue, StringValue, InArrayIndex);
+			bSuccess = ProtocolEntityInterpolator::InterpolateValue(Entity, StrProperty, OuterProperty, InRangeMappingBuffers, InProtocolValue, StringValue, InArrayIndex);
 			WritePropertyValue(InCborWriter, InProperty, StringValue, !bIsInArray);
 		}
 		else if (FNameProperty* NameProperty = CastField<FNameProperty>(InProperty))
 		{
 			FName NameValue;
-			bSuccess = InterpolateValue(NameProperty, OuterProperty, InRangeMappingBuffers, InProtocolValue, NameValue, InArrayIndex);
+			bSuccess = ProtocolEntityInterpolator::InterpolateValue(Entity, NameProperty, OuterProperty, InRangeMappingBuffers, InProtocolValue, NameValue, InArrayIndex);
 			WritePropertyValue(InCborWriter, InProperty, NameValue, !bIsInArray);
 		}
 		else if (FTextProperty* TextProperty = CastField<FTextProperty>(InProperty))
 		{
 			FText TextValue;
-			bSuccess = InterpolateValue(TextProperty, OuterProperty, InRangeMappingBuffers, InProtocolValue, TextValue, InArrayIndex);
+			bSuccess = ProtocolEntityInterpolator::InterpolateValue(Entity, TextProperty, OuterProperty, InRangeMappingBuffers, InProtocolValue, TextValue, InArrayIndex);
 			WritePropertyValue(InCborWriter, InProperty, TextValue, !bIsInArray);
 		}
 
@@ -334,12 +201,18 @@ namespace EntityInterpolation
 		}
 #endif
 
+		PRAGMA_ENABLE_DEPRECATION_WARNINGS
+
 		return bSuccess;
 	}
 
 	template <typename ProtocolValueType>
-	bool ApplyProtocolValueToProperty(FProperty* InProperty, ProtocolValueType InProtocolValue, TArray<FRemoteControlProtocolEntity::FRangeMappingData>& InRangeMappingBuffers, FCborWriter& InCborWriter)
+	UE_DEPRECATED(5.5, "Instead refer to RemoteControlProtocolEntityProcessor")
+	bool ApplyProtocolValueToProperty(FRemoteControlProtocolEntity* Entity, FProperty* InProperty, ProtocolValueType InProtocolValue, TArray<FRemoteControlProtocolEntity::FRangeMappingData>& InRangeMappingBuffers, FCborWriter& InCborWriter)
 	{
+		PRAGMA_DISABLE_DEPRECATION_WARNINGS
+		using namespace UE::RemoteControl;
+
 		// Structures
 		if (FStructProperty* StructProperty = CastField<FStructProperty>(InProperty))
 		{
@@ -351,9 +224,9 @@ namespace EntityInterpolation
 			bool bStructSuccess = true;
 			for (TFieldIterator<FProperty> It(ScriptStruct); It; ++It)
 			{
-				bStructSuccess &= WriteProperty(*It, StructProperty, InRangeMappingBuffers, InProtocolValue, InCborWriter);
+				bStructSuccess &= WriteProperty(Entity, *It, StructProperty, InRangeMappingBuffers, InProtocolValue, InCborWriter);
 			}
-
+			
 			InCborWriter.WriteContainerEnd();
 
 			return bStructSuccess;
@@ -384,7 +257,7 @@ namespace EntityInterpolation
 
 			for (auto ArrayIndex = 0; ArrayIndex < ElementWithSmallestNum->NumElements; ++ArrayIndex)
 			{
-				bArraySuccess &= WriteProperty(InnerProperty, ArrayProperty, InRangeMappingBuffers, InProtocolValue, InCborWriter, ArrayIndex);
+				bArraySuccess &= WriteProperty(Entity, InnerProperty, ArrayProperty, InRangeMappingBuffers, InProtocolValue, InCborWriter, ArrayIndex);
 			}
 
 			InCborWriter.WriteContainerEnd();
@@ -416,8 +289,9 @@ namespace EntityInterpolation
 			// All other properties
 		else
 		{
-			return WriteProperty(InProperty, nullptr, InRangeMappingBuffers, InProtocolValue, InCborWriter);
+			return WriteProperty(Entity, InProperty, nullptr, InRangeMappingBuffers, InProtocolValue, InCborWriter);
 		}
+		PRAGMA_ENABLE_DEPRECATION_WARNINGS
 	}
 }
 
@@ -636,6 +510,9 @@ const FString& FRemoteControlProtocolEntity::GetRangePropertyMaxValue() const
 
 bool FRemoteControlProtocolEntity::ApplyProtocolValueToProperty(const double InProtocolValue)
 {
+	// DEPRECATED 5.5
+
+	PRAGMA_DISABLE_DEPRECATION_WARNINGS
 	if (Mappings.Num() <= 1)
 	{
 		UE_LOG(LogRemoteControl, Warning, TEXT("Binding doesn't container any range mappings."));
@@ -695,7 +572,6 @@ bool FRemoteControlProtocolEntity::ApplyProtocolValueToProperty(const double InP
 		IRemoteControlModule::Get().ResolveObjectProperty(ObjectRef.Access, Object, ObjectRef.PropertyPathInfo, ObjectRef);
 
 		TSharedPtr<FRCMaskingOperation> MaskingOperation = MakeShared<FRCMaskingOperation>(ObjectRef.PropertyPathInfo, Object);
-
 		if (OverridenMasks == ERCMask::NoMask)
 		{
 			MaskingOperation->Masks = RemoteControlProperty->GetActiveMasks();
@@ -706,7 +582,7 @@ bool FRemoteControlProtocolEntity::ApplyProtocolValueToProperty(const double InP
 		}
 
 		// Cache the values.
-		IRemoteControlModule::Get().PerformMasking(MaskingOperation.ToSharedRef(), ModifyOperationFlags);
+		IRemoteControlModule::Get().PerformMasking(MaskingOperation.ToSharedRef());
 
 		// Set properties after interpolation
 		TArray<uint8> InterpolatedBuffer;
@@ -717,11 +593,12 @@ bool FRemoteControlProtocolEntity::ApplyProtocolValueToProperty(const double InP
 			FMemoryReader MemoryReader(InterpolatedBuffer);
 			FCborStructDeserializerBackend CborStructDeserializerBackend(MemoryReader);
 			bSuccess &= IRemoteControlModule::Get().SetObjectProperties(ObjectRef, CborStructDeserializerBackend, ERCPayloadType::Cbor, InterpolatedBuffer, ModifyOperation, ModifyOperationFlags);
-
+		
 			// Apply the masked the values.
-			IRemoteControlModule::Get().PerformMasking(MaskingOperation.ToSharedRef(), ModifyOperationFlags);
+			IRemoteControlModule::Get().PerformMasking(MaskingOperation.ToSharedRef());
 		}
 	}
+	PRAGMA_ENABLE_DEPRECATION_WARNINGS
 
 	return bSuccess;
 }
@@ -778,6 +655,12 @@ const FName FRemoteControlProtocolEntity::GetPropertyName(const FName& ForColumn
 
 bool FRemoteControlProtocolEntity::GetInterpolatedPropertyBuffer(FProperty* InProperty, double InProtocolValue, TArray<uint8>& OutBuffer)
 {
+	// DEPRECATED 5.5
+
+	PRAGMA_DISABLE_DEPRECATION_WARNINGS
+
+	using namespace UE::RemoteControl;
+
 	OutBuffer.Empty();
 
 	TArray<FRemoteControlProtocolEntity::FRangeMappingData> RangeMappingBuffers = GetRangeMappingBuffers();
@@ -792,34 +675,34 @@ bool FRemoteControlProtocolEntity::GetInterpolatedPropertyBuffer(FProperty* InPr
 	switch (*GetRangePropertyName().ToEName())
 	{
 	case NAME_Int8Property:
-		bSuccess = EntityInterpolation::ApplyProtocolValueToProperty(InProperty, static_cast<int8>(InProtocolValue), RangeMappingBuffers, CborWriter);
+		bSuccess = Private::ApplyProtocolValueToProperty(this, InProperty, static_cast<int8>(InProtocolValue), RangeMappingBuffers, CborWriter);
 		break;
 	case NAME_Int16Property:
-		bSuccess = EntityInterpolation::ApplyProtocolValueToProperty(InProperty, static_cast<int16>(InProtocolValue), RangeMappingBuffers, CborWriter);
+		bSuccess = Private::ApplyProtocolValueToProperty(this, InProperty, static_cast<int16>(InProtocolValue), RangeMappingBuffers, CborWriter);
 		break;
 	case NAME_IntProperty:
-		bSuccess = EntityInterpolation::ApplyProtocolValueToProperty(InProperty, static_cast<int32>(InProtocolValue), RangeMappingBuffers, CborWriter);
+		bSuccess = Private::ApplyProtocolValueToProperty(this, InProperty, static_cast<int32>(InProtocolValue), RangeMappingBuffers, CborWriter);
 		break;
 	case NAME_Int64Property:
-		bSuccess = EntityInterpolation::ApplyProtocolValueToProperty(InProperty, static_cast<int64>(InProtocolValue), RangeMappingBuffers, CborWriter);
+		bSuccess = Private::ApplyProtocolValueToProperty(this, InProperty, static_cast<int64>(InProtocolValue), RangeMappingBuffers, CborWriter);
 		break;
 	case NAME_ByteProperty:
-		bSuccess = EntityInterpolation::ApplyProtocolValueToProperty(InProperty, static_cast<uint8>(InProtocolValue), RangeMappingBuffers, CborWriter);
+		bSuccess = Private::ApplyProtocolValueToProperty(this, InProperty, static_cast<uint8>(InProtocolValue), RangeMappingBuffers, CborWriter);
 		break;
 	case NAME_UInt16Property:
-		bSuccess = EntityInterpolation::ApplyProtocolValueToProperty(InProperty, static_cast<uint16>(InProtocolValue), RangeMappingBuffers, CborWriter);
+		bSuccess = Private::ApplyProtocolValueToProperty(this, InProperty, static_cast<uint16>(InProtocolValue), RangeMappingBuffers, CborWriter);
 		break;
 	case NAME_UInt32Property:
-		bSuccess = EntityInterpolation::ApplyProtocolValueToProperty(InProperty, static_cast<uint32>(InProtocolValue), RangeMappingBuffers, CborWriter);
+		bSuccess = Private::ApplyProtocolValueToProperty(this, InProperty, static_cast<uint32>(InProtocolValue), RangeMappingBuffers, CborWriter);
 		break;
 	case NAME_UInt64Property:
-		bSuccess = EntityInterpolation::ApplyProtocolValueToProperty(InProperty, static_cast<uint64>(InProtocolValue), RangeMappingBuffers, CborWriter);
+		bSuccess = Private::ApplyProtocolValueToProperty(this, InProperty, static_cast<uint64>(InProtocolValue), RangeMappingBuffers, CborWriter);
 		break;
 	case NAME_FloatProperty:
-		bSuccess = EntityInterpolation::ApplyProtocolValueToProperty(InProperty, static_cast<float>(InProtocolValue), RangeMappingBuffers, CborWriter);
+		bSuccess = Private::ApplyProtocolValueToProperty(this, InProperty, static_cast<float>(InProtocolValue), RangeMappingBuffers, CborWriter);
 		break;
 	case NAME_DoubleProperty:
-		bSuccess = EntityInterpolation::ApplyProtocolValueToProperty(InProperty, static_cast<double>(InProtocolValue), RangeMappingBuffers, CborWriter);
+		bSuccess = Private::ApplyProtocolValueToProperty(this, InProperty, static_cast<double>(InProtocolValue), RangeMappingBuffers, CborWriter);
 		break;
 	default:
 		checkNoEntry();
@@ -827,6 +710,8 @@ bool FRemoteControlProtocolEntity::GetInterpolatedPropertyBuffer(FProperty* InPr
 
 	CborWriter.WriteContainerEnd();
 
+	PRAGMA_ENABLE_DEPRECATION_WARNINGS
+
 	return bSuccess;
 }
 
diff --git a/Engine/Plugins/VirtualProduction/RemoteControl/Source/RemoteControl/Private/RemoteControlProtocolEntityInterpolator.h b/Engine/Plugins/VirtualProduction/RemoteControl/Source/RemoteControl/Private/RemoteControlProtocolEntityInterpolator.h
new file mode 100644
--- /dev/null
+++ b/Engine/Plugins/VirtualProduction/RemoteControl/Source/RemoteControl/Private/RemoteControlProtocolEntityInterpolator.h
@@ -0,0 +1,189 @@
+// Copyright Epic Games, Inc. All Rights Reserved.
+
+#pragma once
+
+#include "IRemoteControlModule.h"
+#include "RemoteControlProtocolBinding.h"
+
+namespace UE::RemoteControl::ProtocolEntityInterpolator
+{
+	namespace Internal
+	{
+		/** 
+		 * Converts a map with the range buffer pointers to a map with value pointers. Range keys stay the same.
+		 *
+		 * For example, when we have an input map TArray<TPair<int32, uint8*>> where uint8* points to the float buffer (4 bytes).
+		 * we need to convert it to value buffer where we convert uint8* to value pointer. In this case, is float*.
+		 */
+		template <typename ValueType, typename PropertyType, typename RangeValueType>
+		TArray<FRemoteControlProtocolEntity::TRangeMappingData<RangeValueType, ValueType>> ContainerPtrMapToValuePtrMap(PropertyType* InProperty, FProperty* Outer, TArray<FRemoteControlProtocolEntity::FRangeMappingData>& InRangeMappingBuffers, int32 InArrayIndex)
+		{
+			TArray<FRemoteControlProtocolEntity::TRangeMappingData<RangeValueType, ValueType>> ValueMap;
+			ValueMap.Reserve(InRangeMappingBuffers.Num());
+
+			int32 MappingIndex = 0;
+			for (FRemoteControlProtocolEntity::FRangeMappingData& RangePair : InRangeMappingBuffers)
+			{
+				FRemoteControlProtocolEntity::TRangeMappingData<RangeValueType, ValueType> TypedPair;
+				TypedPair.Range = *reinterpret_cast<RangeValueType*>(RangePair.Range.GetData());
+
+				// In cases when we are dealing with Struct or any containers (Array, Map, Structs) inner elements we need to convert the container pointer to value pointer
+				if (Outer)
+				{
+					if (Outer->GetClass() == FArrayProperty::StaticClass())
+					{
+						FScriptArrayHelper Helper(CastField<FArrayProperty>(Outer), RangePair.Mapping.GetData());
+						TypedPair.Mapping = *reinterpret_cast<ValueType*>(Helper.GetRawPtr(InArrayIndex));
+					}
+					else
+					{
+						TypedPair.Mapping = *InProperty->template ContainerPtrToValuePtr<ValueType>(RangePair.Mapping.GetData(), InArrayIndex);
+					}
+				}
+				else
+				{
+					TArray<uint8> Buffer;
+					RemoteControlPropertyUtilities::FRCPropertyVariant Src{ InProperty, RangePair.Mapping };
+					RemoteControlPropertyUtilities::FRCPropertyVariant Dst{ InProperty, Buffer };
+					RemoteControlPropertyUtilities::Deserialize<PropertyType>(Src, Dst);
+					TypedPair.Mapping = *Dst.GetPropertyValue<ValueType>();
+				}
+
+				ValueMap.Add(TypedPair);
+				MappingIndex++;
+			}
+
+			return ValueMap;
+		}
+
+		/** Wraps FMath::Lerp, allowing value specific specialization */
+		template <class T, class U>
+		static T Lerp(const T& A, const T& B, const U& Alpha)
+		{
+			return FMath::Lerp(A, B, Alpha);
+		}
+
+		/** Specialization for bool, toggles at 0.5 Alpha instead of 1.0 */
+		template <>
+		bool Lerp(const bool& A, const bool& B, const float& Alpha)
+		{
+			return Alpha >= 0.5f ? B : A;
+		}
+
+		/** Specialization for FString, toggles at 0.5 Alpha instead of 1.0 */
+		template <>
+		FString Lerp(const FString& A, const FString& B, const float& Alpha)
+		{
+			return Alpha >= 0.5f ? B : A;
+		}
+
+		/** Specialization for FName, toggles at 0.5 Alpha instead of 1.0 */
+		template <>
+		FName Lerp(const FName& A, const FName& B, const float& Alpha)
+		{
+			return Alpha >= 0.5f ? B : A;
+		}
+
+		/** Specialization for FText, toggles at 0.5 Alpha instead of 1.0 */
+		template <>
+		FText Lerp(const FText& A, const FText& B, const float& Alpha)
+		{
+			return Alpha >= 0.5f ? B : A;
+		}
+
+		/** Interpolates the range of an entity, using the protocol value as alpha */
+		template <typename RangeValueType, typename ValueType, typename PropertyType>
+		bool InterpolateValue(PropertyType* InProperty, FProperty* InOuter, TArray<FRemoteControlProtocolEntity::FRangeMappingData>& InRangeMappingBuffers, RangeValueType InProtocolValue, ValueType& OutResultValue, int32 InArrayIndex = 0)
+		{
+			TArray<FRemoteControlProtocolEntity::TRangeMappingData<RangeValueType, ValueType>> ValueMap = ContainerPtrMapToValuePtrMap<ValueType, PropertyType, RangeValueType>(InProperty, InOuter, InRangeMappingBuffers, InArrayIndex);
+
+			// sort by input protocol value
+			Algo::SortBy(ValueMap, [](const FRemoteControlProtocolEntity::TRangeMappingData<RangeValueType, ValueType>& Item) -> RangeValueType
+				{
+					return Item.Range;
+				});
+
+			if (ValueMap.IsEmpty())
+			{
+				return false;
+			}
+
+			// clamp to min and max mapped values
+			RangeValueType ClampProtocolValue = FMath::Clamp(InProtocolValue, ValueMap[0].Range, ValueMap.Last().Range);
+
+			FRemoteControlProtocolEntity::TRangeMappingData<RangeValueType*, ValueType*> RangeMin{ nullptr, nullptr };
+			FRemoteControlProtocolEntity::TRangeMappingData<RangeValueType*, ValueType*> RangeMax{ nullptr, nullptr };
+
+			for (int32 RangeIdx = 0; RangeIdx < ValueMap.Num(); ++RangeIdx)
+			{
+				FRemoteControlProtocolEntity::TRangeMappingData<RangeValueType, ValueType>& Range = ValueMap[RangeIdx];
+				if (ClampProtocolValue > Range.Range || RangeMin.Mapping == nullptr)
+				{
+					RangeMin.Range = &Range.Range;
+					RangeMin.Mapping = &Range.Mapping;
+				}
+				else if (ClampProtocolValue <= Range.Range)
+				{
+					RangeMax.Range = &Range.Range;
+					RangeMax.Mapping = &Range.Mapping;
+					// Max found, no need to continue
+					break;
+				}
+			}
+
+			if (RangeMax.Mapping == nullptr || RangeMin.Mapping == nullptr)
+			{
+				return ensure(false);
+			}
+			if (RangeMax.Range == nullptr || RangeMin.Range == nullptr)
+			{
+				return ensure(false);
+			}
+			else if (*RangeMax.Range == *RangeMin.Range)
+			{
+				UE_LOG(LogRemoteControl, Warning, TEXT("Range input values are the same."));
+				return true;
+			}
+
+			const float Alpha = static_cast<float>(ClampProtocolValue - *RangeMin.Range) / static_cast<float>(*RangeMax.Range - *RangeMin.Range);
+
+			OutResultValue = Lerp(*RangeMin.Mapping, *RangeMax.Mapping, Alpha);
+
+			return true;
+		}
+	}
+
+	/** Interpolates the range of an entity, using the protocol value as alpha */
+	template <typename ValueType, typename PropertyType>
+	bool InterpolateValue(FRemoteControlProtocolEntity* InEntity, PropertyType* InProperty, FProperty* InOuter, TArray<FRemoteControlProtocolEntity::FRangeMappingData>& InRangeMappingBuffers, double InProtocolValue, ValueType& OutResultValue, int32 InArrayIndex = 0)
+	{
+		// Depending on the range property type
+		switch (*InEntity->GetRangePropertyName().ToEName())
+		{
+		case NAME_DoubleProperty:
+			return Internal::InterpolateValue<double>(InProperty, InOuter, InRangeMappingBuffers, InProtocolValue, OutResultValue, InArrayIndex);
+		case NAME_FloatProperty:
+			return Internal::InterpolateValue<float>(InProperty, InOuter, InRangeMappingBuffers, InProtocolValue, OutResultValue, InArrayIndex);
+		case NAME_Int8Property:
+			return Internal::InterpolateValue<int8>(InProperty, InOuter, InRangeMappingBuffers, InProtocolValue, OutResultValue, InArrayIndex);
+		case NAME_Int16Property:
+			return Internal::InterpolateValue<int16>(InProperty, InOuter, InRangeMappingBuffers, InProtocolValue, OutResultValue, InArrayIndex);
+		case NAME_IntProperty:
+			return Internal::InterpolateValue<int32>(InProperty, InOuter, InRangeMappingBuffers, InProtocolValue, OutResultValue, InArrayIndex);
+		case NAME_Int64Property:
+			return Internal::InterpolateValue<int64>(InProperty, InOuter, InRangeMappingBuffers, InProtocolValue, OutResultValue, InArrayIndex);
+		case NAME_ByteProperty:
+			return Internal::InterpolateValue<uint8>(InProperty, InOuter, InRangeMappingBuffers, InProtocolValue, OutResultValue, InArrayIndex);
+		case NAME_UInt16Property:
+			return Internal::InterpolateValue<uint16>(InProperty, InOuter, InRangeMappingBuffers, InProtocolValue, OutResultValue, InArrayIndex);
+		case NAME_UInt32Property:
+			return Internal::InterpolateValue<uint32>(InProperty, InOuter, InRangeMappingBuffers, InProtocolValue, OutResultValue, InArrayIndex);
+		case NAME_UInt64Property:
+			return Internal::InterpolateValue<uint64>(InProperty, InOuter, InRangeMappingBuffers, InProtocolValue, OutResultValue, InArrayIndex);
+		default:
+			checkNoEntry();
+		}
+
+		return false;
+	}
+}
diff --git a/Engine/Plugins/VirtualProduction/RemoteControl/Source/RemoteControl/Private/RemoteControlProtocolEntityProcessor.cpp b/Engine/Plugins/VirtualProduction/RemoteControl/Source/RemoteControl/Private/RemoteControlProtocolEntityProcessor.cpp
new file mode 100644
--- /dev/null
+++ b/Engine/Plugins/VirtualProduction/RemoteControl/Source/RemoteControl/Private/RemoteControlProtocolEntityProcessor.cpp
@@ -0,0 +1,840 @@
+// Copyright Epic Games, Inc. All Rights Reserved.
+
+#include "RemoteControlProtocolEntityProcessor.h"
+
+#include "Algo/Find.h"
+#include "Algo/MinElement.h"
+#include "Backends/CborStructDeserializerBackend.h"
+#include "CborWriter.h"
+#include "GameFramework/Actor.h"
+#include "IRemoteControlModule.h"
+#include "RemoteControlCommon.h"
+#include "RemoteControlField.h"
+#include "RemoteControlPreset.h"
+#include "RemoteControlProtocolEntityInterpolator.h"
+
+namespace UE::RemoteControl::ProtocolEntityProcessor
+{
+	/** Fundamental types */
+	namespace Types
+	{
+		/** Holds the data that correspond to a single property. */
+		struct FRCPropertyData
+		{
+			/** The preset in which the property resides */
+			URemoteControlPreset* Preset = nullptr;
+
+			/** The property */
+			TSharedPtr<FRemoteControlProperty> Property;
+
+			/** The entities of this property with their protocol value */
+			TMap<FRemoteControlProtocolEntity*, double> EntityToValueMap;
+		};
+	}
+
+	/** Methods to write cbor values */
+	namespace Cbor
+	{
+		/** Writes a property value to the serialization output */
+		template <typename ValueType>
+		void WritePropertyValue(FCborWriter& CborWriter, FProperty* InProperty, const ValueType& Value, bool bWriteName = true)
+		{
+			if (bWriteName)
+			{
+				CborWriter.WriteValue(InProperty->GetName());
+			}
+
+			CborWriter.WriteValue(Value);
+		}
+
+		/** Writes an FName property value to the serialization output */
+		template <>
+		void WritePropertyValue<FName>(FCborWriter& CborWriter, FProperty* InProperty, const FName& Value, bool bWriteName)
+		{
+			if (bWriteName)
+			{
+				CborWriter.WriteValue(InProperty->GetName());
+			}
+			CborWriter.WriteValue(Value.ToString());
+		}
+
+		/** Writes an FText property value to the serialization output */
+		template <>
+		void WritePropertyValue<FText>(FCborWriter& CborWriter, FProperty* InProperty, const FText& Value, bool bWriteName)
+		{
+			if (bWriteName)
+			{
+				CborWriter.WriteValue(InProperty->GetName());
+			}
+			CborWriter.WriteValue(Value.ToString());
+		}
+	}
+
+	/** Methods to apply masking */
+	namespace Masking
+	{
+		template <typename VectorType>
+		bool IsInMask(const FName& PropertyName, ERCMask Mask)
+		{
+			if constexpr (
+				std::is_same_v<VectorType, FVector> ||
+				std::is_same_v<VectorType, FIntVector>)
+			{
+				static const FName X = "X";
+				static const FName Y = "Y";
+				static const FName Z = "Z";
+
+				return
+					(EnumHasAnyFlags(Mask, ERCMask::MaskA) && PropertyName == X) ||
+					(EnumHasAnyFlags(Mask, ERCMask::MaskB) && PropertyName == Y) ||
+					(EnumHasAnyFlags(Mask, ERCMask::MaskC) && PropertyName == Z);
+			}
+			else if constexpr (
+				std::is_same_v<VectorType, FVector4> ||
+				std::is_same_v<VectorType, FIntVector4>)
+			{
+				static const FName X = "X";
+				static const FName Y = "Y";
+				static const FName Z = "Z";
+				static const FName W = "W";
+
+				return
+					(EnumHasAnyFlags(Mask, ERCMask::MaskA) && PropertyName == X) ||
+					(EnumHasAnyFlags(Mask, ERCMask::MaskB) && PropertyName == Y) ||
+					(EnumHasAnyFlags(Mask, ERCMask::MaskC) && PropertyName == Z) ||
+					(EnumHasAnyFlags(Mask, ERCMask::MaskD) && PropertyName == W);
+			}
+			else if constexpr (
+				std::is_same_v<VectorType, FColor> ||
+				std::is_same_v<VectorType, FLinearColor>)
+			{
+				static const FName R = "R";
+				static const FName G = "G";
+				static const FName B = "B";
+				static const FName A = "A";
+
+				return
+					(EnumHasAnyFlags(Mask, ERCMask::MaskA) && PropertyName == R) ||
+					(EnumHasAnyFlags(Mask, ERCMask::MaskB) && PropertyName == G) ||
+					(EnumHasAnyFlags(Mask, ERCMask::MaskC) && PropertyName == B) ||
+					(EnumHasAnyFlags(Mask, ERCMask::MaskD) && PropertyName == A);
+			}
+			else if constexpr (std::is_same_v<VectorType, FRotator>)
+			{
+				static const FName Roll = "Roll";
+				static const FName Pitch = "Pitch";
+				static const FName Yaw = "Yaw";
+
+				return
+					(EnumHasAnyFlags(Mask, ERCMask::MaskA) && PropertyName == Roll) ||
+					(EnumHasAnyFlags(Mask, ERCMask::MaskB) && PropertyName == Pitch) ||
+					(EnumHasAnyFlags(Mask, ERCMask::MaskC) && PropertyName == Yaw);
+			}
+			else
+			{
+				return false;
+			}
+		}
+	}
+
+	/** Methods to write properties with a Cbor writer */
+	namespace PropertyWriter
+	{
+		/** Writes an interpolated property value using the Cbor writer */
+		template <typename ProtocolValueType>
+		bool WriteInterpolatedPropertyValue(
+			FRemoteControlProtocolEntity* Entity, 
+			FProperty* InProperty, 
+			FProperty* OuterProperty, 
+			TArray<FRemoteControlProtocolEntity::FRangeMappingData>& InRangeMappingBuffers, 
+			ProtocolValueType InProtocolValue, 
+			FCborWriter& CborWriter, 
+			int32 InArrayIndex = 0)
+		{
+			// Value nested in Array/Set (except single element) or map as array or as root
+			const bool bIsInArray = OuterProperty != nullptr
+				&& (InProperty->ArrayDim > 1
+					|| OuterProperty->GetClass() == FArrayProperty::StaticClass()
+					|| OuterProperty->GetClass() == FSetProperty::StaticClass()
+					|| OuterProperty->GetClass() == FMapProperty::StaticClass());
+
+
+			bool bSuccess = false;
+			if (FBoolProperty* BoolProperty = CastField<FBoolProperty>(InProperty))
+			{
+				bool bBoolValue = false;
+				bSuccess = ProtocolEntityInterpolator::InterpolateValue(Entity, BoolProperty, OuterProperty, InRangeMappingBuffers, InProtocolValue, bBoolValue, InArrayIndex);
+				bBoolValue = StaticCast<uint8>(bBoolValue) > 0; // Ensure 0 or 1, can be different if property was packed.
+				Cbor::WritePropertyValue(CborWriter, InProperty, bBoolValue, !bIsInArray);
+			}
+			else if (FNumericProperty* NumericProperty = CastField<FNumericProperty>(InProperty))
+			{
+				if (CastField<FFloatProperty>(InProperty))
+				{
+					float FloatValue = 0.f;
+					bSuccess = ProtocolEntityInterpolator::InterpolateValue(Entity, NumericProperty, OuterProperty, InRangeMappingBuffers, InProtocolValue, FloatValue, InArrayIndex);
+					Cbor::WritePropertyValue(CborWriter, InProperty, FloatValue, !bIsInArray);
+				}
+				else if (CastField<FDoubleProperty>(InProperty))
+				{
+					double DoubleValue = 0.0;
+					bSuccess = ProtocolEntityInterpolator::InterpolateValue(Entity, NumericProperty, OuterProperty, InRangeMappingBuffers, InProtocolValue, DoubleValue, InArrayIndex);
+					Cbor::WritePropertyValue(CborWriter, InProperty, DoubleValue, !bIsInArray);
+				}
+				else if (NumericProperty->IsInteger() && !NumericProperty->IsEnum())
+				{
+					if (FByteProperty* ByteProperty = CastField<FByteProperty>(InProperty))
+					{
+						uint8 IntValue = 0;
+						bSuccess = ProtocolEntityInterpolator::InterpolateValue(Entity, NumericProperty, OuterProperty, InRangeMappingBuffers, InProtocolValue, IntValue, InArrayIndex);
+						Cbor::WritePropertyValue(CborWriter, InProperty, static_cast<int64>(IntValue), !bIsInArray);
+					}
+					else if (FIntProperty* IntProperty = CastField<FIntProperty>(InProperty))
+					{
+						int IntValue = 0;
+						bSuccess = ProtocolEntityInterpolator::InterpolateValue(Entity, NumericProperty, OuterProperty, InRangeMappingBuffers, InProtocolValue, IntValue, InArrayIndex);
+						Cbor::WritePropertyValue(CborWriter, InProperty, static_cast<int64>(IntValue), !bIsInArray);
+					}
+					else if (FUInt32Property* UInt32Property = CastField<FUInt32Property>(InProperty))
+					{
+						uint32 IntValue = 0;
+						bSuccess = ProtocolEntityInterpolator::InterpolateValue(Entity, NumericProperty, OuterProperty, InRangeMappingBuffers, InProtocolValue, IntValue, InArrayIndex);
+						Cbor::WritePropertyValue(CborWriter, InProperty, static_cast<int64>(IntValue), !bIsInArray);
+					}
+					else if (FInt16Property* Int16Property = CastField<FInt16Property>(InProperty))
+					{
+						int16 IntValue = 0;
+						bSuccess = ProtocolEntityInterpolator::InterpolateValue(Entity, NumericProperty, OuterProperty, InRangeMappingBuffers, InProtocolValue, IntValue, InArrayIndex);
+						Cbor::WritePropertyValue(CborWriter, InProperty, static_cast<int64>(IntValue), !bIsInArray);
+					}
+					else if (FUInt16Property* FInt16Property = CastField<FUInt16Property>(InProperty))
+					{
+						uint16 IntValue = 0;
+						bSuccess = ProtocolEntityInterpolator::InterpolateValue(Entity, NumericProperty, OuterProperty, InRangeMappingBuffers, InProtocolValue, IntValue, InArrayIndex);
+						Cbor::WritePropertyValue(CborWriter, InProperty, static_cast<int64>(IntValue), !bIsInArray);
+					}
+					else if (FInt64Property* Int64Property = CastField<FInt64Property>(InProperty))
+					{
+						int64 IntValue = 0;
+						bSuccess = ProtocolEntityInterpolator::InterpolateValue(Entity, NumericProperty, OuterProperty, InRangeMappingBuffers, InProtocolValue, IntValue, InArrayIndex);
+						Cbor::WritePropertyValue(CborWriter, InProperty, static_cast<int64>(IntValue), !bIsInArray);
+					}
+					else if (FUInt64Property* FInt64Property = CastField<FUInt64Property>(InProperty))
+					{
+						uint64 IntValue = 0;
+						bSuccess = ProtocolEntityInterpolator::InterpolateValue(Entity, NumericProperty, OuterProperty, InRangeMappingBuffers, InProtocolValue, IntValue, InArrayIndex);
+						Cbor::WritePropertyValue(CborWriter, InProperty, static_cast<int64>(IntValue), !bIsInArray);
+					}
+					else if (FInt8Property* Int8Property = CastField<FInt8Property>(InProperty))
+					{
+						int8 IntValue = 0;
+						bSuccess = ProtocolEntityInterpolator::InterpolateValue(Entity, NumericProperty, OuterProperty, InRangeMappingBuffers, InProtocolValue, IntValue, InArrayIndex);
+						Cbor::WritePropertyValue(CborWriter, InProperty, static_cast<int64>(IntValue), !bIsInArray);
+					}
+				}
+			}
+			else if (FStructProperty* StructProperty = CastField<FStructProperty>(InProperty))
+			{
+				if (!bIsInArray)
+				{
+					CborWriter.WriteValue(StructProperty->GetName());
+				}
+
+				CborWriter.WriteContainerStart(ECborCode::Map, -1/*Indefinite*/);
+				{
+					for (TFieldIterator<FProperty> It(StructProperty->Struct); It; ++It)
+					{
+						TArray<FRemoteControlProtocolEntity::FRangeMappingData> RangeMappingBuffers;
+						RangeMappingBuffers.Reserve(InRangeMappingBuffers.Num());
+
+						FProperty* InnerProperty = *It;
+						for (const FRemoteControlProtocolEntity::FRangeMappingData& RangePair : InRangeMappingBuffers)
+						{
+							const uint8* DataInContainer = StructProperty->ContainerPtrToValuePtr<const uint8>(RangePair.Mapping.GetData(), InArrayIndex);
+							const uint8* DataInStruct = InnerProperty->ContainerPtrToValuePtr<const uint8>(DataInContainer);
+							RangeMappingBuffers.Emplace(RangePair.Range, DataInStruct, InnerProperty->GetSize(), 1);
+						}
+
+						bSuccess &= WriteInterpolatedPropertyValue(Entity, InnerProperty, StructProperty, RangeMappingBuffers, InProtocolValue, CborWriter, InArrayIndex);
+					}
+				}
+				CborWriter.WriteContainerEnd();
+			}
+
+			else if (FStrProperty* StrProperty = CastField<FStrProperty>(InProperty))
+			{
+				FString StringValue;
+				bSuccess = ProtocolEntityInterpolator::InterpolateValue(Entity, StrProperty, OuterProperty, InRangeMappingBuffers, InProtocolValue, StringValue, InArrayIndex);
+				Cbor::WritePropertyValue(CborWriter, InProperty, StringValue, !bIsInArray);
+			}
+			else if (FNameProperty* NameProperty = CastField<FNameProperty>(InProperty))
+			{
+				FName NameValue;
+				bSuccess = ProtocolEntityInterpolator::InterpolateValue(Entity, NameProperty, OuterProperty, InRangeMappingBuffers, InProtocolValue, NameValue, InArrayIndex);
+				Cbor::WritePropertyValue(CborWriter, InProperty, NameValue, !bIsInArray);
+			}
+			else if (FTextProperty* TextProperty = CastField<FTextProperty>(InProperty))
+			{
+				FText TextValue;
+				bSuccess = ProtocolEntityInterpolator::InterpolateValue(Entity, TextProperty, OuterProperty, InRangeMappingBuffers, InProtocolValue, TextValue, InArrayIndex);
+				Cbor::WritePropertyValue(CborWriter, InProperty, TextValue, !bIsInArray);
+			}
+
+#if !UE_BUILD_SHIPPING && UE_BUILD_DEBUG
+			if (!bSuccess)
+			{
+				if (RemoteControlTypeUtilities::IsSupportedMappingType(InProperty))
+				{
+					UE_LOG(LogRemoteControl, Error, TEXT("Property type %s is supported for mapping, but unhandled in ProtocolEntityProcessor::WriteInterpolatedPropertyValue"), *InProperty->GetClass()->GetName());
+				}
+			}
+#endif
+
+			return bSuccess;
+		}
+
+		/** Writes the property value as is, without using the protocol value and interpolating */
+		bool WriteUnchangedPropertyValue(
+			const FRCObjectReference& ObjectRef, 
+			FProperty* Property, 
+			FProperty* OuterProperty, 
+			FCborWriter& CborWriter, 
+			int32 InArrayIndex = 0)
+		{
+			if (!Property || !OuterProperty)
+			{
+				return false;
+			}
+
+			// Get the outer in which properties are stored
+			void* Outer = [&OuterProperty, &ObjectRef]() -> void*
+				{
+					if (OuterProperty->GetOwner<UClass>())
+					{
+						return OuterProperty->ContainerPtrToValuePtr<void>(ObjectRef.Object.Get());
+					}
+					else
+					{
+						FProperty* OwnerProperty = OuterProperty->GetOwnerProperty();
+						if (OwnerProperty && OwnerProperty != OuterProperty)
+						{
+							return OuterProperty->ContainerPtrToValuePtr<void>(OwnerProperty);
+						}
+					}
+					
+					return nullptr;
+				}();
+			if (!Outer)
+			{
+				return false;
+			}
+
+			// Value nested in Array/Set (except single element) or map as array or as root
+			const bool bIsInArray = OuterProperty != nullptr
+				&& (Property->ArrayDim > 1
+					|| OuterProperty->GetClass() == FArrayProperty::StaticClass()
+					|| OuterProperty->GetClass() == FSetProperty::StaticClass()
+					|| OuterProperty->GetClass() == FMapProperty::StaticClass());
+
+			bool bSuccess = false;
+			if (FBoolProperty* BoolProperty = CastField<FBoolProperty>(Property))
+			{
+				const bool* ValuePtr = Property->ContainerPtrToValuePtr<bool>(Outer);
+				const bool bBoolValue = ValuePtr ? *ValuePtr : false;
+				Cbor::WritePropertyValue(CborWriter, Property, bBoolValue, !bIsInArray);
+			}
+			else if (FNumericProperty* NumericProperty = CastField<FNumericProperty>(Property))
+			{
+				if (FFloatProperty* FloatPropery = CastField<FFloatProperty>(Property))
+				{
+					const float* ValuePtr = Property->ContainerPtrToValuePtr<float>(Outer);
+					const float FloatValue = ValuePtr ? *ValuePtr : 0.f;
+					Cbor::WritePropertyValue(CborWriter, Property, FloatValue, !bIsInArray);
+				}
+				else if (FDoubleProperty* DoubleProperty = CastField<FDoubleProperty>(Property))
+				{
+					const double* ValuePtr = Property->ContainerPtrToValuePtr<double>(Outer);
+					const double DoubleValue = ValuePtr ? *ValuePtr : 0.0;
+					Cbor::WritePropertyValue(CborWriter, Property, DoubleValue, !bIsInArray);
+				}
+				else if (NumericProperty->IsInteger() && !NumericProperty->IsEnum())
+				{
+					if (FByteProperty* ByteProperty = CastField<FByteProperty>(Property))
+					{
+						const uint8* ValuePtr = Property->ContainerPtrToValuePtr<uint8>(Outer);
+						const uint8 ByteValue = ValuePtr ? *ValuePtr : 0;
+						Cbor::WritePropertyValue(CborWriter, Property, static_cast<int64>(ByteValue), !bIsInArray);
+					}
+					else if (FIntProperty* IntProperty = CastField<FIntProperty>(Property))
+					{
+						const int32* ValuePtr = Property->ContainerPtrToValuePtr<int32>(Outer);
+						const int32 IntValue = ValuePtr ? *ValuePtr : 0;
+						Cbor::WritePropertyValue(CborWriter, Property, static_cast<int64>(IntValue), !bIsInArray);
+					}
+					else if (FUInt32Property* UInt32Property = CastField<FUInt32Property>(Property))
+					{
+						const uint32* ValuePtr = Property->ContainerPtrToValuePtr<uint32>(Outer);
+						const uint32 IntValue = ValuePtr ? *ValuePtr : 0;
+						Cbor::WritePropertyValue(CborWriter, Property, static_cast<int64>(IntValue), !bIsInArray);
+					}
+					else if (FInt16Property* Int16Property = CastField<FInt16Property>(Property))
+					{
+						const int16* ValuePtr = Property->ContainerPtrToValuePtr<int16>(Outer);
+						const int16 IntValue = ValuePtr ? *ValuePtr : 0;
+						Cbor::WritePropertyValue(CborWriter, Property, static_cast<int64>(IntValue), !bIsInArray);
+					}
+					else if (FUInt16Property* FInt16Property = CastField<FUInt16Property>(Property))
+					{
+						const uint16* ValuePtr = Property->ContainerPtrToValuePtr<uint16>(Outer);
+						const uint16 IntValue = ValuePtr ? *ValuePtr : 0;
+						Cbor::WritePropertyValue(CborWriter, Property, static_cast<int64>(IntValue), !bIsInArray);
+					}
+					else if (FInt64Property* Int64Property = CastField<FInt64Property>(Property))
+					{
+						const int64* ValuePtr = Property->ContainerPtrToValuePtr<int64>(Outer);
+						const int64 IntValue = ValuePtr ? *ValuePtr : 0;
+						Cbor::WritePropertyValue(CborWriter, Property, static_cast<int64>(IntValue), !bIsInArray);
+					}
+					else if (FUInt64Property* FInt64Property = CastField<FUInt64Property>(Property))
+					{
+						const uint64* ValuePtr = Property->ContainerPtrToValuePtr<uint64>(Outer);
+						const uint64 IntValue = ValuePtr ? *ValuePtr : 0;
+						Cbor::WritePropertyValue(CborWriter, Property, static_cast<int64>(IntValue), !bIsInArray);
+					}
+					else if (FInt8Property* Int8Property = CastField<FInt8Property>(Property))
+					{
+						const int8* ValuePtr = Property->ContainerPtrToValuePtr<int8>(Outer);
+						int8 IntValue = ValuePtr ? *ValuePtr : 0;
+						Cbor::WritePropertyValue(CborWriter, Property, static_cast<int64>(IntValue), !bIsInArray);
+					}
+				}
+			}
+			else if (FStructProperty* StructProperty = CastField<FStructProperty>(Property))
+			{
+				if (!bIsInArray)
+				{
+					CborWriter.WriteValue(StructProperty->GetName());
+				}
+
+				CborWriter.WriteContainerStart(ECborCode::Map, -1/*Indefinite*/);
+				{
+					bool bStructSuccess = true;
+					for (TFieldIterator<FProperty> It(StructProperty->Struct); It; ++It)
+					{
+						bSuccess &= WriteUnchangedPropertyValue(ObjectRef, *It, StructProperty, CborWriter, InArrayIndex);
+					}
+				}
+				CborWriter.WriteContainerEnd();
+			}
+
+			else if (FStrProperty* StrProperty = CastField<FStrProperty>(Property))
+			{
+				const FString* ValuePtr = Property->ContainerPtrToValuePtr<FString>(Outer);
+				const FString StringValue = ValuePtr ? *ValuePtr : TEXT("");
+				Cbor::WritePropertyValue(CborWriter, Property, StringValue, !bIsInArray);
+			}
+			else if (FNameProperty* NameProperty = CastField<FNameProperty>(Property))
+			{
+				const FName* ValuePtr = Property->ContainerPtrToValuePtr<FName>(Outer);
+				const FName NameValue = ValuePtr ? *ValuePtr : NAME_None;
+				Cbor::WritePropertyValue(CborWriter, Property, NameValue, !bIsInArray);
+			}
+			else if (FTextProperty* TextProperty = CastField<FTextProperty>(Property))
+			{
+				const FText* ValuePtr = Property->ContainerPtrToValuePtr<FText>(Outer);
+				const FText TextValue = ValuePtr ? *ValuePtr : FText();
+				Cbor::WritePropertyValue(CborWriter, Property, TextValue, !bIsInArray);
+			}
+
+#if !UE_BUILD_SHIPPING && UE_BUILD_DEBUG
+			if (!bSuccess)
+			{
+				if (RemoteControlTypeUtilities::IsSupportedMappingType(Property))
+				{
+					UE_LOG(LogRemoteControl, Error, TEXT("Property type %s is supported for mapping, but unhandled in ProtocolEntityProcessor::WriteUnchangedPropertyValue"), *Property->GetClass()->GetName());
+				}
+			}
+#endif
+
+			return bSuccess;
+		}
+
+		/** Writes a struct property that supports masking, considering the mask */
+		template <typename VectorType>
+		bool WriteMaskedStructProperty(
+			const FRCObjectReference& ObjectRef, 
+			FStructProperty* StructProperty, 
+			const Types::FRCPropertyData& PropertyData, 
+			FCborWriter& CborWriter)
+		{
+			using namespace UE::RemoteControl;
+
+			if (!StructProperty)
+			{
+				return false;
+			}
+
+			bool bSuccess = true;
+			for (TFieldIterator<FProperty> It(StructProperty->Struct); It; ++It)
+			{
+				const FName& PropertyName = (*It)->GetFName();
+
+				// Find an entity that fits the current mask and property name 
+				const TTuple<FRemoteControlProtocolEntity*, double>* EntityToValuePtr = Algo::FindByPredicate(PropertyData.EntityToValueMap,
+					[&PropertyName, &PropertyData](const TTuple<FRemoteControlProtocolEntity*, double>& EntityToValuePair)
+					{
+						const ERCMask Mask = EntityToValuePair.Key->GetOverridenMask() == ERCMask::NoMask ?
+							PropertyData.Property->GetActiveMasks() :
+							EntityToValuePair.Key->GetOverridenMask();
+
+						return Masking::IsInMask<VectorType>(PropertyName, Mask);
+					});
+
+				if (EntityToValuePtr)
+				{
+					// Write the interpolated property value
+					TArray<FRemoteControlProtocolEntity::FRangeMappingData> RangeMappingBuffers = EntityToValuePtr->Key->GetRangeMappingBuffers();
+					bSuccess &= WriteInterpolatedPropertyValue(EntityToValuePtr->Key, (*It), StructProperty, RangeMappingBuffers, EntityToValuePtr->Value, CborWriter);
+				}
+				else
+				{
+					// There is no property that fits the mask, hence write the unchanged property value
+					bSuccess &= WriteUnchangedPropertyValue(ObjectRef, (*It), StructProperty, CborWriter);
+				}
+			}
+
+			return true;
+		}
+
+		/** Writes a single property using the Cbor writer */
+		bool WriteProperty(const FRCObjectReference& ObjectRef, FProperty* Property, FProperty* OuterProperty, const Types::FRCPropertyData& PropertyData, FCborWriter& CborWriter, int32 ArrayIndex = 0)
+		{
+			FStructProperty* StructProperty = CastField<FStructProperty>(OuterProperty);
+			const bool bIsMaskedStructProperty = StructProperty && DoesScriptStructSupportMasking(StructProperty->Struct);
+
+			bool bSuccess = true;
+			if (bIsMaskedStructProperty)
+			{
+				if (StructProperty->Struct == TBaseStructure<FVector>::Get())
+				{
+					bSuccess = WriteMaskedStructProperty<FVector>(ObjectRef, StructProperty, PropertyData, CborWriter);
+				}
+				else if (StructProperty->Struct == TBaseStructure<FRotator>::Get())
+				{
+					bSuccess = WriteMaskedStructProperty<FRotator>(ObjectRef, StructProperty, PropertyData, CborWriter);
+				}
+				else if (StructProperty->Struct == TBaseStructure<FColor>::Get())
+				{
+					bSuccess = WriteMaskedStructProperty<FColor>(ObjectRef, StructProperty, PropertyData, CborWriter);
+				}
+				else if (StructProperty->Struct == TBaseStructure<FLinearColor>::Get())
+				{
+					bSuccess = WriteMaskedStructProperty<FLinearColor>(ObjectRef, StructProperty, PropertyData, CborWriter);
+				}
+				else if (StructProperty->Struct == TBaseStructure<FIntVector>::Get())
+				{
+					bSuccess = WriteMaskedStructProperty<FIntVector>(ObjectRef, StructProperty, PropertyData, CborWriter);
+				}
+				else if (StructProperty->Struct == TBaseStructure<FVector4>::Get())
+				{
+					bSuccess = WriteMaskedStructProperty<FVector4>(ObjectRef, StructProperty, PropertyData, CborWriter);
+				}
+				else if (StructProperty->Struct == TBaseStructure<FIntVector4>::Get())
+				{
+					bSuccess = WriteMaskedStructProperty<FIntVector4>(ObjectRef, StructProperty, PropertyData, CborWriter);
+				}
+				else
+				{
+					ensureMsgf(0, TEXT("Struct type should but does not support masking. Cannot process remote control for struct."), *StructProperty->Struct->GetName());
+				}
+			}
+			else
+			{
+				for (const TTuple<FRemoteControlProtocolEntity*, double>& EntityToValuePair : PropertyData.EntityToValueMap)
+				{
+					TArray<FRemoteControlProtocolEntity::FRangeMappingData> RangeMappingBuffers = EntityToValuePair.Key->GetRangeMappingBuffers();
+
+					bSuccess &= WriteInterpolatedPropertyValue(EntityToValuePair.Key, Property, OuterProperty, RangeMappingBuffers, EntityToValuePair.Value, CborWriter, ArrayIndex);
+				}
+			}
+
+			return bSuccess;
+		}
+
+		/** Converts properties that need to be processed into a map of Property Ids and their related property data */
+		TMap<FGuid, Types::FRCPropertyData> BuildPropertyToDataMap(const TMap<TSharedPtr<TStructOnScope<FRemoteControlProtocolEntity>>, double>& EntityToValueMap)
+		{
+			TMap<FGuid, Types::FRCPropertyData> PropertyToDataMap;
+			for (const TTuple<TSharedPtr<TStructOnScope<FRemoteControlProtocolEntity>>, double>& EntityToValuePair : EntityToValueMap)
+			{
+				if (EntityToValuePair.Key.IsValid() && EntityToValuePair.Key->IsValid())
+				{
+					FRemoteControlProtocolEntity* Entity = EntityToValuePair.Key->CastChecked<FRemoteControlProtocolEntity>();
+					const FGuid& PropertyId = Entity->GetPropertyId();
+
+					Types::FRCPropertyData& PropertyData = PropertyToDataMap.FindOrAdd(Entity->GetPropertyId());
+					if (!PropertyData.Preset)
+					{
+						PropertyData.Preset = Entity->GetOwner().Get();
+					}
+
+					// Property may turn invalid if the entity got deleted, but did not unbind from protocols yet
+					if (!PropertyData.Preset)
+					{
+						PropertyToDataMap.Remove(PropertyId);
+						continue;
+					}
+
+					if (!PropertyData.Property.IsValid())
+					{
+						PropertyData.Property = PropertyData.Preset->GetExposedEntity<FRemoteControlProperty>(PropertyId).Pin();
+					}
+
+					// Property may turn invalid if the entity got deleted, but did not unbind from protocols yet
+					if (!PropertyData.Property)
+					{
+						PropertyToDataMap.Remove(PropertyId);
+						continue;
+					}
+
+					PropertyData.EntityToValueMap.Add(EntityToValuePair.Key->CastChecked<FRemoteControlProtocolEntity>(), EntityToValuePair.Value);
+				}
+			}
+
+			return PropertyToDataMap;
+		}
+
+		/** Serializes a single property into the Cbor buffer */
+		bool SerializeProperty(const FRCObjectReference& ObjectRef, const Types::FRCPropertyData& PropertyData, TArray<uint8>& OutCborBuffer)
+		{
+			using namespace UE::RemoteControl;
+
+			if (PropertyData.EntityToValueMap.IsEmpty())
+			{
+				return false;
+			}
+
+			bool bSuccess = true;
+
+			FMemoryWriter MemoryWriter(OutCborBuffer);
+			FCborWriter CborWriter(&MemoryWriter);
+
+			CborWriter.WriteContainerStart(ECborCode::Map, -1/*Indefinite*/);
+			{
+				FProperty* Property = PropertyData.Property->GetProperty();
+
+				// Structs
+				if (FStructProperty* StructProperty = CastField<FStructProperty>(Property))
+				{
+					CborWriter.WriteValue(StructProperty->GetName());
+					CborWriter.WriteContainerStart(ECborCode::Map, -1/*Indefinite*/);
+					{
+						bool bStructSuccess = true;
+						for (TFieldIterator<FProperty> It(StructProperty->Struct); It; ++It)
+						{
+							bSuccess &= WriteProperty(ObjectRef, *It, StructProperty, PropertyData, CborWriter);
+						}
+					}
+					CborWriter.WriteContainerEnd();
+				}
+
+				// @note: temporarily disabled - array of primitives supported, array of structs not
+				// Dynamic arrays
+				else if (FArrayProperty* ArrayProperty = CastField<FArrayProperty>(Property))
+				{
+					CborWriter.WriteValue(ArrayProperty->GetName());
+					CborWriter.WriteContainerStart(ECborCode::Array, -1/*Indefinite*/);
+					{
+						check(!PropertyData.EntityToValueMap.IsEmpty());
+
+						FProperty* InnerProperty = ArrayProperty->Inner;
+						TArray<FRemoteControlProtocolEntity::FRangeMappingData> RangeMappingBuffers = PropertyData.EntityToValueMap.begin()->Key->GetRangeMappingBuffers();
+
+						// Get minimum item count
+						const FRemoteControlProtocolEntity::FRangeMappingData* ElementWithSmallestNum = Algo::MinElementBy(RangeMappingBuffers, [](const FRemoteControlProtocolEntity::FRangeMappingData& RangePair)
+							{
+								return RangePair.NumElements;
+							});
+
+						// No elements in array
+						if (ElementWithSmallestNum)
+						{
+							for (auto ArrayIndex = 0; ArrayIndex < ElementWithSmallestNum->NumElements; ++ArrayIndex)
+							{
+								bSuccess &= WriteProperty(ObjectRef, InnerProperty, ArrayProperty, PropertyData, CborWriter, ArrayIndex);
+							}
+						}
+					}
+					CborWriter.WriteContainerEnd();
+				}
+
+				// Maps
+				else if (FMapProperty* MapProperty = CastField<FMapProperty>(Property))
+				{
+					UE_LOG(LogRemoteControl, Warning, TEXT("MapProperty not supported"));
+					bSuccess = false;
+				}
+
+				// Sets
+				else if (FSetProperty* SetProperty = CastField<FSetProperty>(Property))
+				{
+					UE_LOG(LogRemoteControl, Warning, TEXT("SetProperty not supported"));
+					bSuccess = false;
+				}
+
+				// Static arrays
+				else if (Property->ArrayDim > 1)
+				{
+					UE_LOG(LogRemoteControl, Warning, TEXT("Static arrays not supported"));
+					bSuccess = false;
+				}
+
+				// All other properties
+				else
+				{
+					bSuccess = WriteProperty(ObjectRef, Property, nullptr, PropertyData, CborWriter);
+				}
+			}
+			CborWriter.WriteContainerEnd();
+
+			return bSuccess;
+		}
+
+		/** Compares and swaps the protocol value of each entity in property data, returns true if all are equal. */
+		bool CompareSwapProtocolValuesForProperty(const Types::FRCPropertyData& PropertyData)
+		{
+			bool bEqual = true;
+			for (const TTuple<FRemoteControlProtocolEntity*, double>& EntityToValuePair : PropertyData.EntityToValueMap)
+			{
+				check(EntityToValuePair.Key);
+				if (!FMath::IsNearlyEqual(EntityToValuePair.Key->ProtocolValue, EntityToValuePair.Value))
+				{
+					bEqual = false;
+					EntityToValuePair.Key->ProtocolValue = EntityToValuePair.Value;
+				}
+			}
+
+			return bEqual;
+		}
+	}
+	
+	bool DoesScriptStructSupportMasking(const UScriptStruct* InStruct)
+	{
+		if (InStruct)
+		{
+			// When adding new types here, also add them to logic in ProcessProperties
+			return
+				InStruct == TBaseStructure<FVector>::Get() ||
+				InStruct == TBaseStructure<FRotator>::Get() ||
+				InStruct == TBaseStructure<FColor>::Get() ||
+				InStruct == TBaseStructure<FLinearColor>::Get() ||
+				InStruct == TBaseStructure<FIntVector>::Get() ||
+				InStruct == TBaseStructure<FVector4>::Get() ||
+				InStruct == TBaseStructure<FIntVector4>::Get();
+		}
+
+		return false;
+	}
+
+	bool DoesPropertySupportMasking(const FProperty* Property)
+	{
+		const FStructProperty* StructProperty = CastField<FStructProperty>(Property);
+		return StructProperty ? DoesScriptStructSupportMasking(StructProperty->Struct) : false;
+	}
+
+	void ProcessEntities(const TMap<TSharedPtr<TStructOnScope<FRemoteControlProtocolEntity>>, double>& EntityToValueMap)
+	{
+		using namespace PropertyWriter;
+		using namespace Types;
+
+		const TMap<FGuid, FRCPropertyData> PropertyToDataMap = BuildPropertyToDataMap(EntityToValueMap);
+		
+		if (PropertyToDataMap.IsEmpty())
+		{
+			return;
+		}
+
+		// Process each property
+		TArray<AActor*> ChangedActors;
+		for (const TTuple<FGuid, FRCPropertyData>& PropertyToDataPair : PropertyToDataMap)
+		{
+			const URemoteControlPreset* Preset = PropertyToDataPair.Value.Preset;
+			const TSharedPtr<FRemoteControlProperty>& RemoteControlProperty = PropertyToDataPair.Value.Property;
+			if (!PropertyToDataPair.Value.Property->IsBound())
+			{
+				UE_LOG(LogRemoteControl, Warning, TEXT("Entity isn't bound to any objects."));
+				continue;
+			}
+
+			if (RemoteControlProperty->GetActiveMasks() == ERCMask::NoMask)
+			{
+				continue;
+			}
+
+			FProperty* Property = RemoteControlProperty->GetProperty();
+			if (!Property)
+			{
+				continue;
+			}
+
+			if (!RemoteControlTypeUtilities::IsSupportedMappingType(Property))
+			{
+				UE_LOG(LogRemoteControl, Warning, TEXT("Property type %s is unsupported for mapping."), *Property->GetClass()->GetName());
+				continue;
+			}
+
+			if (CompareSwapProtocolValuesForProperty(PropertyToDataPair.Value))
+			{
+				continue;
+			}
+
+			FRCObjectReference ObjectRef;
+			ObjectRef.Property = Property;
+			ObjectRef.Access = ERCAccess::WRITE_ACCESS;
+			ObjectRef.PropertyPathInfo = RemoteControlProperty->FieldPathInfo.ToString();
+
+			const ERCModifyOperationFlags ModifyOperationFlags = Preset->GetModifyOperationFlagsForProtocols();
+			if (EnumHasAnyFlags(ModifyOperationFlags, ERCModifyOperationFlags::None))
+			{
+				ObjectRef.Access = ERCAccess::WRITE_TRANSACTION_ACCESS;
+			}
+
+			bool bSuccess = true;
+			for (UObject* BoundObject : RemoteControlProperty->GetBoundObjects())
+			{
+				// Resolve the property for the currently bound bject
+				if (!IRemoteControlModule::Get().ResolveObjectProperty(ObjectRef.Access, BoundObject, ObjectRef.PropertyPathInfo, ObjectRef))
+				{
+					continue;
+				}
+
+				// Remember actors about to be changed
+				if (AActor* Actor = Cast<AActor>(BoundObject))
+				{
+					ChangedActors.AddUnique(Actor);
+				}
+				else if (AActor* ParentActor = BoundObject->GetTypedOuter<AActor>())
+				{
+					ChangedActors.AddUnique(ParentActor);
+				}
+
+				TArray<uint8> CborBuffer;
+				if (SerializeProperty(ObjectRef, PropertyToDataPair.Value, CborBuffer))
+				{
+					constexpr ERCModifyOperation ModifyOperation = ERCModifyOperation::EQUAL;
+
+					FMemoryReader MemoryReader(CborBuffer);
+					FCborStructDeserializerBackend CborStructDeserializerBackend(MemoryReader);
+					bSuccess &= IRemoteControlModule::Get().SetObjectProperties(ObjectRef, CborStructDeserializerBackend, ERCPayloadType::Cbor, CborBuffer, ModifyOperation, ModifyOperationFlags);
+				}
+			}
+		}
+
+		// Update changed actors
+		for (AActor* Actor : ChangedActors)
+		{
+			Actor->UpdateComponentTransforms();
+			Actor->MarkComponentsRenderStateDirty();
+		}
+	}
+}
diff --git a/Engine/Plugins/VirtualProduction/RemoteControl/Source/RemoteControl/Public/Factories/IRemoteControlMaskingFactory.h b/Engine/Plugins/VirtualProduction/RemoteControl/Source/RemoteControl/Public/Factories/IRemoteControlMaskingFactory.h
--- a/Engine/Plugins/VirtualProduction/RemoteControl/Source/RemoteControl/Public/Factories/IRemoteControlMaskingFactory.h
+++ b/Engine/Plugins/VirtualProduction/RemoteControl/Source/RemoteControl/Public/Factories/IRemoteControlMaskingFactory.h
@@ -9,20 +9,25 @@
 /**
  * Reference to a UObject or one of its properties for the purpose of masking.
  */
+struct UE_DEPRECATED(5.5, "FRCMaskingOperation is deprecated, masking is now applied where appropriate automatically.") FRCMaskingOperation;
 struct FRCMaskingOperation
 {
-	FRCMaskingOperation() = default;
+	FRCMaskingOperation()
+		: ModifyOperationFlags(ERCModifyOperationFlags::None)
+	{}
 
-	explicit FRCMaskingOperation(FRCFieldPathInfo InPathInfo, UObject* InObject)
+	explicit FRCMaskingOperation(FRCFieldPathInfo InPathInfo, UObject* InObject, const ERCModifyOperationFlags InModifyOperationFlag = ERCModifyOperationFlags::None)
 		: OperationId(FGuid::NewGuid())
 		, ObjectRef(ERCAccess::NO_ACCESS, InObject, InPathInfo)
+		, ModifyOperationFlags(InModifyOperationFlag)
 	{
 		check(InObject);
 	}
 
-	explicit FRCMaskingOperation(const FRCObjectReference& InObjectRef)
+	explicit FRCMaskingOperation(const FRCObjectReference& InObjectRef, const ERCModifyOperationFlags InModifyOperationFlag = ERCModifyOperationFlags::None)
 		: OperationId(FGuid::NewGuid())
 		, ObjectRef(InObjectRef)
+		, ModifyOperationFlags(InModifyOperationFlag)
 	{
 	}
 
@@ -47,7 +52,6 @@ struct FRCMaskingOperation
 	}
 
 public:
-
 	/** Unique identifier of the operation being performed. */
 	FGuid OperationId;
 
@@ -59,15 +63,25 @@ public:
 
 	/** Holds the state of this RC property before applying any masking. */
 	FVector4 PreMaskingCache = FVector4::Zero();
+
+	/** Modify operation flags used when masking */
+	const ERCModifyOperationFlags ModifyOperationFlags;
+	
+	// Workaround for clang deprecation warnings for any deprecated members in implicitly-defined special member functions
+	PRAGMA_DISABLE_DEPRECATION_WARNINGS
+	FRCMaskingOperation(const FRCMaskingOperation&) = default;
+	FRCMaskingOperation(FRCMaskingOperation&&) = default;
+	PRAGMA_ENABLE_DEPRECATION_WARNINGS
 };
 
 /**
  * Factory which is responsible for masking support for FRemoteControlProperty.
  */
+class UE_DEPRECATED(5.5, "IRemoteControlMaskingFactory is deprecated, masking is now applied where appropriate automatically.") IRemoteControlMaskingFactory;
 class IRemoteControlMaskingFactory : public TSharedFromThis<IRemoteControlMaskingFactory>
 {
 public:
-
+	PRAGMA_DISABLE_DEPRECATION_WARNINGS
 	/** Virtual destructor */
 	virtual ~IRemoteControlMaskingFactory(){}
 
@@ -78,7 +92,7 @@ public:
 	 * @param bIsInteractive				If bWithPropertyChangedEvents, defined if the property changed events are interactive or not.
 	 * @param ModifyOperationFlags			(optional) Flags that specify how the property is modified when the value is applied.
 	 */
-	virtual void ApplyMaskedValues(const TSharedRef<FRCMaskingOperation>& InMaskingOperation, bool bIsInteractive, const ERCModifyOperationFlags ModifyOperationFlags = ERCModifyOperationFlags::None) = 0;
+	virtual void ApplyMaskedValues(const TSharedRef<FRCMaskingOperation>& InMaskingOperation, bool bIsInteractive) {};
 
 	/**
 	 * Caches premasking values from the given struct property.
@@ -92,5 +106,10 @@ public:
 	 * @return true if the script struct is supported by given factory
 	 */
 	virtual bool SupportsExposedEntity(UScriptStruct* ScriptStruct) const = 0;
-};
 
+	// Workaround for clang deprecation warnings for any deprecated members in implicitly-defined special member functions
+	IRemoteControlMaskingFactory() = default;
+	IRemoteControlMaskingFactory(const IRemoteControlMaskingFactory&) = default;
+	IRemoteControlMaskingFactory(IRemoteControlMaskingFactory&&) = default;
+	PRAGMA_ENABLE_DEPRECATION_WARNINGS
+};
diff --git a/Engine/Plugins/VirtualProduction/RemoteControl/Source/RemoteControl/Public/IRemoteControlModule.h b/Engine/Plugins/VirtualProduction/RemoteControl/Source/RemoteControl/Public/IRemoteControlModule.h
--- a/Engine/Plugins/VirtualProduction/RemoteControl/Source/RemoteControl/Public/IRemoteControlModule.h
+++ b/Engine/Plugins/VirtualProduction/RemoteControl/Source/RemoteControl/Public/IRemoteControlModule.h
@@ -318,16 +318,19 @@ public:
 	 * @param InMaskingOperation	Masking operation to be performed.
 	 * @param ModifyOperationFlags	(optional) Flags that specify how the property is modified when the value is applied.
 	 */
-	virtual void PerformMasking(const TSharedRef<FRCMaskingOperation>& InMaskingOperation, const ERCModifyOperationFlags ModifyOperationFlags = ERCModifyOperationFlags::None) = 0;
+	UE_DEPRECATED(5.5, "PerformMasking is deprecated, masking is now applied where appropriate automatically.")
+	virtual void PerformMasking(const TSharedRef<FRCMaskingOperation>& InMaskingOperation) = 0;
 
 	/**
 	 * Register a masking factory to handle that masks the supported properties.
 	 */
+	UE_DEPRECATED(5.5, "RegisterMaskingFactoryForType is deprecated, masking is now applied where appropriate automatically.")
 	virtual void RegisterMaskingFactoryForType(UScriptStruct* RemoteControlPropertyType, const TSharedPtr<IRemoteControlMaskingFactory>& InMaskingFactory) = 0;
 
 	/**
 	 * Unregister a previously registered masking factory.
 	 */
+	UE_DEPRECATED(5.5, "UnregisterMaskingFactoryForType is deprecated, masking is now applied where appropriate automatically.")
 	virtual void UnregisterMaskingFactoryForType(UScriptStruct* RemoteControlPropertyType) = 0;
 
 	/**
diff --git a/Engine/Plugins/VirtualProduction/RemoteControl/Source/RemoteControl/Public/RemoteControlProtocolBinding.h b/Engine/Plugins/VirtualProduction/RemoteControl/Source/RemoteControl/Public/RemoteControlProtocolBinding.h
--- a/Engine/Plugins/VirtualProduction/RemoteControl/Source/RemoteControl/Public/RemoteControlProtocolBinding.h
+++ b/Engine/Plugins/VirtualProduction/RemoteControl/Source/RemoteControl/Public/RemoteControlProtocolBinding.h
@@ -115,6 +115,23 @@ public:
 		RefreshCachedData(PropertyType::StaticClass()->GetFName());
 	}
 
+	/** Returns true if one can get mapping value as primitive for the specified type */
+	template <typename ValueType>
+	bool CanGetMappingValueAsPrimitive() const
+	{
+		if constexpr (RemoteControlTypeTraits::TIsStringLikeValue<ValueType>::Value)
+		{
+			return TRemoteControlTypeTraits<ValueType>::IsSupportedMappingType();
+		}
+		else
+		{
+			return 
+				TRemoteControlTypeTraits<ValueType>::IsSupportedMappingType() &&
+				!InterpolationMappingPropertyData.IsEmpty() &&
+				InterpolationMappingPropertyData.Num() == sizeof(ValueType);
+		}
+	}
+
 	/** Get Mapping Property Value as a primitive type */
 	template <typename ValueType>
 	typename TEnableIf<!RemoteControlTypeTraits::TIsStringLikeValue<ValueType>::Value, ValueType>::Type
@@ -391,6 +408,7 @@ public:
 	 * @param ModifyOperationFlags			(optional) Flags that specify how the property is modified when the value is applied.
 	 * @return								True if the value was applied successfully
 	 */
+	UE_DEPRECATED(5.5, "ApplyProtocolValueToProperty is deprecated. Instead please refer to UE::RemoteControl::ProtocolEntityProcessor::ProcessEntities")
 	bool ApplyProtocolValueToProperty(const double InProtocolValue);
 
 	/** 
@@ -439,6 +457,9 @@ public:
 	/** Returns true if the given mask is enabled, false otherwise. */
 	virtual bool HasMask(ERCMask InMaskBit) const;
 
+	/** Returns the overriden masks. */
+	ERCMask GetOverridenMask() const { return OverridenMasks; }
+
 #if WITH_EDITOR
 	/** Retrieves the name of property corresponding to the given column name. */
 	const FName GetPropertyName(const FName& ForColumnName);
@@ -450,7 +471,6 @@ public:
 	/** Returns the preset that owns this protocol entity */
 	const TWeakObjectPtr<URemoteControlPreset>& GetOwner() const { return Owner; }
 
-public:
 	/** Container for range and mapping value pointers, and an optional number of elements (arrays, strings). */
 	struct FRangeMappingData
 	{
@@ -514,21 +534,25 @@ public:
 		}
 	};
 
+	/** Get Ranges and Mapping Value pointers */
+	TArray<FRangeMappingData> GetRangeMappingBuffers();
+
+	/** The current protocol value. Useful to test if the value changed */
+	double ProtocolValue = 0.0;
+
 private:
 	friend struct FRemoteControlProtocolBinding;
 
 	/**
 	 * Serialize interpolated property value to Cbor buffer
-	 * @param InProperty Property to apply serialization 
+	 * @param InProperty Property to apply serialization
 	 * @param InProtocolValue double value from the protocol
 	 * @param OutBuffer serialized buffer
 	 * @return true if serialized correctly
 	 */
+	UE_DEPRECATED(5.5, "Deprecated in favor of the overload that specifies the OwnerObject of the property.")
 	bool GetInterpolatedPropertyBuffer(FProperty* InProperty, double InProtocolValue, TArray<uint8>& OutBuffer);
 
-private:
-	/** Get Ranges and Mapping Value pointers */
-	TArray<FRangeMappingData> GetRangeMappingBuffers();
 
 protected:
 	/** The preset that owns this entity. */
diff --git a/Engine/Plugins/VirtualProduction/RemoteControl/Source/RemoteControl/Public/RemoteControlProtocolEntityProcessor.h b/Engine/Plugins/VirtualProduction/RemoteControl/Source/RemoteControl/Public/RemoteControlProtocolEntityProcessor.h
new file mode 100644
--- /dev/null
+++ b/Engine/Plugins/VirtualProduction/RemoteControl/Source/RemoteControl/Public/RemoteControlProtocolEntityProcessor.h
@@ -0,0 +1,22 @@
+// Copyright Epic Games, Inc. All Rights Reserved.
+
+#pragma once
+
+#include "Templates/SharedPointer.h"
+#include "UObject/StructOnScope.h"
+
+class FProperty;
+class UScriptStruct;
+struct FRemoteControlProtocolEntity;
+
+namespace UE::RemoteControl::ProtocolEntityProcessor
+{
+	/** Returns true if the script struct supports masking */
+	bool DoesScriptStructSupportMasking(const UScriptStruct* InStruct);
+
+	/** Returns true if the property supports masking. */
+	bool DoesPropertySupportMasking(const FProperty* Property);
+
+	/** Processes provided entities and their respective protocol values */
+	REMOTECONTROL_API void ProcessEntities(const TMap<TSharedPtr<TStructOnScope<FRemoteControlProtocolEntity>>, double>& EntityToValueMap);
+}
diff --git a/Engine/Plugins/VirtualProduction/RemoteControl/Source/RemoteControlProtocol/Private/RemoteControlProtocol.cpp b/Engine/Plugins/VirtualProduction/RemoteControl/Source/RemoteControlProtocol/Private/RemoteControlProtocol.cpp
--- a/Engine/Plugins/VirtualProduction/RemoteControl/Source/RemoteControlProtocol/Private/RemoteControlProtocol.cpp
+++ b/Engine/Plugins/VirtualProduction/RemoteControl/Source/RemoteControlProtocol/Private/RemoteControlProtocol.cpp
@@ -7,6 +7,7 @@
 #include "RCModifyOperationFlags.h"
 #include "RemoteControlPreset.h"
 #include "RemoteControlProtocolBinding.h"
+#include "RemoteControlProtocolEntityProcessor.h"
 #include "RemoteControlProtocolModule.h"
 #include "RemoteControlSettings.h"
 #include "UObject/StructOnScope.h"
@@ -14,11 +15,13 @@
 FRemoteControlProtocol::FRemoteControlProtocol(const FName InProtocolName)
 	: ProtocolName(InProtocolName)
 {
+	FCoreDelegates::OnBeginFrame.AddRaw(this, &FRemoteControlProtocol::OnBeginFrame);
 	FCoreDelegates::OnEndFrame.AddRaw(this, &FRemoteControlProtocol::OnEndFrame);
 }
 
 FRemoteControlProtocol::~FRemoteControlProtocol()
 {
+	FCoreDelegates::OnBeginFrame.RemoveAll(this);
 	FCoreDelegates::OnEndFrame.RemoveAll(this);
 }
 
@@ -45,61 +48,10 @@ void FRemoteControlProtocol::QueueValue(const FRemoteControlProtocolEntityPtr In
 	EntityValuesToApply.Add(InProtocolEntity, InProtocolValue);
 }
 
-void FRemoteControlProtocol::OnEndFrame()
+void FRemoteControlProtocol::OnBeginFrame()
 {
-	TArray<AActor*> ActorsWithoutPropertyChangedEvents;
-	for (const TPair<const FRemoteControlProtocolEntityPtr, double>& EntityValuesToApplyPair : EntityValuesToApply)
-	{
-		// Check is the Shared ptr and TStructOnScope is valid
-		if (EntityValuesToApplyPair.Key.IsValid() && EntityValuesToApplyPair.Key->IsValid())
-		{
-			FRemoteControlProtocolEntity* ProtocolEntity = EntityValuesToApplyPair.Key->Get();
-
-			const double ThisFrameValue = EntityValuesToApplyPair.Value;
-			const double* PreviousFrameValuePtr = PreviousTickValuesToApply.Find(EntityValuesToApplyPair.Key);
-
-			// Check the value from previous frame
-			if (PreviousFrameValuePtr == nullptr || !FMath::IsNearlyEqual(ThisFrameValue, *PreviousFrameValuePtr))
-			{
-				// Remember actors that will not get property change events
-				URemoteControlPreset* Preset = ProtocolEntity->GetOwner().Get();
-				if (Preset && EnumHasAnyFlags(Preset->GetModifyOperationFlagsForProtocols(), ERCModifyOperationFlags::SkipPropertyChangeEvents))
-				{
-					const TSharedPtr<FRemoteControlProperty> RemoteControlProperty = Preset->GetExposedEntity<FRemoteControlProperty>(ProtocolEntity->GetPropertyId()).Pin();
-					if (!RemoteControlProperty.IsValid())
-					{
-						continue;
-					}
-
-					for (UObject* Object : RemoteControlProperty->GetBoundObjects())
-					{
-						if (AActor* Actor = Object->GetTypedOuter<AActor>())
-						{
-							ActorsWithoutPropertyChangedEvents.AddUnique(Actor);
-						}
-					}
-				}
-
-				// Apply the protocol value
-				if (!ProtocolEntity->ApplyProtocolValueToProperty(EntityValuesToApplyPair.Value))
-				{
-					// Warn if the the value can't by applied
-					ensureMsgf(false, TEXT("Can't apply property for Protocol %s and PropertyId %s"),
-					           *ProtocolName.ToString(), *ProtocolEntity->GetPropertyId().ToString());
-				}
-			}
-		}
-	}
-
-	// Update actors that did not get property changed events
-	for (AActor* Actor : ActorsWithoutPropertyChangedEvents)
-	{
-		Actor->UpdateComponentTransforms();
-		Actor->MarkComponentsRenderStateDirty();
-	}
-
-	// Move the values from this frame to cached map
-	PreviousTickValuesToApply = MoveTemp(EntityValuesToApply);
+	using namespace UE::RemoteControl;
+	ProtocolEntityProcessor::ProcessEntities(EntityValuesToApply);
 }
 
 #if WITH_EDITOR
diff --git a/Engine/Plugins/VirtualProduction/RemoteControl/Source/RemoteControlProtocol/Private/Signature/RCSignatureProtocolAction.cpp b/Engine/Plugins/VirtualProduction/RemoteControl/Source/RemoteControlProtocol/Private/Signature/RCSignatureProtocolAction.cpp
--- a/Engine/Plugins/VirtualProduction/RemoteControl/Source/RemoteControlProtocol/Private/Signature/RCSignatureProtocolAction.cpp
+++ b/Engine/Plugins/VirtualProduction/RemoteControl/Source/RemoteControlProtocol/Private/Signature/RCSignatureProtocolAction.cpp
@@ -190,7 +190,8 @@ namespace UE::RemoteControlProtocol::Private
 		bool TryApply(FRemoteControlProtocolMapping& InOutMapping) const
 		{
 			InCppType PropertyValue;
-			if (!PropertyHandle->GetValue(PropertyValue))
+			if (!PropertyHandle->GetValue(PropertyValue) ||
+				!InOutMapping.CanGetMappingValueAsPrimitive<InCppType>())
 			{
 				return false;
 			}
diff --git a/Engine/Plugins/VirtualProduction/RemoteControl/Source/RemoteControlProtocol/Public/IRemoteControlProtocol.h b/Engine/Plugins/VirtualProduction/RemoteControl/Source/RemoteControlProtocol/Public/IRemoteControlProtocol.h
--- a/Engine/Plugins/VirtualProduction/RemoteControl/Source/RemoteControlProtocol/Public/IRemoteControlProtocol.h
+++ b/Engine/Plugins/VirtualProduction/RemoteControl/Source/RemoteControlProtocol/Public/IRemoteControlProtocol.h
@@ -112,6 +112,9 @@ public:
 	/** Unbind all protocol entities from the protocol */
 	virtual void UnbindAll() = 0;
 
+	/** Build to core engine delegated and called at the beginning of a frame */
+	virtual void OnBeginFrame() {};
+
 	/** Build to core engine delegated and called at the end of a frame */
 	virtual void OnEndFrame() {};
 
diff --git a/Engine/Plugins/VirtualProduction/RemoteControl/Source/RemoteControlProtocol/Public/RemoteControlProtocol.h b/Engine/Plugins/VirtualProduction/RemoteControl/Source/RemoteControlProtocol/Public/RemoteControlProtocol.h
--- a/Engine/Plugins/VirtualProduction/RemoteControl/Source/RemoteControlProtocol/Public/RemoteControlProtocol.h
+++ b/Engine/Plugins/VirtualProduction/RemoteControl/Source/RemoteControlProtocol/Public/RemoteControlProtocol.h
@@ -23,7 +23,7 @@ public:
 	virtual void Init() override;
 	virtual FRemoteControlProtocolEntityPtr CreateNewProtocolEntity(FProperty* InProperty, URemoteControlPreset* InOwner, FGuid InPropertyId) const override;
 	virtual void QueueValue(const FRemoteControlProtocolEntityPtr InProtocolEntity, const double InProtocolValue) override;
-	virtual void OnEndFrame() override;
+	virtual void OnBeginFrame() override;
 
 #if WITH_EDITOR
 
@@ -64,11 +64,5 @@ protected:
 
 private:
 	/** Map of the entities and protocol values about to apply */
-	TMap<const FRemoteControlProtocolEntityPtr, double> EntityValuesToApply;
-
-	/**
-	 * Map of the entities and protocol values from previous tick.
-	 * It allows skipping values that are very close to those of the previous frame
-	 */
-	TMap<const FRemoteControlProtocolEntityPtr, double> PreviousTickValuesToApply;
+	TMap<FRemoteControlProtocolEntityPtr, double> EntityValuesToApply;
 };
diff --git a/Engine/Plugins/VirtualProduction/RemoteControl/Source/RemoteControlProtocolWidgets/Private/Widgets/Masks/SRCProtocolMaskTriplet.cpp b/Engine/Plugins/VirtualProduction/RemoteControl/Source/RemoteControlProtocolWidgets/Private/Widgets/Masks/SRCProtocolMaskTriplet.cpp
--- a/Engine/Plugins/VirtualProduction/RemoteControl/Source/RemoteControlProtocolWidgets/Private/Widgets/Masks/SRCProtocolMaskTriplet.cpp
+++ b/Engine/Plugins/VirtualProduction/RemoteControl/Source/RemoteControlProtocolWidgets/Private/Widgets/Masks/SRCProtocolMaskTriplet.cpp
@@ -30,6 +30,7 @@ const TSet<UScriptStruct*>& FRemoteControlProtocolMasking::GetOptionalMaskStruct
 {
 	static TSet<UScriptStruct*> OptionalMaskStructs = {
 		TBaseStructure<FLinearColor>::Get(),
+		TBaseStructure<FColor>::Get(),
 		TBaseStructure<FVector4>::Get(),
 		TBaseStructure<FQuat>::Get(),
 		TBaseStructure<FIntVector4>::Get()
diff --git a/Engine/Plugins/VirtualProduction/RemoteControl/Source/RemoteControlProtocolWidgets/Private/Widgets/SRCProtocolBinding.cpp b/Engine/Plugins/VirtualProduction/RemoteControl/Source/RemoteControlProtocolWidgets/Private/Widgets/SRCProtocolBinding.cpp
--- a/Engine/Plugins/VirtualProduction/RemoteControl/Source/RemoteControlProtocolWidgets/Private/Widgets/SRCProtocolBinding.cpp
+++ b/Engine/Plugins/VirtualProduction/RemoteControl/Source/RemoteControlProtocolWidgets/Private/Widgets/SRCProtocolBinding.cpp
@@ -217,9 +217,8 @@ FReply SRCProtocolBinding::ToggleRecording() const
 {
 	// Binding can't be nullptr
 	FRemoteControlProtocolBinding* Binding = ViewModel->GetBinding();
-	check(Binding);
-
-	if (const TSharedPtr<TStructOnScope<FRemoteControlProtocolEntity>> ProtocolEntity = Binding->GetRemoteControlProtocolEntityPtr())
+	const TSharedPtr<TStructOnScope<FRemoteControlProtocolEntity>> ProtocolEntity = Binding ? Binding->GetRemoteControlProtocolEntityPtr() : nullptr;
+	if (ensureMsgf(ProtocolEntity && ProtocolEntity->IsValid(), TEXT("Unexpected invalid binding or entity, cannot toggle recording.")))
 	{
 		const ERCBindingStatus BindingStatus = (*ProtocolEntity)->ToggleBindingStatus();
 		if (BindingStatus == ERCBindingStatus::Awaiting)
@@ -241,11 +240,9 @@ FReply SRCProtocolBinding::ToggleRecording() const
 
 FSlateColor SRCProtocolBinding::GetRecordingButtonColor() const
 {
-	// Binding can't be nullptr
 	FRemoteControlProtocolBinding* Binding = ViewModel->GetBinding();
-	check(Binding);
-
-	if (const TSharedPtr<TStructOnScope<FRemoteControlProtocolEntity>> ProtocolEntity = Binding->GetRemoteControlProtocolEntityPtr())
+	const TSharedPtr<TStructOnScope<FRemoteControlProtocolEntity>> ProtocolEntity = Binding ? Binding->GetRemoteControlProtocolEntityPtr() : nullptr;
+	if (ProtocolEntity && ProtocolEntity->IsValid())
 	{
 		const ERCBindingStatus BindingStatus = (*ProtocolEntity)->GetBindingStatus();
 
@@ -262,7 +259,6 @@ FSlateColor SRCProtocolBinding::GetRecordingButtonColor() const
 		}
 	}
 
-	ensure(false);
 	return FLinearColor::Black;
 }
 
