Remote Control: Fix PRAGMA_DISABLE_DEPRECATION_WARNINGS macro appears a second time without matching PRAGMA_ENABLE_DEPRECATION_WARNINGS macro

#rnx

diff --git a/Engine/Plugins/VirtualProduction/RemoteControl/Source/RemoteControl/Private/RemoteControlModule.cpp b/Engine/Plugins/VirtualProduction/RemoteControl/Source/RemoteControl/Private/RemoteControlModule.cpp
--- a/Engine/Plugins/VirtualProduction/RemoteControl/Source/RemoteControl/Private/RemoteControlModule.cpp
+++ b/Engine/Plugins/VirtualProduction/RemoteControl/Source/RemoteControl/Private/RemoteControlModule.cpp
@@ -796,9 +796,7 @@ void FRemoteControlModule::StartupModule()
 	RegisterEntityFactory(FRemoteControlInstanceMaterial::StaticStruct()->GetFName(), FRemoteControlInstanceMaterialFactory::MakeInstance());
 
 	// DEPRECATED 5.5, here to keep support of the old implementation while it cannot be removed yet.
-	PRAGMA_DISABLE_DEPRECATION_WARNINGS
 	RegisterMaskingFactories();
-	PRAGMA_ENABLE_DEPRECATION_WARNINGS
 
 	// Register PropertyIdHandler
 	RegisterPropertyIdHandler();
@@ -834,9 +832,7 @@ void FRemoteControlModule::ShutdownModule()
 		DefaultValueFactories.Empty();
 
 		// DEPRECATED 5.5, here to keep support of the old implementation while it cannot be removed yet.
-		PRAGMA_DISABLE_DEPRECATION_WARNINGS
 		MaskingFactories.Empty();
-		PRAGMA_ENABLE_DEPRECATION_WARNINGS
 	}
 }
 
@@ -1036,7 +1032,6 @@ void FRemoteControlModule::ResetToDefaultValue(UObject* InObject, FRCResetToDefa
 
 void FRemoteControlModule::PerformMasking(const TSharedRef<FRCMaskingOperation>& InMaskingOperation)
 {
-	PRAGMA_DISABLE_DEPRECATION_WARNINGS
 	// DEPRECATED 5.5, here to keep support of the old implementation while it cannot be removed yet.
 	TRACE_CPUPROFILER_EVENT_SCOPE(FRemoteControlModule::PerformMasking);
 
@@ -1071,7 +1066,6 @@ void FRemoteControlModule::PerformMasking(const TSharedRef<FRCMaskingOperation>&
 			}
 		}
 	}
-	PRAGMA_ENABLE_DEPRECATION_WARNINGS
 }
 
 void FRemoteControlModule::RegisterMaskingFactoryForType(UScriptStruct* RemoteControlPropertyType, const TSharedPtr<IRemoteControlMaskingFactory>& InMaskingFactory)
@@ -2723,7 +2717,6 @@ void FRemoteControlModule::RegisterDefaultValueFactories()
 void FRemoteControlModule::RegisterMaskingFactories()
 {
 	// DEPRECATED 5.5, here to keep support of the old implementation while it cannot be removed yet.
-	PRAGMA_DISABLE_DEPRECATION_WARNINGS
 	RegisterMaskingFactoryForType(TBaseStructure<FVector>::Get(), FVectorMaskingFactory::MakeInstance());
 	RegisterMaskingFactoryForType(TBaseStructure<FVector4>::Get(), FVector4MaskingFactory::MakeInstance());
 	RegisterMaskingFactoryForType(TBaseStructure<FIntVector>::Get(), FIntVectorMaskingFactory::MakeInstance());
@@ -2731,7 +2724,6 @@ void FRemoteControlModule::RegisterMaskingFactories()
 	RegisterMaskingFactoryForType(TBaseStructure<FRotator>::Get(), FRotatorMaskingFactory::MakeInstance());
 	RegisterMaskingFactoryForType(TBaseStructure<FColor>::Get(), FColorMaskingFactory::MakeInstance());
 	RegisterMaskingFactoryForType(TBaseStructure<FLinearColor>::Get(), FLinearColorMaskingFactory::MakeInstance());
-	PRAGMA_ENABLE_DEPRECATION_WARNINGS
 }
 
 void FRemoteControlModule::RegisterPropertyIdHandler()
