Remote Control: Fix issue with users unable to delete Entity Groups

diff --git a/Engine/Plugins/VirtualProduction/RemoteControl/Source/RemoteControlUI/Private/UI/SRemoteControlPanel.cpp b/Engine/Plugins/VirtualProduction/RemoteControl/Source/RemoteControlUI/Private/UI/SRemoteControlPanel.cpp
--- a/Engine/Plugins/VirtualProduction/RemoteControl/Source/RemoteControlUI/Private/UI/SRemoteControlPanel.cpp
+++ b/Engine/Plugins/VirtualProduction/RemoteControl/Source/RemoteControlUI/Private/UI/SRemoteControlPanel.cpp
@@ -1625,6 +1625,11 @@ void SRemoteControlPanel::UpdateEntityDetailsView(const TSharedPtr<SRCPanelTreeN
 	}
 	else
 	{
+		if (EntityDetailsView)
+		{
+			EntityDetailsView->SetStructureData(nullptr);
+		}
+
 		WrappedEntityDetailsView->SetContent(CreateNoneSelectedWidget());
 	}
 
@@ -2138,12 +2143,31 @@ void SRemoteControlPanel::DeleteEntity_Execute()
 	{
 		FScopedTransaction Transaction(LOCTEXT("UnexposeFunction", "Unexpose remote control entity"));
 		Preset->Modify();
-		TArray<TSharedPtr<SRCPanelTreeNode>> SelectedEntities = EntityList->GetSelectedEntities();
-		for (int32 Index = 0; Index < SelectedEntities.Num(); ++Index)
+
+		TArray<TSharedPtr<SRCPanelTreeNode>> EntitiesToDelete = EntityList->GetSelectedEntities();
+
+		// Reverse the list as items will be processed from the back
+		Algo::Reverse(EntitiesToDelete);
+
+		while (!EntitiesToDelete.IsEmpty())
 		{
-			if (SelectedEntities[Index]->GetRCType() != SRCPanelTreeNode::FieldChild)
+			TSharedPtr<SRCPanelTreeNode> EntityToDelete = EntitiesToDelete.Pop();
+			if (!EntityToDelete.IsValid())
+			{
+				continue;
+			}
+
+			const SRCPanelTreeNode::ENodeType RCType = EntityToDelete->GetRCType();
+			if (RCType == SRCPanelTreeNode::FieldGroup)
+			{
+				// Deleting a Field Group would mean that all its children get deleted/unexposed.
+				TArray<TSharedPtr<SRCPanelTreeNode>> Children;
+				EntityToDelete->GetNodeChildren(Children);
+				EntitiesToDelete.Append(MoveTemp(Children));
+			}
+			else if (RCType != SRCPanelTreeNode::FieldChild)
 			{
-				Preset->Unexpose(SelectedEntities[Index]->GetRCId());
+				Preset->Unexpose(EntityToDelete->GetRCId());
 			}
 		}
 	}
