Remote Control: Protocol bindings no longer work after deleting a property or binding and undo the change

diff --git a/Engine/Plugins/VirtualProduction/RemoteControl/Source/RemoteControlProtocol/Private/RemoteControlProtocol.cpp b/Engine/Plugins/VirtualProduction/RemoteControl/Source/RemoteControlProtocol/Private/RemoteControlProtocol.cpp
--- a/Engine/Plugins/VirtualProduction/RemoteControl/Source/RemoteControlProtocol/Private/RemoteControlProtocol.cpp
+++ b/Engine/Plugins/VirtualProduction/RemoteControl/Source/RemoteControlProtocol/Private/RemoteControlProtocol.cpp
@@ -52,6 +52,8 @@ void FRemoteControlProtocol::OnBeginFrame()
 {
 	using namespace UE::RemoteControl;
 	ProtocolEntityProcessor::ProcessEntities(EntityValuesToApply);
+
+	EntityValuesToApply.Reset();
 }
 
 #if WITH_EDITOR
