Remote Control: Fix an issue where FProtocolRangeViewModel starts transacting when transactions are already ongoing

#rnx

diff --git a/Engine/Plugins/VirtualProduction/RemoteControl/Source/RemoteControlProtocolWidgets/Private/ViewModels/ProtocolRangeViewModel.cpp b/Engine/Plugins/VirtualProduction/RemoteControl/Source/RemoteControlProtocolWidgets/Private/ViewModels/ProtocolRangeViewModel.cpp
--- a/Engine/Plugins/VirtualProduction/RemoteControl/Source/RemoteControlProtocolWidgets/Private/ViewModels/ProtocolRangeViewModel.cpp
+++ b/Engine/Plugins/VirtualProduction/RemoteControl/Source/RemoteControlProtocolWidgets/Private/ViewModels/ProtocolRangeViewModel.cpp
@@ -191,7 +191,8 @@ void FProtocolRangeViewModel::SetOutputData(const FRCFieldResolvedData& InResolv
 
 	if (InputProxyPropertyContainer.IsValid() && OutputProxyPropertyContainer.IsValid())
 	{
-		FScopedTransaction Transaction(LOCTEXT("SetPresetProtocolOutputData", "Set Preset protocol binding output data"));
+		const bool bShouldActuallyTransact = !GIsTransacting;
+		FScopedTransaction Transaction(LOCTEXT("SetPresetProtocolOutputData", "Set Preset protocol binding output data"), bShouldActuallyTransact);
 		OutputProxyPropertyContainer.Get()->Modify();
 
 		void* SrcPropertyData = InResolvedData.Field->ContainerPtrToValuePtr<void>(InResolvedData.ContainerAddress);
