Remote Control:  Added getters for the controller container and property bag struct to allow RC Preset capabilities like being a data source for property binding systems

diff --git a/Engine/Plugins/VirtualProduction/RemoteControl/Source/RemoteControl/Private/RCVirtualPropertyContainer.cpp b/Engine/Plugins/VirtualProduction/RemoteControl/Source/RemoteControl/Private/RCVirtualPropertyContainer.cpp
--- a/Engine/Plugins/VirtualProduction/RemoteControl/Source/RemoteControl/Private/RCVirtualPropertyContainer.cpp
+++ b/Engine/Plugins/VirtualProduction/RemoteControl/Source/RemoteControl/Private/RCVirtualPropertyContainer.cpp
@@ -262,6 +262,16 @@ int32 URCVirtualPropertyContainerBase::GetNumVirtualProperties() const
 	return NumPropertiesInBag;
 }
 
+const UPropertyBag* URCVirtualPropertyContainerBase::GetPropertyBagStruct() const
+{
+	return Bag.GetPropertyBagStruct();
+}
+
+FStructView URCVirtualPropertyContainerBase::GetPropertyBagMutableValue()
+{
+	return Bag.GetMutableValue();
+}
+
 TSharedPtr<FStructOnScope> URCVirtualPropertyContainerBase::CreateStructOnScope()
 {
 	return MakeShared<FStructOnScope>(Bag.GetPropertyBagStruct(), Bag.GetMutableValue().GetMemory());
diff --git a/Engine/Plugins/VirtualProduction/RemoteControl/Source/RemoteControl/Public/RCVirtualPropertyContainer.h b/Engine/Plugins/VirtualProduction/RemoteControl/Source/RemoteControl/Public/RCVirtualPropertyContainer.h
--- a/Engine/Plugins/VirtualProduction/RemoteControl/Source/RemoteControl/Public/RCVirtualPropertyContainer.h
+++ b/Engine/Plugins/VirtualProduction/RemoteControl/Source/RemoteControl/Public/RCVirtualPropertyContainer.h
@@ -117,6 +117,16 @@ public:
 	 */ 
 	int32 GetNumVirtualProperties() const;
 
+	/**
+	 * Returns the Property Bag Struct of this Container
+	 */
+	const UPropertyBag* GetPropertyBagStruct() const;
+
+	/**
+	 * Creates new Struct View for this Property Bag UStruct and Memory
+	 */
+	FStructView GetPropertyBagMutableValue();
+
 	/**
 	 * Creates new Struct on Scope for this Property Bag UStruct and Memory
 	 */
diff --git a/Engine/Plugins/VirtualProduction/RemoteControl/Source/RemoteControl/Public/RemoteControlPreset.h b/Engine/Plugins/VirtualProduction/RemoteControl/Source/RemoteControl/Public/RemoteControlPreset.h
--- a/Engine/Plugins/VirtualProduction/RemoteControl/Source/RemoteControl/Public/RemoteControlPreset.h
+++ b/Engine/Plugins/VirtualProduction/RemoteControl/Source/RemoteControl/Public/RemoteControlPreset.h
@@ -790,6 +790,12 @@ public:
 	* Currently used by the UI class SRCControllerPanelList for user value entry via the RC Controllers panel*/
 	TSharedPtr<FStructOnScope> GetControllerContainerStructOnScope();
 
+	/** Returns the Property Bag Struct of this Container */
+	URCVirtualPropertyContainerBase* GetControllerContainer() const
+	{
+		return ControllerContainer;
+	}
+
 	/** Sets the Controller Container (holds all Virtual Properties) */
 	void SetControllerContainer(URCVirtualPropertyContainerBase* InControllerContainer);
 
