Remote Control: Added support to expose properties within instanced structs.

diff --git a/Engine/Plugins/VirtualProduction/RemoteControl/Source/RemoteControl/Private/RemoteControlFieldPath.cpp b/Engine/Plugins/VirtualProduction/RemoteControl/Source/RemoteControl/Private/RemoteControlFieldPath.cpp
--- a/Engine/Plugins/VirtualProduction/RemoteControl/Source/RemoteControl/Private/RemoteControlFieldPath.cpp
+++ b/Engine/Plugins/VirtualProduction/RemoteControl/Source/RemoteControl/Private/RemoteControlFieldPath.cpp
@@ -1,11 +1,20 @@
 // Copyright Epic Games, Inc. All Rights Reserved.
 
 #include "RemoteControlFieldPath.h"
+#include "StructUtils/InstancedStruct.h"
 
 namespace RemoteControlFieldUtils
 {
 	void ResolveSegment(const FRCFieldPathSegment& Segment, UStruct* Owner, void* ContainerAddress, FRCFieldResolvedData& OutResolvedData)
 	{
+		// When resolving for Instanced Struct, use the underlying Struct/Memory
+		if (Owner == TBaseStructure<FInstancedStruct>::Get())
+		{
+			FInstancedStruct& InstancedStruct = *static_cast<FInstancedStruct*>(ContainerAddress);
+			Owner = const_cast<UScriptStruct*>(InstancedStruct.GetScriptStruct());
+			ContainerAddress = InstancedStruct.GetMutableMemory();
+		}
+
 		if (FProperty* FoundField = FindFProperty<FProperty>(Owner, Segment.Name))
 		{
 			OutResolvedData.ContainerAddress = ContainerAddress;
@@ -31,6 +40,19 @@ namespace RemoteControlFieldUtils
 
 bool FRCFieldPathInfo::ResolveInternalRecursive(UStruct* OwnerType, void* ContainerAddress, int32 SegmentIndex)
 {
+	// NOTE: There are other struct types like FSharedStruct, FInstancedStructContainer that could need support here in a similar way.
+	if (OwnerType == TBaseStructure<FInstancedStruct>::Get() && Segments[SegmentIndex].Name == TEXT("Struct"))
+	{
+		// Instanced Structs generate an additional "Struct." segment from property handles path generation, as they use FStructurePropertyNode
+		// This "Struct." segment is not a property to resolve so needs to be skipped.
+		// see FStructurePropertyNode::GetQualifiedName.
+		++SegmentIndex;
+		if (!ensureMsgf(SegmentIndex < Segments.Num(), TEXT("Instanced Structs expect more segments after 'Struct'. Index: %d, Num: %d"), SegmentIndex, Segments.Num()))
+		{
+			return false;
+		}
+	}
+
 	const bool bLastSegment = (SegmentIndex == Segments.Num() - 1);
 
 	//Resolve the desired segment
