[Remote Control] Improving extension column size management in detail panel to avoid visually large gap.

The width was hardcoded to 100. This way of doing it is not compatible with the different cases and leave a visual gap in the column in all configs anyway.
A better method is to do it like SDetailsViewBase::SetKeyframeHandler. incrementing or decrementing width relative to the current width.

diff --git a/Engine/Plugins/VirtualProduction/RemoteControl/Source/RemoteControlUI/Private/RemoteControlUIModule.cpp b/Engine/Plugins/VirtualProduction/RemoteControl/Source/RemoteControlUI/Private/RemoteControlUIModule.cpp
--- a/Engine/Plugins/VirtualProduction/RemoteControl/Source/RemoteControlUI/Private/RemoteControlUIModule.cpp
+++ b/Engine/Plugins/VirtualProduction/RemoteControl/Source/RemoteControlUI/Private/RemoteControlUIModule.cpp
@@ -260,6 +260,13 @@ namespace RemoteControlUIModule
 
 		return false;
 	}
+
+	void UpdateDetailViewExtensionWidth(const TSharedPtr<IDetailsView>& InDetailsView, bool bInOnOpen)
+	{		
+		const float DeltaWidth = bInOnOpen ? 22.0f : -22.0f;
+		InDetailsView->SetRightColumnMinWidth(InDetailsView->GetRightColumnMinWidth() + DeltaWidth);
+		InDetailsView->ForceRefresh();
+	}
 }
 
 void FRemoteControlUIModule::StartupModule()
@@ -360,16 +367,14 @@ TSharedRef<SRemoteControlPanel> FRemoteControlUIModule::CreateRemoteControlPanel
 
 			if (SharedDetailsPanel.IsValid())
 			{
-				SharedDetailsPanel->SetRightColumnMinWidth(100);
-				SharedDetailsPanel->ForceRefresh();
+				RemoteControlUIModule::UpdateDetailViewExtensionWidth(SharedDetailsPanel, /*bInOnOpen*/ true);
 				break;
 			}
 		}
 	}
 	else
 	{
-		SharedDetailsPanel->SetRightColumnMinWidth(100);
-		SharedDetailsPanel->ForceRefresh();
+		RemoteControlUIModule::UpdateDetailViewExtensionWidth(SharedDetailsPanel, /*bInOnOpen*/ true);
 	}
 
 	OnRemoteControlPresetOpened().Broadcast(Preset);
@@ -1656,8 +1661,7 @@ void FRemoteControlUIModule::UnregisterRemoteControlPanel(SRemoteControlPanel* P
 	if (SharedDetailsPanel.IsValid())
 	{
 		// reset column min width to standard size
-		SharedDetailsPanel->SetRightColumnMinWidth(22);
-		SharedDetailsPanel->ForceRefresh();
+		RemoteControlUIModule::UpdateDetailViewExtensionWidth(SharedDetailsPanel, /*bInOnOpen*/ false);
 	}
 
 	if (!Panel)
