[Remote Control] Fix Editor only warnings to properly reflect packaged build.

diff --git a/Engine/Plugins/VirtualProduction/RemoteControl/Source/RemoteControl/Private/RemoteControlField.cpp b/Engine/Plugins/VirtualProduction/RemoteControl/Source/RemoteControl/Private/RemoteControlField.cpp
--- a/Engine/Plugins/VirtualProduction/RemoteControl/Source/RemoteControl/Private/RemoteControlField.cpp
+++ b/Engine/Plugins/VirtualProduction/RemoteControl/Source/RemoteControl/Private/RemoteControlField.cpp
@@ -328,7 +328,8 @@ void FRemoteControlProperty::EnableEditCondition()
 
 bool FRemoteControlProperty::IsEditableInPackaged(FString* OutError) const
 {
-	return bIsEditableInPackaged || IRemoteControlModule::Get().PropertySupportsRawModification(GetProperty(), GetBoundObject(), false, OutError);
+	// Inaccurate: bIsEditableInPackaged = !Property->HasAnyPropertyFlags(CPF_BlueprintReadOnly) -> leads to false positives
+	return IRemoteControlModule::Get().PropertySupportsRawModification(GetProperty(), GetBoundObject(), false, OutError);
 }
 
 bool FRemoteControlProperty::IsEditableInEditor(FString* OutError) const
diff --git a/Engine/Plugins/VirtualProduction/RemoteControl/Source/RemoteControl/Private/RemoteControlModule.cpp b/Engine/Plugins/VirtualProduction/RemoteControl/Source/RemoteControl/Private/RemoteControlModule.cpp
--- a/Engine/Plugins/VirtualProduction/RemoteControl/Source/RemoteControl/Private/RemoteControlModule.cpp
+++ b/Engine/Plugins/VirtualProduction/RemoteControl/Source/RemoteControl/Private/RemoteControlModule.cpp
@@ -126,6 +126,132 @@ namespace RemoteControlUtil
 		return Access == ERCAccess::READ_ACCESS;
 	}
 
+#if WITH_EDITOR
+	/**
+	 * For an editor build to restrict access only to what is available in packaged game.
+	 */
+	static bool GEmulatePackagedAccess = false;
+
+	bool IsEmulatePackaged()
+	{
+		return GEmulatePackagedAccess;
+	}
+#endif
+
+	// Todo: Temporary fork of RemoteControlPropertyUtilities::FindSetterFunctionInternal for hotfix. Merge in main.
+	static UFunction* FindSetterFunctionInternal_TMP(const FProperty* Property, UClass* OwnerClass)
+	{
+		using namespace RemoteControlPropertyUtilities;
+
+		// Check if the property setter is already cached.
+		const TWeakObjectPtr<UFunction> SetterPtr = CachedSetterFunctions.FindRef(Property);
+		if (SetterPtr.IsValid())
+		{
+			return SetterPtr.Get();
+		}
+
+		static const FName RelativeLocationPropertyName = USceneComponent::GetRelativeLocationPropertyName();
+		static const FName RelativeRotationPropertyName = USceneComponent::GetRelativeRotationPropertyName();
+		static const FName RelativeScalePropertyName = USceneComponent::GetRelativeScale3DPropertyName();
+
+		static TMap<FName, FName> CommonPropertiesToFunctions = { 
+			{ RelativeLocationPropertyName, "K2_SetRelativeLocation"},
+			{ RelativeRotationPropertyName, "K2_SetRelativeRotation" },
+			{ RelativeScalePropertyName, "SetRelativeScale3D" }
+		};
+
+		FString SetterName;
+		UFunction* SetterFunction = nullptr;
+		if(OwnerClass->IsChildOf<ULightComponent>())
+		{
+			const FName FoundSetterName = FindLightSetterFunctionInternal(Property, OwnerClass);
+			if(!FoundSetterName.IsNone())
+			{
+				SetterName = FoundSetterName.ToString();				
+			}
+		}
+#if WITH_EDITOR
+		if (!IsEmulatePackaged()) // -- Added for hotfix
+		{
+			if(SetterName.IsEmpty())
+			{
+				SetterName = Property->GetMetaData(*NAME_BlueprintSetter.ToString());
+			}
+			if (!SetterName.IsEmpty())
+			{
+				SetterFunction = OwnerClass->FindFunctionByName(*SetterName);
+			}
+			else
+			{
+				if (CommonPropertiesToFunctions.Contains(Property->GetFName()))
+				{
+					SetterFunction = OwnerClass->FindFunctionByName(*CommonPropertiesToFunctions[Property->GetFName()].ToString());
+				}
+			}
+		}
+#endif
+
+		if(!SetterFunction)
+		{
+			FString PropertyName = Property->GetName();
+			if (Property->IsA<FBoolProperty>())
+			{
+				PropertyName.RemoveFromStart("b", ESearchCase::CaseSensitive);
+			}
+
+			static const TArray<FString> SetterPrefixes = {
+				FString("Set"),
+				FString("K2_Set"),
+				FString("BP_Set")
+			};
+
+			for (const FString& Prefix : SetterPrefixes)
+			{
+				const FName SetterFunctionName = FName(Prefix + PropertyName);
+				SetterFunction = OwnerClass->FindFunctionByName(SetterFunctionName);
+				if (SetterFunction)
+				{
+					break;
+				}
+			}
+		}
+
+		if (SetterFunction && FindSetterArgument(SetterFunction, Property))
+		{
+			CachedSetterFunctions.Add(Property, SetterFunction);
+		}
+		else
+		{
+			// Arguments are not compatible so don't use this setter.
+			SetterFunction = nullptr;
+		}
+
+		return SetterFunction;
+	}
+
+	// Todo: temporary fork of RemoteControlPropertyUtilities::FindSetterFunction for hotfix. Merge in main.
+	static UFunction* FindSetterFunction_TMP(const FProperty* Property, UClass* OwnerClass)
+	{
+		// UStruct properties cannot have setters.
+		if (!ensure(Property) || !Property->GetOwnerClass())
+		{
+			return nullptr;
+		}
+
+		UFunction* SetterFunction = nullptr;
+		if (OwnerClass && OwnerClass != Property->GetOwnerClass())
+		{
+			SetterFunction = FindSetterFunctionInternal_TMP(Property, OwnerClass);
+		}
+		
+		if (!SetterFunction)
+		{
+			SetterFunction = FindSetterFunctionInternal_TMP(Property, Property->GetOwnerClass());
+		}
+		
+		return SetterFunction;
+	}
+
 	bool PropertyModificationShouldUseSetter(const UObject* Object, FProperty* Property)
 	{
 		if (!Property || !Object)
@@ -133,50 +259,61 @@ namespace RemoteControlUtil
 			return false;
 		}
 
-		return Property->HasSetter() || !!RemoteControlPropertyUtilities::FindSetterFunction(Property, Object->GetClass());
+		return Property->HasSetter() || !!FindSetterFunction_TMP(Property, Object->GetClass());
 	}
 
 	bool IsPropertyAllowed(const FProperty* InProperty, ERCAccess InAccessType, const UObject* InObject, bool bObjectInGamePackage, FString* OutError)
 	{
 		FTextBuilder Error;
-		FText PreError;
-		if (IsWriteAccess(InAccessType))
+		ON_SCOPE_EXIT
 		{
-			PreError = LOCTEXT("RCWritingAccess", "Property is not writeable because");
-		}
-		else if (IsReadAccess(InAccessType))
+			if (OutError)
+			{
+				*OutError = Error.ToText().ToString();
+			}
+		};
+
+		auto PreError = [InAccessType]()-> FText
 		{
-			PreError = LOCTEXT("RCReadingAccess", "Property is not readable because");
-		}
-		else
+			if (IsWriteAccess(InAccessType))
+			{
+				return LOCTEXT("RCWritingAccess", "Property is not writeable because");
+			}
+			if (IsReadAccess(InAccessType))
+			{
+				return LOCTEXT("RCReadingAccess", "Property is not readable because");
+			}
+			return LOCTEXT("RCNoAccess", "Property");
+		};
+
+#if WITH_EDITOR
+		const URemoteControlSettings* RCSettings = GetMutableDefault<URemoteControlSettings>();
+		
+		if (!IsEmulatePackaged())
 		{
-			PreError = LOCTEXT("RCNoAccess", "Property");
+			// Override this flag to false if we are running in Editor Mode with -game flag.
+			bObjectInGamePackage = false;
 		}
-#if WITH_EDITOR
-		URemoteControlSettings* RCSettings = GetMutableDefault<URemoteControlSettings>();
-		// Override this flag to false if we are running in Editor Mode with -game flag.
-		bObjectInGamePackage = false;
 #endif
+
 		// The property is allowed to be accessed if it exists...
-		const bool bPropertyIsValid = InProperty != nullptr;
-		if (!bPropertyIsValid)
+		if (!InProperty)
 		{
 			if (OutError)
 			{
-				Error.AppendLine(FText::Format(LOCTEXT("RCPropertyNotValid", "{0} is not valid."), PreError));
-				*OutError = Error.ToText().ToString();
+				Error.AppendLine(FText::Format(LOCTEXT("RCPropertyNotValid", "{0} it is not valid."), PreError()));
 			}
-
 			return false;
 		}
+
 #if WITH_EDITOR
 		// it doesn't have exposed getter AND setter that should be used instead or if it should ignore them
 		const bool bHasBlueprintGetterOrSetterOrIgnored = !InProperty->HasMetaData(RemoteControlUtil::NAME_BlueprintGetter) ||
 												 		  !InProperty->HasMetaData(RemoteControlUtil::NAME_BlueprintSetter) ||
 												 		  RCSettings->bIgnoreGetterSetterCheck;
-		if (!bHasBlueprintGetterOrSetterOrIgnored)
+		if (!bHasBlueprintGetterOrSetterOrIgnored && OutError)
 		{
-			Error.AppendLine(FText::Format(LOCTEXT("RCBlueprintHasGetterSetterOrNotIgnored", "{0} has blueprint getter and setter to use.\nEnable in the RemoteControl settings IgnoreGetterSetterCheck if you want to ignore it."), PreError));
+			Error.AppendLine(FText::Format(LOCTEXT("RCBlueprintHasGetterSetterOrNotIgnored", "{0} it has blueprint getter and setter to use.\nEnable in the RemoteControl settings IgnoreGetterSetterCheck if you want to ignore it."), PreError()));
 		}
 
 		// it isn't private or protected, except if AllowPrivateAccess is true and if only for Protected whenever it should be ignored
@@ -186,53 +323,66 @@ namespace RemoteControlUtil
 		const bool bAllowPrivateAccessForProtected = !InProperty->HasAnyPropertyFlags(CPF_NativeAccessSpecifierProtected) ||
 													 RCSettings->bIgnoreProtectedCheck ||
 													 InProperty->GetBoolMetaData(RemoteControlUtil::NAME_AllowPrivateAccess);
-
+		
+		const bool bIsAccessibleEditor = !IsEmulatePackaged() ?
+			bHasBlueprintGetterOrSetterOrIgnored &&
+			bAllowPrivateAccessForPrivate &&
+			bAllowPrivateAccessForProtected : true;
+		
 #endif
 
-		const bool bPropertyHasSetterToUse = IsWriteAccess(InAccessType)? PropertyModificationShouldUseSetter(InObject, const_cast<FProperty*>(InProperty)) : true;
-
-		const bool bPropertyHasGetterToUse = IsReadAccess(InAccessType)? InProperty->HasGetter() : true;
+		const bool bPropertyHasSetterToUse = IsWriteAccess(InAccessType) ? PropertyModificationShouldUseSetter(InObject, const_cast<FProperty*>(InProperty)) : true;
+		const bool bPropertyHasGetterToUse = IsReadAccess(InAccessType) ? InProperty->HasGetter() : true;
 
 		bool bAppendSettingInformation = false;
-		if (!bPropertyHasSetterToUse)
+		if (!bPropertyHasSetterToUse && OutError)
 		{
 #if WITH_EDITOR
-			if (!bAllowPrivateAccessForPrivate)
+			if (!IsEmulatePackaged())
 			{
-				Error.AppendLine(FText::Format(LOCTEXT("RCPropertyPrivateWithoutSetter", "{0} is private without the AllowPrivateAccess specifier set to true or a setter assigned."), PreError));
-			}
+				if (!bAllowPrivateAccessForPrivate)
+				{
+					Error.AppendLine(FText::Format(LOCTEXT("RCPropertyPrivateWithoutSetter", "{0} it is private without the AllowPrivateAccess specifier set to true or a setter assigned."), PreError()));
+				}
 
-			if (!bAllowPrivateAccessForProtected)
-			{
-				Error.AppendLine(FText::Format(LOCTEXT("RCPropertyProtectedWithoutSetterOrIgnored", "{0} is protected without the AllowPrivateAccess specifier set to true or a setter assigned."), PreError));
-				bAppendSettingInformation = true;
+				if (!bAllowPrivateAccessForProtected)
+				{
+					Error.AppendLine(FText::Format(LOCTEXT("RCPropertyProtectedWithoutSetterOrIgnored", "{0} it is protected without the AllowPrivateAccess specifier set to true or a setter assigned."), PreError()));
+					bAppendSettingInformation = true;
+				}
 			}
-#else
-			Error.AppendLine(FText::Format(LOCTEXT("RCPropertyNoSetter", "{0} it doesn't have a setter specifier assigned to it."), PreError));
+			else
 #endif
-			
+			{
+				Error.AppendLine(FText::Format(LOCTEXT("RCPropertyNoSetter", "{0} it doesn't have a setter specifier assigned to it."), PreError()));
+			}
 		}
 
-		if (!bPropertyHasGetterToUse)
+		if (!bPropertyHasGetterToUse && OutError)
 		{
 #if WITH_EDITOR
-			if (!bAllowPrivateAccessForPrivate)
+			if (!IsEmulatePackaged())
 			{
-				Error.AppendLine(FText::Format(LOCTEXT("RCPropertyPrivateWithoutGetter", "{0} is private without the AllowPrivateAccess specifier set to true or a getter assigned."), PreError));
-			}
+				if (!bAllowPrivateAccessForPrivate)
+				{
+					Error.AppendLine(FText::Format(LOCTEXT("RCPropertyPrivateWithoutGetter", "{0} it is private without the AllowPrivateAccess specifier set to true or a getter assigned."), PreError()));
+				}
 
-			if (!bAllowPrivateAccessForProtected)
-			{
-				Error.AppendLine(FText::Format(LOCTEXT("RCPropertyProtectedWithoutGetterOrIgnored", "{0} is protected without the AllowPrivateAccess specifier set to true or a getter assigned."), PreError));
-				bAppendSettingInformation = true;
+				if (!bAllowPrivateAccessForProtected)
+				{
+					Error.AppendLine(FText::Format(LOCTEXT("RCPropertyProtectedWithoutGetterOrIgnored", "{0} it is protected without the AllowPrivateAccess specifier set to true or a getter assigned."), PreError()));
+					bAppendSettingInformation = true;
+				}
 			}
-#else
-			Error.AppendLine(FText::Format(LOCTEXT("RCPropertyNoGetter", "{0} doesn't have a getter specifier assigned to it."), PreError));
+			else
 #endif
+			{
+				Error.AppendLine(FText::Format(LOCTEXT("RCPropertyNoGetter", "{0} it doesn't have a getter specifier assigned to it."), PreError()));
+			}
 		}
 
 #if WITH_EDITOR
-		if (bAppendSettingInformation)
+		if (bAppendSettingInformation && OutError)
 		{
 			Error.AppendLine(LOCTEXT("RCSettingsIgnoreProtected","You can enable the RemoteControl setting IgnoreProtectedCheck to ignore it for protected properties."));
 		}
@@ -240,60 +390,44 @@ namespace RemoteControlUtil
 		// it isn't blueprint private
 		const bool bIsBlueprintPublic = !InProperty->HasAnyPropertyFlags(CPF_DisableEditOnInstance);
 
-		if (!bIsBlueprintPublic)
+		if (!bIsBlueprintPublic && OutError)
 		{
-			Error.AppendLine(FText::Format(LOCTEXT("RCBlueprintPropertyPrivate", "{0} has flag DisableEditOnInstance.\nEnable the InstanceEditable flag of the property in the blueprint editor."), PreError));
+			Error.AppendLine(FText::Format(LOCTEXT("RCBlueprintPropertyPrivate", "{0} it has flag DisableEditOnInstance.\nEnable the InstanceEditable flag of the property in the blueprint editor."), PreError()));
 		}
 
 		// and it's either blueprint visible if in game or editable if in editor and it isn't read only if the access type is write
-		bool bIsAccessibleInGamePackage;
-		bool bIsAccessibleInNonGamePackage;
+		bool bIsAccessible;
 
 		if (bObjectInGamePackage)
 		{
-			bIsAccessibleInGamePackage = InProperty->HasAnyPropertyFlags(CPF_BlueprintVisible) &&
-										 (InAccessType == ERCAccess::READ_ACCESS ||
-										 !InProperty->HasAnyPropertyFlags(CPF_BlueprintReadOnly));
+			bIsAccessible = InProperty->HasAnyPropertyFlags(CPF_BlueprintVisible) &&
+							 (InAccessType == ERCAccess::READ_ACCESS ||
+							 !InProperty->HasAnyPropertyFlags(CPF_BlueprintReadOnly));
 
-			if (!bIsAccessibleInGamePackage)
+			if (!bIsAccessible && OutError)
 			{
-				Error.AppendLine(FText::Format(LOCTEXT("RCBlueprintNotAccessibleInGamePackage", "{0} doesn't have the BlueprintVisible flag and/or you are trying to write to it but it has the BlueprintReadOnly flag."), PreError));
+				Error.AppendLine(FText::Format(LOCTEXT("RCBlueprintNotAccessibleInGamePackage", "{0} it doesn't have the BlueprintVisible flag and/or you are trying to write to it but it has the BlueprintReadOnly flag."), PreError()));
 			}
-			// set this to true since it will not be evaluated in GamePackage
-			bIsAccessibleInNonGamePackage = true;
 		}
 		else
 		{
-			bIsAccessibleInNonGamePackage = InAccessType == ERCAccess::READ_ACCESS ||
-											(InProperty->HasAnyPropertyFlags(CPF_Edit) &&
-											!InProperty->HasAnyPropertyFlags(CPF_EditConst));
+			bIsAccessible = InAccessType == ERCAccess::READ_ACCESS ||
+							(InProperty->HasAnyPropertyFlags(CPF_Edit) &&
+							!InProperty->HasAnyPropertyFlags(CPF_EditConst));
 
-			if (!bIsAccessibleInNonGamePackage)
+			if (!bIsAccessible && OutError)
 			{
-				Error.AppendLine(FText::Format(LOCTEXT("RCBlueprintNotAccessibleInNonGamePackage", "{0} has the EditConst flag."), PreError));
+				Error.AppendLine(FText::Format(LOCTEXT("RCBlueprintNotAccessibleInNonGamePackage", "{0} it has the EditConst flag."), PreError()));
 			}
-			// set this to true since it will not be evaluated in NON GamePackage
-			bIsAccessibleInGamePackage = true;
 		}
-
-		const bool bIsAccessible = bObjectInGamePackage? bIsAccessibleInGamePackage : bIsAccessibleInNonGamePackage;
-
-		// Assign the error to be propagated after this call in case it failed
-		if (OutError)
-		{
-			*OutError = Error.ToText().ToString();
-		}
-
-		return (bPropertyIsValid &&
+		
+		return (bIsAccessible &&
 #if WITH_EDITOR
-			bHasBlueprintGetterOrSetterOrIgnored &&
-			bAllowPrivateAccessForPrivate &&
-			bAllowPrivateAccessForProtected &&
+			bIsAccessibleEditor &&
 #endif
-			bIsBlueprintPublic &&
-			bIsAccessible) ||
+			bIsBlueprintPublic) ||
 			(bPropertyHasGetterToUse && bPropertyHasSetterToUse);
-	};
+	}
 
 	FARFilter GetBasePresetFilter()
 	{
@@ -2276,6 +2410,9 @@ void FRemoteControlModule::UnregisterDefaultEntityMetadata(FName MetadataKey)
 bool FRemoteControlModule::PropertySupportsRawModificationWithoutEditor(FProperty* Property, UClass* OwnerClass) const
 {
 	constexpr bool bInGameOrPackage = true;
+#if WITH_EDITOR
+	TGuardValue EmulatePackagedGuard(RemoteControlUtil::GEmulatePackagedAccess, bInGameOrPackage);
+#endif
 	return Property && (RemoteControlUtil::IsPropertyAllowed(Property, ERCAccess::WRITE_ACCESS, nullptr, bInGameOrPackage, nullptr) || !!RemoteControlPropertyUtilities::FindSetterFunction(Property, OwnerClass));
 }
 
@@ -2286,9 +2423,16 @@ bool FRemoteControlModule::PropertySupportsRawModification(FProperty* InProperty
 	const bool bCheckPackaged = !bInWithEditor;
 	if (!InProperty)
 	{
+		if (OutError)
+		{
+			*OutError = TEXT("Property is not bound");
+		}
 		return false;
 	}
 
+#if WITH_EDITOR
+	TGuardValue EmulatePackagedGuard(RemoteControlUtil::GEmulatePackagedAccess, bCheckPackaged);
+#endif
 	const bool bIsPropertyAllowedToSet = RemoteControlUtil::IsPropertyAllowed(InProperty, ERCAccess::WRITE_ACCESS, InObject, bCheckPackaged, &ErrorTextSet);
 	const bool bIsPropertyAllowedToRead = RemoteControlUtil::IsPropertyAllowed(InProperty, ERCAccess::READ_ACCESS, InObject, bCheckPackaged, &ErrorTextGet);
 
