Fix ensure fail when using remote control with -game

#11258

diff --git a/Engine/Plugins/VirtualProduction/RemoteControl/Source/RemoteControlUI/Private/RemoteControlUIModule.cpp b/Engine/Plugins/VirtualProduction/RemoteControl/Source/RemoteControlUI/Private/RemoteControlUIModule.cpp
--- a/Engine/Plugins/VirtualProduction/RemoteControl/Source/RemoteControlUI/Private/RemoteControlUIModule.cpp
+++ b/Engine/Plugins/VirtualProduction/RemoteControl/Source/RemoteControlUI/Private/RemoteControlUIModule.cpp
@@ -384,6 +384,11 @@ TSharedRef<SRemoteControlPanel> FRemoteControlUIModule::CreateRemoteControlPanel
 
 void FRemoteControlUIModule::RegisterContextMenuExtender()
 {
+	if (IsRunningGame())
+	{
+		return;
+	}
+
 	// Extend the level viewport context menu to add an option to copy the object path.
 	LevelViewportContextMenuRemoteControlExtender = FLevelEditorModule::FLevelViewportMenuExtender_SelectedActors::CreateRaw(this, &FRemoteControlUIModule::ExtendLevelViewportContextMenuForRemoteControl);
 	FLevelEditorModule& LevelEditorModule = FModuleManager::Get().LoadModuleChecked<FLevelEditorModule>("LevelEditor");
