Remote Control - Allow specifying charset on http requests

#11059

diff --git a/Engine/Plugins/VirtualProduction/RemoteControl/Source/WebRemoteControl/Private/WebRemoteControlInternalUtils.cpp b/Engine/Plugins/VirtualProduction/RemoteControl/Source/WebRemoteControl/Private/WebRemoteControlInternalUtils.cpp
--- a/Engine/Plugins/VirtualProduction/RemoteControl/Source/WebRemoteControl/Private/WebRemoteControlInternalUtils.cpp
+++ b/Engine/Plugins/VirtualProduction/RemoteControl/Source/WebRemoteControl/Private/WebRemoteControlInternalUtils.cpp
@@ -631,9 +631,9 @@ bool WebRemoteControlInternalUtils::IsRequestContentType(const FHttpServerReques
 {
 	if (const TArray<FString>* ContentTypeHeaders = InRequest.Headers.Find(TEXT("Content-Type")))
 	{
-		if (ContentTypeHeaders->Num() > 0 && (*ContentTypeHeaders)[0] == InContentType)
+		if (ContentTypeHeaders->Num() > 0)
 		{
-			return true;
+			return (*ContentTypeHeaders)[0].StartsWith(InContentType);
 		}
 	}
 
