CIS fixes.

diff --git a/Engine/Plugins/VirtualProduction/RemoteControl/Source/RemoteControl/Private/RemoteControlBinding.cpp b/Engine/Plugins/VirtualProduction/RemoteControl/Source/RemoteControl/Private/RemoteControlBinding.cpp
--- a/Engine/Plugins/VirtualProduction/RemoteControl/Source/RemoteControl/Private/RemoteControlBinding.cpp
+++ b/Engine/Plugins/VirtualProduction/RemoteControl/Source/RemoteControl/Private/RemoteControlBinding.cpp
@@ -328,7 +328,7 @@ UObject* URemoteControlLevelDependantBinding::ResolveForWorld(const UWorld* InWo
 	{
 		const UObject* CurrentResolvedObject = ResolveForCurrentWorld(true).Get();
 		FSoftObjectPath ObjectPath = CurrentResolvedObject;
-		ObjectPath = FSoftObjectPath(WorldPackage->GetFName(), ObjectPath.GetAssetFName(), ObjectPath.GetSubPathString());
+		ObjectPath = FSoftObjectPath(FTopLevelAssetPath(WorldPackage->GetFName(), ObjectPath.GetAssetFName()), ObjectPath.GetSubPathString());
 
 		if (UObject* ResolvedObject = ObjectPath.ResolveObject())
 		{
diff --git a/Engine/Plugins/VirtualProduction/RemoteControl/Source/RemoteControlUI/Private/UI/SRemoteControlPanel.cpp b/Engine/Plugins/VirtualProduction/RemoteControl/Source/RemoteControlUI/Private/UI/SRemoteControlPanel.cpp
--- a/Engine/Plugins/VirtualProduction/RemoteControl/Source/RemoteControlUI/Private/UI/SRemoteControlPanel.cpp
+++ b/Engine/Plugins/VirtualProduction/RemoteControl/Source/RemoteControlUI/Private/UI/SRemoteControlPanel.cpp
@@ -2521,7 +2521,7 @@ void SRemoteControlPanel::OpenEditorEmbeddedPreset()
 		{
 			const FSoftObjectPath RemotePreset = Preset.Get();
 			// Fix path for the editor world
-			const FSoftObjectPath EditorPresetPath = FSoftObjectPath(EditorWorldPackage->GetFName(), RemotePreset.GetAssetFName(), RemotePreset.GetSubPathString());
+			const FSoftObjectPath EditorPresetPath = FSoftObjectPath(FTopLevelAssetPath(EditorWorldPackage->GetFName(), RemotePreset.GetAssetFName()), RemotePreset.GetSubPathString());
 			OpenEmbeddedPreset(EditorPresetPath);
 		}
 	}
@@ -2536,7 +2536,7 @@ void SRemoteControlPanel::OpenPanelForEmbeddedPreset(const UWorld* InWorld)
 		{
 			FSoftObjectPath CurrentPresetSoftPath = Preset.Get();
 			// Fix path for the current selected world
-			CurrentPresetSoftPath = FSoftObjectPath(Package->GetFName(), CurrentPresetSoftPath.GetAssetFName(), CurrentPresetSoftPath.GetSubPathString());
+			CurrentPresetSoftPath = FSoftObjectPath(FTopLevelAssetPath(Package->GetFName(), CurrentPresetSoftPath.GetAssetFName()), CurrentPresetSoftPath.GetSubPathString());
 			OpenEmbeddedPreset(CurrentPresetSoftPath);
 		}
 		else if (InWorld->WorldType == EWorldType::Editor)
