[Remote Control] Fix Texture Controller in external mode.

The previous method of setting transient textures in the material instances didn't work because the texture objects would be gone on the next load and it breaks the material instances. The texture objects must be non-transient and remain valid when the materials are loading. The new approach is to create non-transient texture objects embedded in the RCP package itself.

The new texture class URCExternalTexture will automatically load from the external file. It is a low overhead object in the level itself.

Also, for the rundown, we consider the texture controller to be special and we automatically apply them to ensure the external textuers get loaded and assigned.

When the external textures are set in the level instances, we use the "transient" method except now, the textures are named according to the file path to reuse the textures once loaded.

Object Property Handle Fix:
- Fix Material Designer related crash because transaction buffer is holding on to materials in the runtime instances.
- Fix SetObjectValue to correctly do the serialize - deserialize method. Specialized the template to support UObject*.

diff --git a/Engine/Plugins/Experimental/Avalanche/Source/AvalancheMedia/Private/Playable/AvaPlayable.cpp b/Engine/Plugins/Experimental/Avalanche/Source/AvalancheMedia/Private/Playable/AvaPlayable.cpp
--- a/Engine/Plugins/Experimental/Avalanche/Source/AvalancheMedia/Private/Playable/AvaPlayable.cpp
+++ b/Engine/Plugins/Experimental/Avalanche/Source/AvalancheMedia/Private/Playable/AvaPlayable.cpp
@@ -5,21 +5,38 @@
 #include "AvaSequence.h"
 #include "AvaSequencePlayer.h"
 #include "Broadcast/AvaBroadcast.h"
+#include "Controller/RCCustomControllerUtilities.h"
 #include "Engine/Engine.h"
 #include "Framework/AvaSoftAssetPtr.h"
 #include "IAvaSceneInterface.h"
 #include "IAvaSequenceProvider.h"
 #include "Playable/AvaPlayableGroup.h"
+#include "Playable/AvaPlayableRemoteControl.h"
 #include "Playable/AvaPlayableRemoteControlValues.h"
 #include "Playable/Playables/AvaPlayableLevelStreaming.h"
 #include "Playable/Playables/AvaPlayableRemoteProxy.h"
 #include "Playback/AvaPlaybackUtils.h"
+#include "RCVirtualProperty.h"
 #include "RemoteControlPreset.h"
 
 DEFINE_LOG_CATEGORY(LogAvaPlayable);
 
 #define LOCTEXT_NAMESPACE "AvaPlayable"
 
+// Transition logic and Scene State requires the RC controller values to be updated.
+static TAutoConsoleVariable<bool> CVarUpdateRCControllerValues(
+	TEXT("MotionDesignPlayable.UpdateRCControllerValues"),
+	true,
+	TEXT("Set the controller values without running the behaviors."),
+	ECVF_Default);
+
+// Some "special" controllers, such as texture controller require the behavior to be executed.
+static TAutoConsoleVariable<bool> CVarExecuteSpecialRCControllerBehavior(
+	TEXT("MotionDesignPlayable.ExecuteSpecialRCControllerBehavior"),
+	true,
+	TEXT("Execute the special controller behaviors (such as texture controllers)."),
+	ECVF_Default);
+
 namespace UE::AvaPlayable::Private
 {
 	bool ShouldCreateLocalPlayable(const FName& InChannelName, const UAvaBroadcast& InBroadcast)
@@ -262,10 +279,74 @@ EAvaPlayableCommandResult UAvaPlayable::UpdateRemoteControlCommand(const TShared
 	using namespace UE::AvaPlayable::Private;
 	UE_LOG(LogAvaPlayable, Verbose, TEXT("%s Playable {%s} -> Executing RC Update."), *GetBriefFrameInfo(), *GetPrettyPlayableInfo(this));
 
+	TSet<FGuid> AppliedEntities;
+	
+	// Apply special controllers that don't work just with entity values.
+	const bool bExecuteSpecialControllerBehavior = CVarExecuteSpecialRCControllerBehavior.GetValueOnGameThread();
+	const bool bUpdateControllerValues = CVarUpdateRCControllerValues.GetValueOnGameThread();
+	if (bExecuteSpecialControllerBehavior || bUpdateControllerValues)
+	{
+		TArray<URCVirtualPropertyBase*> Controllers = RemoteControlPreset->GetControllers();
+		for (URCVirtualPropertyBase* Controller : Controllers)
+		{
+			using namespace UE::AvaPlayableRemoteControl;
+			using namespace UE::RCCustomControllers;
+
+			// Texture Controller Require the Bind Behavior to be executed to setup the texture in external mode.
+			if (bExecuteSpecialControllerBehavior &&
+				GetCustomControllerTypeName(Controller) == CustomTextureControllerName)
+			{
+				if (const FAvaPlayableRemoteControlValue* ControllerValue = InRemoteControlValues->GetControllerValue(Controller->Id))
+				{
+					SetValueOfController(Controller, ControllerValue->Value);
+					GetEntitiesControlledByController(RemoteControlPreset, Controller, AppliedEntities);
+				}
+			}
+			else if (bUpdateControllerValues)
+			{
+				if (const FAvaPlayableRemoteControlValue* ControllerValue = InRemoteControlValues->GetControllerValue(Controller->Id))
+				{
+					FScopedPushControllerBehavioursEnable PushBehaviorsEnable(Controller, /*bInBehaviorsEnabled*/ false);
+					SetValueOfController(Controller, ControllerValue->Value);
+				}
+			}
+		}
+	}
+
 	// WYSIWYG (Solution): For the runtime/playback RCP, we don't apply the controllers.
 	// We assume the controller actions are already executed in the rundown's managed RCP
 	// during page edition and the resulting entity values are already captured.
-	InRemoteControlValues->ApplyEntityValuesToRemoteControlPreset(RemoteControlPreset);
+
+	// Todo: change this function to support filter predicate.
+	//InRemoteControlValues->ApplyEntityValuesToRemoteControlPreset(RemoteControlPreset);
+
+	for (const TWeakPtr<FRemoteControlEntity>& EntityWeakPtr : RemoteControlPreset->GetExposedEntities<FRemoteControlEntity>())
+	{
+		if (const TSharedPtr<FRemoteControlEntity> Entity = EntityWeakPtr.Pin())
+		{
+			// Skip properties that are already set by controllers at runtime.
+			if (AppliedEntities.Contains(Entity->GetId()))
+			{
+				continue;
+			}
+			
+			if (const FAvaPlayableRemoteControlValue* EntityValue = InRemoteControlValues->GetEntityValue(Entity->GetId()))
+			{
+				using namespace UE::AvaPlayableRemoteControl;
+				const EAvaPlayableRemoteControlResult Result = SetValueOfEntity(Entity, EntityValue->Value);
+				if (Failed(Result))
+				{
+					UE_LOG(LogAvaPlayable, Error, TEXT("Failed to set value of exposed entity \"%s\" (id:%s): %s."),
+						*Entity->GetLabel().ToString(), *Entity->GetId().ToString(), *EnumToString(Result));
+				}
+			}
+			else
+			{
+				UE_LOG(LogAvaPlayable, Error, TEXT("Exposed entity \"%s\" (id:%s): value not found in page."),
+					*Entity->GetLabel().ToString(), *Entity->GetId().ToString());
+			}
+		}
+	}
 
 	OnRemoteControlValuesApplied();
 
diff --git a/Engine/Plugins/VirtualProduction/RemoteControl/Source/RemoteControl/Private/RemoteControlPropertyHandle.cpp b/Engine/Plugins/VirtualProduction/RemoteControl/Source/RemoteControl/Private/RemoteControlPropertyHandle.cpp
--- a/Engine/Plugins/VirtualProduction/RemoteControl/Source/RemoteControl/Private/RemoteControlPropertyHandle.cpp
+++ b/Engine/Plugins/VirtualProduction/RemoteControl/Source/RemoteControl/Private/RemoteControlPropertyHandle.cpp
@@ -157,6 +157,18 @@ namespace
 		return true;
 	}
 
+	template<typename InValueType>
+	void WritePropertyValue(FCborWriter& InCborWriter, const InValueType& InValue)
+	{
+		InCborWriter.WriteValue(InValue);
+	}
+
+	template<>
+	void WritePropertyValue<UObject*>(FCborWriter& InCborWriter, UObject* const& InValue)
+	{
+		InCborWriter.WriteValue(InValue ? InValue->GetPathName() : FString());
+	}
+
 	/**
 	 * Set property value for given property Handle
 	 * Setter following standard SetObjectProperties path from IRemoteControlModule
@@ -197,7 +209,7 @@ namespace
 			CborWriter.WriteValue(Property->GetName());
 		}
 
-		CborWriter.WriteValue(Value);
+		WritePropertyValue(CborWriter, Value);
 		CborWriter.WriteContainerEnd();
 
 		// Set a deserializer backend
@@ -1871,10 +1883,7 @@ bool FRemoteControlPropertyHandleObject::SetValue(UObject* InValue)
 	{
 		return SetStaticMeshValue(InValue);
 	}
-	else
-	{
-		return SetObjectValue(InValue);
-	}
+	return SetPropertyValue(*this, InValue);
 }
 
 bool FRemoteControlPropertyHandleObject::GetValue(UObject* OutValue) const
@@ -1989,54 +1998,6 @@ bool FRemoteControlPropertyHandleObject::SetMaterialValue(UObject* InValue, int3
 	return bSuccess;
 }
 
-bool FRemoteControlPropertyHandleObject::SetObjectValue(UObject* InValue) const
-{
-	if (!InValue)
-	{
-		return false;
-	}
-
-	const TSharedPtr<FRemoteControlProperty>& RemoteControlPropertyPtr = GetRCProperty();
-	
-	FRCObjectReference ObjectRef;
-	ObjectRef.Property = RemoteControlPropertyPtr->GetProperty();
-	ObjectRef.Access = RemoteControlPropertyPtr->GetPropertyHandle()->ShouldGenerateTransaction() ? ERCAccess::WRITE_TRANSACTION_ACCESS : ERCAccess::WRITE_ACCESS;
-	ObjectRef.PropertyPathInfo = RemoteControlPropertyPtr->FieldPathInfo.ToString();
-
-#if WITH_EDITOR
-	FScopedTransaction Transaction(LOCTEXT("SetObjectValue", "Set Object Value"));
-#endif
-	
-	for (UObject* Object : RemoteControlPropertyPtr->GetBoundObjects())
-	{
-		if (IRemoteControlModule::Get().ResolveObjectProperty(ObjectRef.Access, Object, ObjectRef.PropertyPathInfo, ObjectRef))
-		{
-#if WITH_EDITOR
-			Object->Modify();
-#endif
-			/** Setting with Pointer and using Serialization of itself afterwards as a workaround with the Asset not updating in the world. */
-			FObjectPropertyBase* ObjectProperty = CastField<FObjectPropertyBase>(RemoteControlPropertyPtr->GetProperty());
-			if (!ObjectProperty)
-			{
-				return false;
-			}
-			
-			FProperty* Property = RemoteControlPropertyPtr->GetProperty();
-			uint8* ValueAddress = Property->ContainerPtrToValuePtr<uint8>(ObjectRef.ContainerAdress);
-			ObjectProperty->SetObjectPropertyValue(ValueAddress, InValue);
-
-			TArray<uint8> Buffer;
-			FMemoryReader Reader(Buffer);
-			FCborStructDeserializerBackend DeserializerBackend(Reader);
-
-			IRemoteControlModule::Get().SetObjectProperties(ObjectRef, DeserializerBackend, ERCPayloadType::Cbor, Buffer);
-			return true;
-		}
-	}
-	
-	return false;
-}
-
 bool FRemoteControlPropertyHandleObject::Supports(const FProperty* InProperty)
 {
 	if (!InProperty)
diff --git a/Engine/Plugins/VirtualProduction/RemoteControl/Source/RemoteControl/Private/RemoteControlPropertyHandle.h b/Engine/Plugins/VirtualProduction/RemoteControl/Source/RemoteControl/Private/RemoteControlPropertyHandle.h
--- a/Engine/Plugins/VirtualProduction/RemoteControl/Source/RemoteControl/Private/RemoteControlPropertyHandle.h
+++ b/Engine/Plugins/VirtualProduction/RemoteControl/Source/RemoteControl/Private/RemoteControlPropertyHandle.h
@@ -624,9 +624,6 @@ public:
 
 	/** Set a Material value using the handled Object */
 	bool SetMaterialValue(UObject* InValue, int32 InIndex) const;
-
-	/** Set Object Value */
-	bool SetObjectValue(UObject* InValue) const;
 };
 
 #undef DECLARE_PROPERTY_ACCESSOR_CONSTRUCTOR
diff --git a/Engine/Plugins/VirtualProduction/RemoteControl/Source/RemoteControlLogic/Private/Action/Bind/RCCustomBindActionUtilities.cpp b/Engine/Plugins/VirtualProduction/RemoteControl/Source/RemoteControlLogic/Private/Action/Bind/RCCustomBindActionUtilities.cpp
--- a/Engine/Plugins/VirtualProduction/RemoteControl/Source/RemoteControlLogic/Private/Action/Bind/RCCustomBindActionUtilities.cpp
+++ b/Engine/Plugins/VirtualProduction/RemoteControl/Source/RemoteControlLogic/Private/Action/Bind/RCCustomBindActionUtilities.cpp
@@ -1,44 +1,227 @@
-﻿// Copyright Epic Games, Inc. All Rights Reserved.
+// Copyright Epic Games, Inc. All Rights Reserved.
 
 #include "Action/Bind/RCCustomBindActionUtilities.h"
 #include "AssetRegistry/AssetData.h"
+#include "Engine/Level.h"
 #include "Engine/Texture2D.h"
 #include "IRemoteControlPropertyHandle.h"
 #include "Kismet/KismetRenderingLibrary.h"
 #include "Misc/Paths.h"
+#include "RCCustomBindActionUtilitiesPrivate.h"
+#include "RCExternalTexture.h"
 #include "RemoteControlField.h"
 
-UTexture2D* UE::RCCustomBindActionUtilities::LoadTextureFromPath(const FString& InPath)
+namespace UE::RCCustomBindActionUtilities::Private
 {
-	if (UTexture2D* Texture = Cast<UTexture2D>(StaticLoadObject(UTexture2D::StaticClass(), nullptr, *InPath)))
+	/**
+	 * Expand some common path tokens.
+	 * Can be found elsewhere in engine: mrq, render grid, img media or python for instance.
+	 */
+	FString ExpandSequencePathTokens(const FString& InPath)
 	{
-		return Texture;
+		// Note: Replace ignores case by default (intended).
+		FString ExpandedPath = InPath
+			.Replace(TEXT("{project_dir}"), *FPaths::ConvertRelativePathToFull(FPaths::ProjectDir()));
+
+		FPaths::NormalizeDirectoryName(ExpandedPath);
+		FPaths::RemoveDuplicateSlashes(ExpandedPath);
+		return ExpandedPath;
 	}
 
-	if (FPaths::FileExists(InPath))
-	{	
-		return UKismetRenderingLibrary::ImportFileAsTexture2D(nullptr, InPath);
+	/**
+	 * Try to resolve the given string to a valid absolute path.
+	 */
+	FString TryResolvePath(const FString& InPath)
+	{
+		// 1- Can be already an absolute path
+		if (FPaths::FileExists(InPath))
+		{
+			return InPath;
+		}
+
+		// 2- Try Relative to content
+		const FString ResolvedPath = FPaths::Combine(FPaths::ProjectContentDir(), InPath);
+		if (FPaths::FileExists(ResolvedPath))
+		{
+			return FPaths::ConvertRelativePathToFull(ResolvedPath);
+		}
+
+		// 3- Using other explicit path tokens (engine, project, ...)
+		return ExpandSequencePathTokens(InPath);
 	}
 
-	return nullptr;
+	static FString TryMakeRelativeToContent(const FString& InPath)
+	{
+		const FString ContentPath = FPaths::ConvertRelativePathToFull(FPaths::ProjectContentDir());
+		FString ConvertedPath = ExpandSequencePathTokens(InPath);
+		if (FPaths::IsUnderDirectory(ConvertedPath, ContentPath))
+		{
+			if (FPaths::MakePathRelativeTo(ConvertedPath, *ContentPath))
+			{
+				return ConvertedPath;
+			}
+		}
+		return InPath;
+	}
+	
+	/**
+	 * Generate a unique readable object name from the path. 
+	 */
+	static FString GetTextureObjectNameFromPath(const FString& InPath)
+	{
+		// This name appears in the asset picker
+		FString TextureName = TEXT("RCExtTexture_");
+
+		// Try to make a cleaner name if under content.
+		FString ConvertedPath = TryMakeRelativeToContent(InPath);
+
+		// Clean up the other path tokens
+		ConvertedPath.ReplaceInline(TEXT("{project_dir}"), TEXT("Project_"));
+
+		// Replace forbidden characters in the path (can't be in an object name)
+		for (const TCHAR* InvalidCharacters = INVALID_OBJECTNAME_CHARACTERS; *InvalidCharacters; ++InvalidCharacters)
+		{
+			ConvertedPath.ReplaceCharInline(*InvalidCharacters, TEXT('_'));
+		}
+
+		TextureName.Append(ConvertedPath);
+
+		// Hash the full original filepath to ensure uniqueness.
+		TextureName.Appendf(TEXT("_%x"), GetTypeHash(InPath));
+		return TextureName;
+	}
+
+	/**
+	 * Check if the given string is a valid asset path, if so, load it. 
+	 */
+	static UTexture2D* LoadAsTextureAsset(const FString& InPath)
+	{
+		if (FPackageName::IsValidTextForLongPackageName(InPath))
+		{
+			return Cast<UTexture2D>(StaticLoadObject(UTexture2D::StaticClass(), nullptr, *InPath));
+		}
+		return nullptr;
+	}
+
+	UTexture2D* LoadExternalAsTransientTexture(const FString& InPath)
+	{
+		const FString TextureName = Private::GetTextureObjectNameFromPath(InPath);
+	
+		if (UTexture2D* Texture = Cast<UTexture2D>(StaticFindObject(UTexture2D::StaticClass(), GetTransientPackage(), *TextureName)))
+		{
+			return Texture;
+		}
+
+		const FString ResolvedPath = TryResolvePath(InPath);
+
+		if (FPaths::FileExists(ResolvedPath))
+		{
+			// Todo: generate mipmaps.
+			UTexture2D* Texture = UKismetRenderingLibrary::ImportFileAsTexture2D(nullptr, ResolvedPath);
+		
+			if (Texture)
+			{
+				Texture->Rename(*TextureName, GetTransientPackage(), REN_DontCreateRedirectors | REN_DoNotDirty);
+			}
+			return Texture;
+		}
+		return nullptr;
+	}
+	
 }
 
-void UE::RCCustomBindActionUtilities::SetTexturePropertyFromPath(const TSharedRef<FRemoteControlProperty>& InRemoteControlEntityAsProperty, const FString& InPath)
+UTexture* UE::RCCustomBindActionUtilities::LoadTextureFromPath(UObject* InOuter, const FString& InPath)
 {
-	if (UTexture2D* LoadedTexture = UE::RCCustomBindActionUtilities::LoadTextureFromPath(InPath))
+	// Check if the string is an asset.
+	if (UTexture* TextureAsset = Private::LoadAsTextureAsset(InPath))
+	{
+		return TextureAsset;
+	}
+	
+	// No outer, most likely for UI thumbnail.
+	if (!InOuter)
 	{
-		const FProperty* RemoteControlProperty = InRemoteControlEntityAsProperty->GetProperty();
-		if (RemoteControlProperty == nullptr)
+		return Private::LoadExternalAsTransientTexture(InPath);
+	}
+
+	// Check if the RCP is part of a level instance
+	if (const ULevel* Level = InOuter->GetTypedOuter<ULevel>())
+	{
+		if (Level->IsInstancedLevel())
+		{
+			return Private::LoadExternalAsTransientTexture(InPath);
+		}
+	}
+
+	const FString ResolvedPath = Private::TryResolvePath(InPath);
+
+	// We have an outer and it is not a level instance -> load the texture embedded in the outer.
+	// This will ensure there is a valid texture reference for materials in the level template.
+	if (FPaths::FileExists(ResolvedPath))
+	{
+		const FString TextureName = Private::GetTextureObjectNameFromPath(InPath);
+
+		if (UTexture* Texture = Cast<UTexture>(StaticFindObject(UTexture::StaticClass(), InOuter, *TextureName)))
 		{
-			return;
+			return Texture;
 		}
 
-		const TSharedPtr<IRemoteControlPropertyHandle>& RemoteControlHandle = InRemoteControlEntityAsProperty->GetPropertyHandle();
-		if (!ensure(RemoteControlHandle))
+		// Using an embedded "external" texture object that will automatically load it's content
+		// from the external file.
+		UTexture* Texture = URCExternalTexture::Create(ResolvedPath);
+
+		if (Texture && InOuter)
 		{
-			return;
+			Texture->Rename(*TextureName, InOuter, REN_DontCreateRedirectors | REN_DoNotDirty);
+			Texture->ClearFlags(RF_Transient);
 		}
-			
-		RemoteControlHandle->SetValue(LoadedTexture);
+		return Texture;
+	}
+	return nullptr;
+}
+
+UTexture2D* UE::RCCustomBindActionUtilities::LoadTextureFromPath(const FString& InPath)
+{
+	// Check if the string is an asset.
+	if (UTexture2D* TextureAsset = Private::LoadAsTextureAsset(InPath))
+	{
+		return TextureAsset;
+	}
+
+	// Consider the string as an external path.
+	return Private::LoadExternalAsTransientTexture(InPath);
+}
+
+void UE::RCCustomBindActionUtilities::SetTexturePropertyFromPath(const TSharedRef<FRemoteControlProperty>& InRemoteControlEntityAsProperty, const FString& InPath)
+{
+	if (UTexture2D* LoadedTexture = LoadTextureFromPath(InPath))
+	{
+		SetTextureProperty(InRemoteControlEntityAsProperty, LoadedTexture);
+	}
+}
+
+bool UE::RCCustomBindActionUtilities::SetTexturePropertyFromPath(const TSharedRef<FRemoteControlProperty>& InRemoteControlEntityAsProperty, UObject* InOuter, const FString& InPath)
+{
+	if (UTexture* LoadedTexture = LoadTextureFromPath(InOuter, InPath))
+	{
+		return SetTextureProperty(InRemoteControlEntityAsProperty, LoadedTexture);
+	}
+	return false;
+}
+
+bool UE::RCCustomBindActionUtilities::SetTextureProperty(const TSharedRef<FRemoteControlProperty>& InRemoteControlEntityAsProperty, UTexture* InTexture)
+{
+	const FProperty* RemoteControlProperty = InRemoteControlEntityAsProperty->GetProperty();
+	if (RemoteControlProperty == nullptr)
+	{
+		return false;
 	}
+
+	const TSharedPtr<IRemoteControlPropertyHandle>& RemoteControlHandle = InRemoteControlEntityAsProperty->GetPropertyHandle();
+	if (!ensure(RemoteControlHandle))
+	{
+		return false;
+	}
+
+	return RemoteControlHandle->SetValue(InTexture);
 }
diff --git a/Engine/Plugins/VirtualProduction/RemoteControl/Source/RemoteControlLogic/Private/Action/Bind/RCCustomBindActionUtilitiesPrivate.h b/Engine/Plugins/VirtualProduction/RemoteControl/Source/RemoteControlLogic/Private/Action/Bind/RCCustomBindActionUtilitiesPrivate.h
new file mode 100644
--- /dev/null
+++ b/Engine/Plugins/VirtualProduction/RemoteControl/Source/RemoteControlLogic/Private/Action/Bind/RCCustomBindActionUtilitiesPrivate.h
@@ -0,0 +1,21 @@
+// Copyright Epic Games, Inc. All Rights Reserved.
+
+#pragma once
+
+#include "Templates/SharedPointer.h"
+
+class UTexture;
+struct FRemoteControlProperty;
+
+namespace UE::RCCustomBindActionUtilities
+{
+	UTexture* LoadTextureFromPath(UObject* InOuter, const FString& InPath);
+
+	/**
+	 * Loads the texture found at the specified path, and uses it to set a Texture value in the specified RC Property
+	 * Useful e.g. for Custom External Texture Controller
+	 */
+	bool SetTexturePropertyFromPath(const TSharedRef<FRemoteControlProperty>& InRemoteControlEntityAsProperty, UObject* InOuter, const FString& InPath);
+
+	bool SetTextureProperty(const TSharedRef<FRemoteControlProperty>& InRemoteControlEntityAsProperty, UTexture* InTexture);
+};
diff --git a/Engine/Plugins/VirtualProduction/RemoteControl/Source/RemoteControlLogic/Private/Action/Bind/RCExternalTexture.cpp b/Engine/Plugins/VirtualProduction/RemoteControl/Source/RemoteControlLogic/Private/Action/Bind/RCExternalTexture.cpp
new file mode 100644
--- /dev/null
+++ b/Engine/Plugins/VirtualProduction/RemoteControl/Source/RemoteControlLogic/Private/Action/Bind/RCExternalTexture.cpp
@@ -0,0 +1,276 @@
+// Copyright Epic Games, Inc. All Rights Reserved.
+
+#include "RCExternalTexture.h"
+
+#include "ImageCoreUtils.h"
+#include "ImageUtils.h"
+#include "TextureResource.h"
+#include "UObject/Package.h"
+
+namespace UE::RCExternalTexture
+{
+	struct FResourceCacheEntry
+	{
+		FRCExternalTextureInfo CachedInfo;
+		FTextureRHIRef TextureRHI;
+		FSamplerStateRHIRef SamplerStateRHI;
+
+		FResourceCacheEntry() = default;
+
+		FResourceCacheEntry(const FRCExternalTextureInfo& InCachedInfo)
+			: CachedInfo(InCachedInfo) {}
+	};
+
+	/**
+	 * This class implements a simple cache for the RHI resources.
+	 * It allows multiple instance of an external texture to use the same RHI resources.
+	 */
+	class FResourceCache
+	{
+	public:
+		TSharedPtr<FResourceCacheEntry> Find(const FString& InPath) const
+		{
+			FScopeLock ScopeLock(&EntriesSection);
+			const TWeakPtr<FResourceCacheEntry>* Found = EntriesWeak.Find(InPath);
+			return Found ? (*Found).Pin() : nullptr;
+		}
+
+		void Add(const FString& InPath, const TSharedRef<FResourceCacheEntry>& InEntry)
+		{
+			FScopeLock ScopeLock(&EntriesSection);
+			EntriesWeak.Add(InPath, InEntry);
+		}
+
+		// Todo: Multi-gpu support.
+		// Do we need a cache per GPU or is this internally handled by RHI?
+		static TSharedPtr<FResourceCache> Get()
+		{
+			static TSharedRef<FResourceCache> Instance = MakeShared<FResourceCache>();
+			return Instance;
+		}
+	
+	private:
+		mutable FCriticalSection EntriesSection;
+		TMap<FString, TWeakPtr<FResourceCacheEntry>> EntriesWeak;
+	};
+}
+
+/**
+ * This implementation of texture resource, while based on Texture2DDynamic
+ * also support sharing the RHI resource for all textures of the same path
+ * so the actual texture is loaded only once and reused in each RCExternalTexture instance.
+ */
+class FRCExternalTextureResource : public FTexture2DDynamicResource
+{
+public:
+	using FResourceCacheEntry = UE::RCExternalTexture::FResourceCacheEntry; 
+	using FResourceCache = UE::RCExternalTexture::FResourceCache; 
+	
+	FRCExternalTextureResource(URCExternalTexture* InOwner, const FString& InPath)
+		: FTexture2DDynamicResource(InOwner)
+		, RCTextureOwner(InOwner)
+		, Path(InPath)
+	{
+		
+	}
+
+	virtual uint32 GetSizeX() const override
+	{
+		return RCTextureOwner->SizeX;
+	}
+
+	virtual uint32 GetSizeY() const override
+	{
+		return RCTextureOwner->SizeY;
+	}
+
+	virtual void InitRHI(FRHICommandListBase& RHICmdList) override
+	{
+		const TSharedPtr<FResourceCache> RHICache = FResourceCache::Get();
+		CacheEntry = RHICache->Find(Path);
+
+		// Fast path: texture already cached.
+		if (CacheEntry && CacheEntry->TextureRHI.IsValid())
+		{
+			UpdateRHIFromCacheEntry();
+			return;
+		}
+		
+		FTexture2DDynamicResource::InitRHI(RHICmdList);
+
+		UpdateCacheEntry();
+	}
+	
+	virtual void ReleaseRHI() override
+	{
+		CacheEntry.Reset();
+		FTexture2DDynamicResource::ReleaseRHI();
+	}
+
+	void SetCacheEntry(const FString& InNewPath, const TSharedPtr<FResourceCacheEntry>& InNewEntry)
+	{
+		ReleaseRHI();
+
+		Path = InNewPath;
+		CacheEntry = InNewEntry;
+		UpdateRHIFromCacheEntry();
+	}
+
+	void UpdateRHIFromCacheEntry()
+	{
+		if (CacheEntry)
+		{
+			TextureRHI = CacheEntry->TextureRHI;
+			SamplerStateRHI = CacheEntry->SamplerStateRHI;
+			RHIUpdateTextureReference(RCTextureOwner->TextureReference.TextureReferenceRHI,TextureRHI);
+		}
+	}
+
+	void UpdateCacheEntry()
+	{
+		if (CacheEntry)
+		{
+			CacheEntry->TextureRHI = GetTextureRHI();
+			CacheEntry->SamplerStateRHI = SamplerStateRHI;
+		}
+	}
+
+	URCExternalTexture* RCTextureOwner;
+	FString Path;
+	TSharedPtr<FResourceCacheEntry> CacheEntry;
+};
+
+void URCExternalTexture::PostLoad()
+{
+	SizeX = CachedInfo.SizeX;
+	SizeY = CachedInfo.SizeY;
+	Format = CachedInfo.Format;
+	NumMips = 1;
+	bIsResolveTarget = false;
+
+	Super::PostLoad(); // will call UpdateResource	
+
+	if (!Path.IsEmpty())
+	{
+		LoadFromPath(Path);
+	}
+}
+
+FTextureResource* URCExternalTexture::CreateResource()
+{
+	return new FRCExternalTextureResource(this, Path);
+}
+
+void URCExternalTexture::LoadFromPath(const FString& InPath)
+{
+	using namespace UE::RCExternalTexture;
+	// Fast path: Check if texture is already in cache.
+	if (const TSharedPtr<FResourceCacheEntry> ExistingEntry = FResourceCache::Get()->Find(InPath))
+	{
+		Path = InPath;
+		CachedInfo = ExistingEntry->CachedInfo;
+
+		// Only recreate resource if something changed.
+		if (SizeX != CachedInfo.SizeX || SizeY != CachedInfo.SizeY || Format != CachedInfo.Format || !GetResource())
+		{
+			Init(CachedInfo.SizeX, CachedInfo.SizeY, CachedInfo.Format);	// Calls UpdateResource
+		}
+		
+		if (FRCExternalTextureResource* TextureResource = static_cast<FRCExternalTextureResource*>(GetResource()))
+		{
+			ENQUEUE_RENDER_COMMAND(FUpdateSharedTextureEntry)(
+				[TextureResource, ExistingEntry, NewPath=Path](FRHICommandListImmediate& RHICmdList)
+				{
+					TextureResource->SetCacheEntry(NewPath, ExistingEntry);
+				});
+		}
+		return;
+	}
+
+
+	// This will be loaded on the main thread for now. Most likely causing game thread spike.
+	// As an optimization for loading, we could try to hook this up to a streaming manager.
+	// The RCExternalTexture is typically embedded in a level, can get access with GetTypedOuter<ULevel>.
+	
+	FImage Image;	
+	if (FImageUtils::LoadImage(*InPath, Image))
+	{
+		// Todo: keep file time and hash to detect content change and reload if necessary.
+		Path = InPath;
+		
+		SRGB = Image.GetGammaSpace() == EGammaSpace::Linear ? 0 : 1;
+		
+		ERawImageFormat::Type PixelFormatRawFormat;
+		const EPixelFormat PixelFormat = FImageCoreUtils::GetPixelFormatForRawImageFormat(Image.Format,&PixelFormatRawFormat);
+		CachedInfo.Set(Image.GetWidth(), Image.GetHeight(), PixelFormat);
+		
+		if (SizeX != Image.GetWidth() || SizeY != Image.GetHeight() || Format != PixelFormat || !GetResource())
+		{
+			Init(Image.GetWidth(), Image.GetHeight(), PixelFormat);
+		}
+
+		if (FRCExternalTextureResource* TextureResource = static_cast<FRCExternalTextureResource*>(GetResource()))
+		{
+			TSharedRef<FResourceCacheEntry> NewEntry = MakeShared<FResourceCacheEntry>(CachedInfo);
+			FResourceCache::Get()->Add(Path, NewEntry);
+			
+			ENQUEUE_RENDER_COMMAND(FWriteRawDataToTexture)(
+				[TextureResource, NewEntry, NewPath=Path, NewImage = MoveTemp(Image)](FRHICommandListImmediate& InRHICmdList)
+				{
+					// Todo: support mipmaps. some format like dds can have mipmaps (but not supported by FImageUtils for now).
+					TextureResource->WriteRawToTexture_RenderThread(NewImage.RawData);
+					
+					TextureResource->Path = NewPath;
+					TextureResource->CacheEntry = NewEntry;
+					TextureResource->UpdateCacheEntry();
+				});
+		}
+		return;
+	}
+	
+	// Fallback: Allocate texture with no data.
+	// Todo: should we create a checkerboard pattern?
+	Path.Reset();
+	Init(CachedInfo.SizeX, CachedInfo.SizeY, CachedInfo.Format);
+}
+
+URCExternalTexture* URCExternalTexture::Create(const FTexture2DDynamicCreateInfo& InCreateInfo)
+{
+	URCExternalTexture* NewTexture = NewObject<URCExternalTexture>(GetTransientPackage(), NAME_None, RF_Transient);
+	if (NewTexture)
+	{
+		NewTexture->Filter = InCreateInfo.Filter;
+		NewTexture->SamplerAddressMode = InCreateInfo.SamplerAddressMode;
+		NewTexture->SRGB = InCreateInfo.bSRGB;
+
+		// Disable compression
+		NewTexture->CompressionSettings = TC_Default;
+#if WITH_EDITORONLY_DATA
+		NewTexture->CompressionNone = true;
+		NewTexture->MipGenSettings = TMGS_NoMipmaps;
+		NewTexture->CompressionNoAlpha = true;
+		NewTexture->DeferCompression = false;
+#endif // #if WITH_EDITORONLY_DATA
+		if (InCreateInfo.bIsResolveTarget)
+		{
+			NewTexture->bNoTiling = false;
+		}
+		else
+		{
+			// Untiled format
+			NewTexture->bNoTiling = true;
+		}
+		NewTexture->bIsResolveTarget = InCreateInfo.bIsResolveTarget;
+	}
+	return NewTexture;
+}
+
+URCExternalTexture* URCExternalTexture::Create(const FString& InPath, const FTexture2DDynamicCreateInfo& InCreateInfo)
+{
+	URCExternalTexture* NewTexture = Create(InCreateInfo);
+	if (NewTexture)
+	{
+		NewTexture->LoadFromPath(InPath);
+	}
+	return NewTexture;
+}
diff --git a/Engine/Plugins/VirtualProduction/RemoteControl/Source/RemoteControlLogic/Private/Action/Bind/RCExternalTexture.h b/Engine/Plugins/VirtualProduction/RemoteControl/Source/RemoteControlLogic/Private/Action/Bind/RCExternalTexture.h
new file mode 100644
--- /dev/null
+++ b/Engine/Plugins/VirtualProduction/RemoteControl/Source/RemoteControlLogic/Private/Action/Bind/RCExternalTexture.h
@@ -0,0 +1,87 @@
+// Copyright Epic Games, Inc. All Rights Reserved.
+
+#pragma once
+
+#include "Engine/Texture2DDynamic.h"
+#include "RCExternalTexture.generated.h"
+
+/** Cached information on the external texture. */
+USTRUCT()
+struct FRCExternalTextureInfo
+{
+	GENERATED_BODY()
+	
+	UPROPERTY()
+	int32 SizeX = 1;
+
+	UPROPERTY()
+	int32 SizeY = 1;
+
+	UPROPERTY()
+	TEnumAsByte<enum EPixelFormat> Format = PF_B8G8R8A8;
+
+	void Set(int32 InSizeX, int32 InSizeY, EPixelFormat InFormat)
+	{
+		SizeX = InSizeX;
+		SizeY = InSizeY;
+		Format = InFormat;
+	}
+};
+
+/**
+ *	Texture object that automatically loads the content from an external file.
+ *	The texture object itself does not store anything, except the filename.
+ *
+ *	This is intended to be used with texture controllers in external mode.
+ *	The texture is embedded in the RCP itself to have a valid texture object
+ *	during loading. This avoids breaking the materials that
+ *	require valid texture objects.
+ */
+UCLASS(hidecategories=Object)
+class URCExternalTexture : public UTexture2DDynamic
+{
+	GENERATED_BODY()
+
+public:
+
+	/**
+	 * Path to the external texture file.
+	 */
+	UPROPERTY()
+	FString Path;
+
+	/*
+	 *	Cached information from the last time the texture was loaded.
+	 *	This is only used as backed up the texture file is not accessible.
+	 */
+	UPROPERTY()
+	FRCExternalTextureInfo CachedInfo;
+
+	//~ Begin UObject
+	virtual void PostLoad() override;
+	//~ End UObject
+
+	//~ Begin UTexture
+	virtual FTextureResource* CreateResource() override;
+	//~ End UTexture
+	
+	/**
+	 * @brief Load the content of the texture from the given path.
+	 * @param InPath Path of the texture file.
+	 */
+	void LoadFromPath(const FString& InPath);
+
+	/**
+	 * @brief Utility function to create an empty, uninitialized external texture. Need to call Init() or LoadFromPath() after this.
+	 * @param InCreateInfo creation parameters
+	 * @return created texture
+	 */
+	static URCExternalTexture* Create(const FTexture2DDynamicCreateInfo& InCreateInfo = FTexture2DDynamicCreateInfo());
+
+	/**
+	 * @brief Utility function to create and initialize the texture from the given file path.
+	 * @param InCreateInfo creation parameters
+	 * @return created and initialized texture
+	 */
+	static URCExternalTexture* Create(const FString& InPath, const FTexture2DDynamicCreateInfo& InCreateInfo = FTexture2DDynamicCreateInfo());
+};
\ No newline at end of file
diff --git a/Engine/Plugins/VirtualProduction/RemoteControl/Source/RemoteControlLogic/Private/Action/Bind/RCPropertyBindAction.cpp b/Engine/Plugins/VirtualProduction/RemoteControl/Source/RemoteControlLogic/Private/Action/Bind/RCPropertyBindAction.cpp
--- a/Engine/Plugins/VirtualProduction/RemoteControl/Source/RemoteControlLogic/Private/Action/Bind/RCPropertyBindAction.cpp
+++ b/Engine/Plugins/VirtualProduction/RemoteControl/Source/RemoteControlLogic/Private/Action/Bind/RCPropertyBindAction.cpp
@@ -1,8 +1,7 @@
 // Copyright Epic Games, Inc. All Rights Reserved.
 
 #include "Action/Bind/RCPropertyBindAction.h"
-#include "Action/Bind/RCCustomBindActionUtilities.h"
-#include "Action/RCAction.h"
+#include "Action/Bind/RCCustomBindActionUtilitiesPrivate.h"
 #include "Components/SceneComponent.h"
 #include "Controller/RCController.h"
 #include "Controller/RCCustomControllerUtilities.h"
@@ -235,7 +234,7 @@ void URCPropertyBindAction::Execute() const
 						{
 							if (RemoteControlProperty->IsA(FObjectProperty::StaticClass()))
 							{
-								UE::RCCustomBindActionUtilities::SetTexturePropertyFromPath(RemoteControlEntityAsProperty.ToSharedRef(), StringValue);
+								UE::RCCustomBindActionUtilities::SetTexturePropertyFromPath(RemoteControlEntityAsProperty.ToSharedRef(), Controller, StringValue);
 							}
 						}
 					}
diff --git a/Engine/Plugins/VirtualProduction/RemoteControl/Source/RemoteControlLogic/RemoteControlLogic.Build.cs b/Engine/Plugins/VirtualProduction/RemoteControl/Source/RemoteControlLogic/RemoteControlLogic.Build.cs
--- a/Engine/Plugins/VirtualProduction/RemoteControl/Source/RemoteControlLogic/RemoteControlLogic.Build.cs
+++ b/Engine/Plugins/VirtualProduction/RemoteControl/Source/RemoteControlLogic/RemoteControlLogic.Build.cs
@@ -19,7 +19,10 @@ public class RemoteControlLogic : ModuleRules
 				"Cbor",
 				"Engine",
 				"HTTP",
+				"ImageCore",
 				"RemoteControl",
+				"RenderCore",
+				"RHI",
 				"Serialization",
 			}
 		);
diff --git a/Engine/Plugins/VirtualProduction/RemoteControl/Source/RemoteControlUI/Private/UI/Controller/CustomControllers/SCustomTextureControllerWidget.cpp b/Engine/Plugins/VirtualProduction/RemoteControl/Source/RemoteControlUI/Private/UI/Controller/CustomControllers/SCustomTextureControllerWidget.cpp
--- a/Engine/Plugins/VirtualProduction/RemoteControl/Source/RemoteControlUI/Private/UI/Controller/CustomControllers/SCustomTextureControllerWidget.cpp
+++ b/Engine/Plugins/VirtualProduction/RemoteControl/Source/RemoteControlUI/Private/UI/Controller/CustomControllers/SCustomTextureControllerWidget.cpp
@@ -1,4 +1,4 @@
-﻿// Copyright Epic Games, Inc. All Rights Reserved.
+// Copyright Epic Games, Inc. All Rights Reserved.
 
 #include "SCustomTextureControllerWidget.h"
 #include "Action/Bind/RCCustomBindActionUtilities.h"
@@ -161,54 +161,42 @@ void SCustomTextureControllerWidget::UpdateControllerValue()
 
 void SCustomTextureControllerWidget::HandleFilePathPickerPathPicked(const FString& InPickedPath)
 {
-	// The following code is adapted from FFilePath customization
+	// The following code is adapted from FFilePathStructCustomization::HandleFilePathPickerPathPicked.
 	FString FinalPath = InPickedPath;
 
-	// Keeping these here in case need additional customization
-	bool bLongPackageName = false;
-	bool bRelativeToGameDir = false;
+	// Important: The result of this path conversion must be so that
+	// FPackageName::IsValidTextForLongPackageName returns false.
+	// For instance, it will be so if it doesn't start with a leading slash.
 	
-	if (bLongPackageName)
+	auto ConvertPath = [](const FString& InAbsolutePath)
 	{
-		FString LongPackageName;
-		FString StringFailureReason;
-		if (FPackageName::TryConvertFilenameToLongPackageName(InPickedPath, LongPackageName, &StringFailureReason) == false)
+		// "Content" doesn't have a token since it is the default base path.
+		const FString ContentPath = FPaths::ConvertRelativePathToFull(FPaths::ProjectContentDir());
+		if (FPaths::IsUnderDirectory(InAbsolutePath, ContentPath))
 		{
-			FMessageDialog::Open(EAppMsgType::Ok, FText::FromString(StringFailureReason));
+			FString FinalPath = InAbsolutePath;
+			if (FPaths::MakePathRelativeTo(FinalPath, *ContentPath))
+			{
+				return FinalPath;
+			}
 		}
-		FinalPath = LongPackageName;
-	}
-	else if (bRelativeToGameDir && !InPickedPath.IsEmpty())
-	{
-		// A filepath under the project directory will be made relative to the project directory
-		// Otherwise, the absolute path will be returned unless it doesn't exist, the current path will
-		// be kept. This can happen if it's already relative to project dir (tabbing when selected)
-
-		const FString ProjectDir = FPaths::ProjectDir();
-		const FString AbsoluteProjectDir = FPaths::ConvertRelativePathToFull(ProjectDir);
-		const FString AbsolutePickedPath = FPaths::ConvertRelativePathToFull(InPickedPath);
 
-		// Verify if absolute path to file exists. If it was already relative to content directory
-		// the absolute will be to binaries and will possibly be garbage
-		if (FPaths::FileExists(AbsolutePickedPath))
+		// Attempt to replace some other known paths with tokens.
+		return InAbsolutePath.Replace(*FPaths::ConvertRelativePathToFull(FPaths::ProjectDir()), TEXT("{project_dir}"));
+	};
+	
+	if (!InPickedPath.IsEmpty())
+	{
+		// if received path is relative, it is likely relative to the editor's exe.
+		// We want to replace that path with tokens for known directories so that it can be replaced properly when cooked.
+		if (FPaths::IsRelative(InPickedPath))
 		{
-			// If file is part of the project dir, chop the project dir part
-			// Otherwise, use the absolute path
-			if (AbsolutePickedPath.StartsWith(AbsoluteProjectDir))
-			{
-				FinalPath = AbsolutePickedPath.RightChop(AbsoluteProjectDir.Len());
-			}
-			else
+			const FString AbsolutePickedPath = FPaths::ConvertRelativePathToFull(InPickedPath);
+			if (FPaths::FileExists(AbsolutePickedPath))
 			{
-				FinalPath = AbsolutePickedPath;
+				FinalPath = ConvertPath(AbsolutePickedPath);
 			}
 		}
-		else
-		{
-			// If absolute file doesn't exist, it might already be relative to project dir
-			// If not, then it might be a manual entry, so keep it untouched either way
-			FinalPath = InPickedPath;
-		}
 	}
 
 	CurrentExternalPath = FinalPath;
@@ -231,7 +219,7 @@ TSharedRef<SWidget> SCustomTextureControllerWidget::GetAssetThumbnailWidget()
 
 TSharedRef<SWidget> SCustomTextureControllerWidget::GetExternalTextureValueWidget()
 {
-	static const FString FileTypeFilter = TEXT("Image files (*.jpg; *.png; *.bmp; *.ico; *.exr; *.icns; *.jpeg; *.tga; *.hdr)|*.jpg; *.png; *.bmp; *.ico; *.exr; *.icns; *.jpeg; *.tga; *.hdr");
+	static const FString FileTypeFilter = TEXT("Image files (*.jpg; *.png; *.bmp; *.ico; *.exr; *.icns; *.jpeg; *.tga; *.hdr; *.dds)|*.jpg; *.png; *.bmp; *.ico; *.exr; *.icns; *.jpeg; *.tga; *.hdr; *.dds");
 	
 	return SNew(SFilePathPicker)
 			.BrowseButtonImage(FAppStyle::GetBrush("PropertyWindow.Button_Ellipsis"))
