[RemoteControl] External Texture: fix server target.

diff --git a/Engine/Plugins/VirtualProduction/RemoteControl/Source/RemoteControlLogic/Private/Action/Bind/RCExternalTexture.cpp b/Engine/Plugins/VirtualProduction/RemoteControl/Source/RemoteControlLogic/Private/Action/Bind/RCExternalTexture.cpp
--- a/Engine/Plugins/VirtualProduction/RemoteControl/Source/RemoteControlLogic/Private/Action/Bind/RCExternalTexture.cpp
+++ b/Engine/Plugins/VirtualProduction/RemoteControl/Source/RemoteControlLogic/Private/Action/Bind/RCExternalTexture.cpp
@@ -217,9 +217,10 @@ void URCExternalTexture::LoadFromPath(const FString& InPath)
 			ENQUEUE_RENDER_COMMAND(FWriteRawDataToTexture)(
 				[TextureResource, NewEntry, NewPath=Path, NewImage = MoveTemp(Image)](FRHICommandListImmediate& InRHICmdList)
 				{
+#if !UE_SERVER
 					// Todo: support mipmaps. some format like dds can have mipmaps (but not supported by FImageUtils for now).
 					TextureResource->WriteRawToTexture_RenderThread(NewImage.RawData);
-					
+#endif					
 					TextureResource->Path = NewPath;
 					TextureResource->CacheEntry = NewEntry;
 					TextureResource->UpdateCacheEntry();
