Remote Control: RC Behavior Nodes with editable properties to have a details panel in the behavior by default (this can be overriden by child behavior models).
This allows external/blueprint behaviors to work with parameters that can be set from the preset

diff --git a/Engine/Plugins/VirtualProduction/RemoteControl/Source/RemoteControlLogic/Private/Behaviour/RCBehaviour.cpp b/Engine/Plugins/VirtualProduction/RemoteControl/Source/RemoteControlLogic/Private/Behaviour/RCBehaviour.cpp
--- a/Engine/Plugins/VirtualProduction/RemoteControl/Source/RemoteControlLogic/Private/Behaviour/RCBehaviour.cpp
+++ b/Engine/Plugins/VirtualProduction/RemoteControl/Source/RemoteControlLogic/Private/Behaviour/RCBehaviour.cpp
@@ -136,13 +136,11 @@ URCBehaviourNode* URCBehaviour::GetBehaviourNode()
 	UClass* OverrideBehaviourBlueprintClass = GetOverrideBehaviourBlueprintClass();
 	UClass* FinalBehaviourNodeClass = OverrideBehaviourBlueprintClass ? OverrideBehaviourBlueprintClass : BehaviourNodeClass.Get();
 
-	if (!CachedBehaviourNode || CachedBehaviourNodeClass != FinalBehaviourNodeClass)
+	if (!CachedBehaviourNode || CachedBehaviourNode->GetClass() != FinalBehaviourNodeClass)
 	{
 		CachedBehaviourNode = NewObject<URCBehaviourNode>(this, FinalBehaviourNodeClass);
 	}
 
-	CachedBehaviourNodeClass = FinalBehaviourNodeClass;
-
 	return CachedBehaviourNode;
 }
 
diff --git a/Engine/Plugins/VirtualProduction/RemoteControl/Source/RemoteControlLogic/Public/Behaviour/RCBehaviour.h b/Engine/Plugins/VirtualProduction/RemoteControl/Source/RemoteControlLogic/Public/Behaviour/RCBehaviour.h
--- a/Engine/Plugins/VirtualProduction/RemoteControl/Source/RemoteControlLogic/Public/Behaviour/RCBehaviour.h
+++ b/Engine/Plugins/VirtualProduction/RemoteControl/Source/RemoteControlLogic/Public/Behaviour/RCBehaviour.h
@@ -97,7 +97,6 @@ public:
 	/** Called when an action value changed */
 	virtual void NotifyActionValueChanged(URCAction* InChangedAction) {}
 
-protected:
 	/**
 	 * It created the node if it called first time
 	 * If BehaviourNodeClass changes it creates new instance
@@ -105,6 +104,7 @@ protected:
 	 */
 	URCBehaviourNode* GetBehaviourNode();
 
+protected:
 	/** Execute all the action if not provided a valid Action otherwise will only execute the given action */
 	virtual void ExecuteInternal(const TSet<TObjectPtr<URCAction>>& InActionsToExecute);
 
@@ -133,9 +133,6 @@ public:
 	TWeakObjectPtr<URCController> ControllerWeakPtr;
 
 private:
-	/** Cached behaviour node class */
-	TSubclassOf<UObject> CachedBehaviourNodeClass;
-
 	/** Cached behaviour node */
 	UPROPERTY(Instanced)
 	TObjectPtr<URCBehaviourNode> CachedBehaviourNode;
diff --git a/Engine/Plugins/VirtualProduction/RemoteControl/Source/RemoteControlUI/Private/RemoteControlUIModule.cpp b/Engine/Plugins/VirtualProduction/RemoteControl/Source/RemoteControlUI/Private/RemoteControlUIModule.cpp
--- a/Engine/Plugins/VirtualProduction/RemoteControl/Source/RemoteControlUI/Private/RemoteControlUIModule.cpp
+++ b/Engine/Plugins/VirtualProduction/RemoteControl/Source/RemoteControlUI/Private/RemoteControlUIModule.cpp
@@ -1116,6 +1116,12 @@ bool FRemoteControlUIModule::IsAllowedOwnerObjects(TArray<UObject*> InOuterObjec
 		{
 			return false;
 		}
+
+		/** Disallow exposing properties that come from RC behavior nodes */
+		if (InOuterObjects[0]->IsA<URCBehaviourNode>() || InOuterObjects[0]->GetTypedOuter<URCBehaviourNode>())
+		{
+			return false;
+		}
 	}
 
 	return true;
diff --git a/Engine/Plugins/VirtualProduction/RemoteControl/Source/RemoteControlUI/Private/UI/Action/SRCBehaviourDetails.cpp b/Engine/Plugins/VirtualProduction/RemoteControl/Source/RemoteControlUI/Private/UI/Action/SRCBehaviourDetails.cpp
--- a/Engine/Plugins/VirtualProduction/RemoteControl/Source/RemoteControlUI/Private/UI/Action/SRCBehaviourDetails.cpp
+++ b/Engine/Plugins/VirtualProduction/RemoteControl/Source/RemoteControlUI/Private/UI/Action/SRCBehaviourDetails.cpp
@@ -66,7 +66,6 @@ void SRCBehaviourDetails::Construct(const FArguments& InArgs, TSharedRef<SRCActi
 
 				// Behaviour Specific Details Panel
 				+ SVerticalBox::Slot()
-				.Padding(8.f, 4.f)
 				.AutoHeight()
 				[
 					BehaviourDetailsBox.ToSharedRef()
diff --git a/Engine/Plugins/VirtualProduction/RemoteControl/Source/RemoteControlUI/Private/UI/Behaviour/Builtin/Bind/SRCBehaviourBind.cpp b/Engine/Plugins/VirtualProduction/RemoteControl/Source/RemoteControlUI/Private/UI/Behaviour/Builtin/Bind/SRCBehaviourBind.cpp
--- a/Engine/Plugins/VirtualProduction/RemoteControl/Source/RemoteControlUI/Private/UI/Behaviour/Builtin/Bind/SRCBehaviourBind.cpp
+++ b/Engine/Plugins/VirtualProduction/RemoteControl/Source/RemoteControlUI/Private/UI/Behaviour/Builtin/Bind/SRCBehaviourBind.cpp
@@ -15,6 +15,7 @@ void SRCBehaviourBind::Construct(const FArguments& InArgs, TSharedRef<FRCBehavio
 	BindBehaviourItemWeakPtr = InBehaviourItem;
 
 	ChildSlot
+	.Padding(8.f, 4.f)
 	[
 		SNew(SHorizontalBox)
 		// Label
diff --git a/Engine/Plugins/VirtualProduction/RemoteControl/Source/RemoteControlUI/Private/UI/Behaviour/Builtin/Conditional/SRCBehaviourConditional.cpp b/Engine/Plugins/VirtualProduction/RemoteControl/Source/RemoteControlUI/Private/UI/Behaviour/Builtin/Conditional/SRCBehaviourConditional.cpp
--- a/Engine/Plugins/VirtualProduction/RemoteControl/Source/RemoteControlUI/Private/UI/Behaviour/Builtin/Conditional/SRCBehaviourConditional.cpp
+++ b/Engine/Plugins/VirtualProduction/RemoteControl/Source/RemoteControlUI/Private/UI/Behaviour/Builtin/Conditional/SRCBehaviourConditional.cpp
@@ -54,6 +54,7 @@ void SRCBehaviourConditional::Construct(const FArguments& InArgs, TSharedRef<FRC
 	}
 
 	ChildSlot
+	.Padding(8.f, 4.f)
 	[
 		SNew(SHorizontalBox)
 		// Conditions Panel
diff --git a/Engine/Plugins/VirtualProduction/RemoteControl/Source/RemoteControlUI/Private/UI/Behaviour/Builtin/Path/SRCBehaviourSetAssetByPath.cpp b/Engine/Plugins/VirtualProduction/RemoteControl/Source/RemoteControlUI/Private/UI/Behaviour/Builtin/Path/SRCBehaviourSetAssetByPath.cpp
--- a/Engine/Plugins/VirtualProduction/RemoteControl/Source/RemoteControlUI/Private/UI/Behaviour/Builtin/Path/SRCBehaviourSetAssetByPath.cpp
+++ b/Engine/Plugins/VirtualProduction/RemoteControl/Source/RemoteControlUI/Private/UI/Behaviour/Builtin/Path/SRCBehaviourSetAssetByPath.cpp
@@ -77,6 +77,7 @@ void SRCBehaviourSetAssetByPath::Construct(const FArguments& InArgs, TSharedRef<
 	];
 	
 	ChildSlot
+	.Padding(8.f, 4.f)
 	[
 		SNew(SVerticalBox)
 		+ SVerticalBox::Slot()
diff --git a/Engine/Plugins/VirtualProduction/RemoteControl/Source/RemoteControlUI/Private/UI/Behaviour/Builtin/RangeMap/SRCBehaviourRangeMap.cpp b/Engine/Plugins/VirtualProduction/RemoteControl/Source/RemoteControlUI/Private/UI/Behaviour/Builtin/RangeMap/SRCBehaviourRangeMap.cpp
--- a/Engine/Plugins/VirtualProduction/RemoteControl/Source/RemoteControlUI/Private/UI/Behaviour/Builtin/RangeMap/SRCBehaviourRangeMap.cpp
+++ b/Engine/Plugins/VirtualProduction/RemoteControl/Source/RemoteControlUI/Private/UI/Behaviour/Builtin/RangeMap/SRCBehaviourRangeMap.cpp
@@ -20,6 +20,7 @@ void SRCBehaviourRangeMap::Construct(const FArguments& InArgs, TSharedRef<const
 	RangeMapBehaviourItemWeakPtr = InBehaviourItem;
 
 	ChildSlot
+	.Padding(8.f, 4.f)
 	[
 		SNew(SHorizontalBox)
 
diff --git a/Engine/Plugins/VirtualProduction/RemoteControl/Source/RemoteControlUI/Private/UI/Behaviour/RCBehaviourModel.cpp b/Engine/Plugins/VirtualProduction/RemoteControl/Source/RemoteControlUI/Private/UI/Behaviour/RCBehaviourModel.cpp
--- a/Engine/Plugins/VirtualProduction/RemoteControl/Source/RemoteControlUI/Private/UI/Behaviour/RCBehaviourModel.cpp
+++ b/Engine/Plugins/VirtualProduction/RemoteControl/Source/RemoteControlUI/Private/UI/Behaviour/RCBehaviourModel.cpp
@@ -5,6 +5,8 @@
 #include "Behaviour/RCBehaviourNode.h"
 #include "Engine/Blueprint.h"
 #include "Engine/BlueprintGeneratedClass.h"
+#include "Modules/ModuleManager.h"
+#include "PropertyEditorModule.h"
 #include "Styling/RemoteControlStyles.h"
 #include "UI/Action/RCActionModel.h"
 #include "UI/Action/SRCActionPanel.h"
@@ -12,6 +14,7 @@
 #include "UI/BaseLogicUI/RCLogicHelpers.h"
 #include "UI/RemoteControlPanelStyle.h"
 #include "UI/SRemoteControlPanel.h"
+#include "Widgets/Layout/SBox.h"
 #include "Widgets/SNullWidget.h"
 #include "Widgets/Text/STextBlock.h"
 
@@ -24,15 +27,28 @@ FRCBehaviourModel::FRCBehaviourModel(URCBehaviour* InBehaviour
 {
 	RCPanelStyle = &FRemoteControlPanelStyle::Get()->GetWidgetStyle<FRCPanelStyle>("RemoteControlPanel.BehaviourPanel");
 
-	if (BehaviourWeakPtr.IsValid())
+	if (InBehaviour)
 	{
-		const FText BehaviourDisplayName = BehaviourWeakPtr->GetDisplayName();
-		
+		// Determine whether the behavior node has any property that is instance-editable
+		if (InBehaviour->BehaviourNodeClass)
+		{
+			for (FProperty* Property : TFieldRange<FProperty>(InBehaviour->BehaviourNodeClass))
+			{
+				if (Property->HasAllPropertyFlags(CPF_Edit) && !Property->HasAnyPropertyFlags(CPF_DisableEditOnInstance))
+				{
+					bDetailsEditableBehavior = true;
+					break;
+				}
+			}
+		}
+
+		const FText BehaviourDisplayName = InBehaviour->GetDisplayName();
+
 		SAssignNew(BehaviourTitleText, STextBlock)
 			.Text(BehaviourDisplayName)
 			.TextStyle(&RCPanelStyle->HeaderTextStyle);
 
-		RefreshIsBehaviourEnabled(BehaviourWeakPtr->bIsEnabled);
+		RefreshIsBehaviourEnabled(InBehaviour->bIsEnabled);
 	}
 }
 
@@ -93,12 +109,50 @@ TSharedRef<SWidget> FRCBehaviourModel::GetWidget() const
 
 bool FRCBehaviourModel::HasBehaviourDetailsWidget()
 {
-	return false;
+	return bDetailsEditableBehavior;
 }
 
 TSharedRef<SWidget> FRCBehaviourModel::GetBehaviourDetailsWidget()
 {
-	return SNullWidget::NullWidget;
+	if (!bDetailsEditableBehavior)
+	{
+		return SNullWidget::NullWidget;
+	}
+
+	URCBehaviour* Behavior = BehaviourWeakPtr.Get();
+	if (!Behavior)
+	{
+		return SNullWidget::NullWidget;
+	}
+
+	URCBehaviourNode* BehaviorNode = Behavior->GetBehaviourNode();
+	if (!BehaviorNode)
+	{
+		return SNullWidget::NullWidget;
+	}
+
+	FDetailsViewArgs DetailsViewArgs;
+	DetailsViewArgs.DefaultsOnlyVisibility = EEditDefaultsOnlyNodeVisibility::Automatic;
+	DetailsViewArgs.bAllowSearch = false;
+	DetailsViewArgs.bAllowFavoriteSystem = false;
+	DetailsViewArgs.bHideSelectionTip = true;
+	DetailsViewArgs.bLockable = false;
+	DetailsViewArgs.bSearchInitialKeyFocus = true;
+	DetailsViewArgs.bUpdatesFromSelection = false;
+	DetailsViewArgs.bShowOptions = false;
+	DetailsViewArgs.bShowObjectLabel = false;
+	DetailsViewArgs.bShowScrollBar = true;
+
+	FPropertyEditorModule& PropertyEditorModule = FModuleManager::Get().LoadModuleChecked<FPropertyEditorModule>("PropertyEditor");
+
+	TSharedRef<IDetailsView> DetailsView = PropertyEditorModule.CreateDetailView(DetailsViewArgs);
+	DetailsView->SetObject(BehaviorNode);
+
+	return SNew(SBox)
+		.MaxDesiredHeight(200.f)
+		[
+			DetailsView
+		];
 }
 
 void FRCBehaviourModel::OnOverrideBlueprint() const
diff --git a/Engine/Plugins/VirtualProduction/RemoteControl/Source/RemoteControlUI/Private/UI/Behaviour/RCBehaviourModel.h b/Engine/Plugins/VirtualProduction/RemoteControl/Source/RemoteControlUI/Private/UI/Behaviour/RCBehaviourModel.h
--- a/Engine/Plugins/VirtualProduction/RemoteControl/Source/RemoteControlUI/Private/UI/Behaviour/RCBehaviourModel.h
+++ b/Engine/Plugins/VirtualProduction/RemoteControl/Source/RemoteControlUI/Private/UI/Behaviour/RCBehaviourModel.h
@@ -83,4 +83,7 @@ private:
 
 	/** Panel Style reference. */
 	const FRCPanelStyle* RCPanelStyle;
+
+	/** Determines if the behavior is an object editable in details view */
+	bool bDetailsEditableBehavior = false;
 };
\ No newline at end of file
