Remote Control: fixed issue with unexpose button only unexposing one of the actors when multiple actors are selected and all have the property to unexpose exposed

diff --git a/Engine/Plugins/VirtualProduction/RemoteControl/Source/RemoteControlUI/Private/UI/SRemoteControlPanel.cpp b/Engine/Plugins/VirtualProduction/RemoteControl/Source/RemoteControlUI/Private/UI/SRemoteControlPanel.cpp
--- a/Engine/Plugins/VirtualProduction/RemoteControl/Source/RemoteControlUI/Private/UI/SRemoteControlPanel.cpp
+++ b/Engine/Plugins/VirtualProduction/RemoteControl/Source/RemoteControlUI/Private/UI/SRemoteControlPanel.cpp
@@ -1446,7 +1446,7 @@ void SRemoteControlPanel::Unexpose(const FRCExposesPropertyArgs& InPropertyArgs)
 		return;
 	}
 
-	auto CheckAndUnexpose = [&](TArray<UObject*> InOuterObjects, const FString& InPath, bool bInUsingDuplicatesInPath)
+	auto CheckAndUnexpose = [Preset=Preset](TArray<UObject*> InOuterObjects, const FString& InPath, bool bInUsingDuplicatesInPath)
 	{
 		// Find an exposed property with the same path.
 		TArray<TSharedPtr<FRemoteControlProperty>, TInlineAllocator<1>> PotentialMatches;
@@ -1456,8 +1456,8 @@ void SRemoteControlPanel::Unexpose(const FRCExposesPropertyArgs& InPropertyArgs)
 			{
 				// If that was exposed by property path it should be checked by the full path with duplicated like propertypath.propertypath[0]
 				// If that was exposed by the owner object it should be without duplicated in the path, just propertypath[0]
-				const bool bIsbound = bInUsingDuplicatesInPath ? Property->CheckIsBoundToPropertyPath(InPath) : Property->CheckIsBoundToString(InPath);
-				if (bIsbound)
+				const bool bIsBound = bInUsingDuplicatesInPath ? Property->CheckIsBoundToPropertyPath(InPath) : Property->CheckIsBoundToString(InPath);
+				if (bIsBound)
 				{
 					PotentialMatches.Add(Property);
 				}
@@ -1469,7 +1469,6 @@ void SRemoteControlPanel::Unexpose(const FRCExposesPropertyArgs& InPropertyArgs)
 			if (Property->ContainsBoundObjects(InOuterObjects))
 			{
 				Preset->Unexpose(Property->GetId());
-				break;
 			}
 		}
 	};
