Remote Control: added support to add multiple dragged exposed properties into the controller area

diff --git a/Engine/Plugins/VirtualProduction/RemoteControl/Source/RemoteControlUI/Private/UI/Controller/SRCControllerPanelList.cpp b/Engine/Plugins/VirtualProduction/RemoteControl/Source/RemoteControlUI/Private/UI/Controller/SRCControllerPanelList.cpp
--- a/Engine/Plugins/VirtualProduction/RemoteControl/Source/RemoteControlUI/Private/UI/Controller/SRCControllerPanelList.cpp
+++ b/Engine/Plugins/VirtualProduction/RemoteControl/Source/RemoteControlUI/Private/UI/Controller/SRCControllerPanelList.cpp
@@ -851,33 +851,50 @@ bool SRCControllerPanelList::OnAllowDrop(TSharedPtr<FDragDropOperation> DragDrop
 		// Fetch the Exposed Entity
 		const TArray<FGuid>& ExposedEntitiesIds = DragDropOp->GetSelectedFieldsId();
 
-		// Check if Entity is supported by controllers and currently only 1 dragged entity dragged is supported
-		return ExposedEntitiesIds.Num() == 1 && IsEntitySupported(ExposedEntitiesIds[0]);
+		// Check if Entity is supported by controllers
+		// Allow drop if at least one entity is supported.
+		for (const FGuid& EntityId : ExposedEntitiesIds)
+		{
+			if (IsEntitySupported(EntityId))
+			{
+				return true;
+			}
+		}
 	}
-
 	return false;
 }
 
-FReply SRCControllerPanelList::OnControllerListViewDragDrop(TSharedPtr<FDragDropOperation> DragDropOperation)
+FReply SRCControllerPanelList::OnControllerListViewDragDrop(TSharedPtr<FDragDropOperation> InDragDropOperation)
 {
-	if (TSharedPtr<FExposedEntityDragDrop> DragDropOp = GetExposedEntityDragDrop(DragDropOperation))
+	URemoteControlPreset* const Preset = GetPreset();
+	if (!Preset)
 	{
-		// Fetch the Exposed Entity
-		const FGuid ExposedEntityId = DragDropOp->GetNodeId();
+		return FReply::Handled();
+	}
 
-		if (URemoteControlPreset* Preset = GetPreset())
+	if (TSharedPtr<FExposedEntityDragDrop> DragDropOp = GetExposedEntityDragDrop(InDragDropOperation))
+	{
+		FScopedTransaction Transaction(LOCTEXT("AutoBindEntities", "Auto bind entities to controllers"));
+
+		bool bModified = false;
+		for (const FGuid& ExposedEntityId : DragDropOp->GetSelectedFieldsId())
 		{
 			if (TSharedPtr<const FRemoteControlProperty> RemoteControlProperty = Preset->GetExposedEntity<FRemoteControlProperty>(ExposedEntityId).Pin())
 			{
-				CreateAutoBindForProperty(RemoteControlProperty);
+				bModified |= CreateAutoBindForProperty(RemoteControlProperty);
 			}
 		}
+
+		if (!bModified)
+		{
+			Transaction.Cancel();
+		}
 	}
-	
+
 	return FReply::Handled();
 }
 
-void SRCControllerPanelList::CreateAutoBindForProperty(TSharedPtr<const FRemoteControlProperty> RemoteControlProperty)
+bool SRCControllerPanelList::CreateAutoBindForProperty(TSharedPtr<const FRemoteControlProperty> RemoteControlProperty)
 {
 	if (URemoteControlPreset* Preset = GetPreset())
 	{
@@ -926,8 +943,10 @@ void SRCControllerPanelList::CreateAutoBindForProperty(TSharedPtr<const FRemoteC
 			// Step 3. Create Bind Behaviour and Bind to the property
 			constexpr bool bExecuteBind = true;
 			CreateBindBehaviourAndAssignTo(NewController, RemoteControlProperty.ToSharedRef(), bExecuteBind);
+			return true;
 		}
 	}
+	return false;
 }
 
 void SRCControllerPanelList::CreateBindBehaviourAndAssignTo(URCController* Controller, TSharedRef<const FRemoteControlProperty> InRemoteControlProperty, const bool bExecuteBind)
diff --git a/Engine/Plugins/VirtualProduction/RemoteControl/Source/RemoteControlUI/Private/UI/Controller/SRCControllerPanelList.h b/Engine/Plugins/VirtualProduction/RemoteControl/Source/RemoteControlUI/Private/UI/Controller/SRCControllerPanelList.h
--- a/Engine/Plugins/VirtualProduction/RemoteControl/Source/RemoteControlUI/Private/UI/Controller/SRCControllerPanelList.h
+++ b/Engine/Plugins/VirtualProduction/RemoteControl/Source/RemoteControlUI/Private/UI/Controller/SRCControllerPanelList.h
@@ -125,7 +125,7 @@ public:
 	bool OnAllowDrop(TSharedPtr<FDragDropOperation> DragDropOperation);
 
 	/** Drag-Drop action delegate for the Controllers Panel List */
-	FReply OnControllerListViewDragDrop(TSharedPtr<FDragDropOperation> DragDropOperation);
+	FReply OnControllerListViewDragDrop(TSharedPtr<FDragDropOperation> InDragDropOperation);
 
 	/** Fetches the Remote Control preset associated with the parent panel */
 	virtual URemoteControlPreset* GetPreset() override;
@@ -188,8 +188,8 @@ private:
 	 */
 	void OnControllerValueChanged(TSharedPtr<FRCControllerModel> InControllerModel, bool bInIsMultiController);
 
-	/** Creates a new Controller for the given Remote Control Property and also binds to it */
-	void CreateAutoBindForProperty(TSharedPtr<const FRemoteControlProperty> RemoteControlProperty);
+	/** Creates a new Controller for the given Remote Control Property and also binds to it. Returns true if operation succeeded, false if nothing happened */
+	bool CreateAutoBindForProperty(TSharedPtr<const FRemoteControlProperty> RemoteControlProperty);
 
 	/** The row generator used to represent each Controller as a row, when used with SListView */
 	TSharedPtr<IPropertyRowGenerator> PropertyRowGenerator;
