Remote Control: added multi selection for controllers. This still keeps the single-selection support for behaviors, while allowing multi-controller deletion and other operations. This also introduces multi drag-drop support

diff --git a/Engine/Plugins/VirtualProduction/RemoteControl/Source/RemoteControlLogic/Private/Controller/RCControllerUtilities.cpp b/Engine/Plugins/VirtualProduction/RemoteControl/Source/RemoteControlLogic/Private/Controller/RCControllerUtilities.cpp
--- a/Engine/Plugins/VirtualProduction/RemoteControl/Source/RemoteControlLogic/Private/Controller/RCControllerUtilities.cpp
+++ b/Engine/Plugins/VirtualProduction/RemoteControl/Source/RemoteControlLogic/Private/Controller/RCControllerUtilities.cpp
@@ -4,6 +4,7 @@
 #include "Engine/StaticMesh.h"
 #include "Engine/Texture.h"
 #include "Materials/MaterialInterface.h"
+#include "RemoteControlField.h"
 #include "StructUtils/PropertyBag.h"
 #include "UObject/Class.h"
 
@@ -75,6 +76,16 @@ namespace UE::RCControllers
 		return false;
 	}
 
+	bool CanCreateControllerFromEntity(const TSharedPtr<const FRemoteControlProperty>& InPropertyEntity)
+	{
+		if (!InPropertyEntity.IsValid() || !InPropertyEntity->IsEditable() || InPropertyEntity->FieldType != EExposedFieldType::Property)
+		{
+			// Property with error(s)
+			return false;
+		}
+		return UE::RCControllers::CanCreateControllerFromProperty(InPropertyEntity->GetProperty());
+	}
+
 	bool CanCreateControllerFromProperty(const FProperty* InProperty)
 	{
 		return InProperty && CanCreateControllerFromPropertyDesc(FPropertyBagPropertyDesc(NAME_None, InProperty));
diff --git a/Engine/Plugins/VirtualProduction/RemoteControl/Source/RemoteControlLogic/Public/Controller/RCControllerUtilities.h b/Engine/Plugins/VirtualProduction/RemoteControl/Source/RemoteControlLogic/Public/Controller/RCControllerUtilities.h
--- a/Engine/Plugins/VirtualProduction/RemoteControl/Source/RemoteControlLogic/Public/Controller/RCControllerUtilities.h
+++ b/Engine/Plugins/VirtualProduction/RemoteControl/Source/RemoteControlLogic/Public/Controller/RCControllerUtilities.h
@@ -3,9 +3,11 @@
 #pragma once
 
 #include "HAL/Platform.h"
+#include "Templates/SharedPointerFwd.h"
 
 class FProperty;
 struct FPropertyBagPropertyDesc;
+struct FRemoteControlProperty;
 
 namespace UE::RCControllers
 {
@@ -15,6 +17,12 @@ namespace UE::RCControllers
 	 */
 	bool CanCreateControllerFromPropertyDesc(const FPropertyBagPropertyDesc& InPropertyDesc);
 
+	/**
+	 * Determines whether a controller can be created from a given property entity explicitly
+	 * Explicit meaning that it won't use IsChildOf or recurse up its Super Types for Object/Struct properties
+	 */
+	REMOTECONTROLLOGIC_API bool CanCreateControllerFromEntity(const TSharedPtr<const FRemoteControlProperty>& InPropertyEntity);
+
 	/**
 	 * Determines whether a controller can be created from a given property type explicitly
 	 * Explicit meaning that it won't use IsChildOf or recurse up its Super Types for Object/Struct properties
diff --git a/Engine/Plugins/VirtualProduction/RemoteControl/Source/RemoteControlUI/Private/UI/Controller/SRCControllerItemRow.cpp b/Engine/Plugins/VirtualProduction/RemoteControl/Source/RemoteControlUI/Private/UI/Controller/SRCControllerItemRow.cpp
new file mode 100644
--- /dev/null
+++ b/Engine/Plugins/VirtualProduction/RemoteControl/Source/RemoteControlUI/Private/UI/Controller/SRCControllerItemRow.cpp
@@ -0,0 +1,310 @@
+// Copyright Epic Games, Inc. All Rights Reserved.
+
+#include "SRCControllerItemRow.h"
+#include "Behaviour/Builtin/Bind/RCBehaviourBind.h"
+#include "Controller/RCController.h"
+#include "Controller/RCControllerUtilities.h"
+#include "RemoteControlPreset.h"
+#include "SRCControllerPanelList.h"
+#include "ScopedTransaction.h"
+#include "UI/RCUIHelpers.h"
+#include "UI/SRCPanelExposedEntity.h"
+
+#define LOCTEXT_NAMESPACE "SRCControllerItemRow"
+
+const FLazyName FRCControllerColumns::TypeColor = TEXT("TypeColor");
+const FLazyName FRCControllerColumns::ControllerId = TEXT("Controller Id");
+const FLazyName FRCControllerColumns::Description = TEXT("Controller Description");
+const FLazyName FRCControllerColumns::Value = TEXT("Controller Value");
+const FLazyName FRCControllerColumns::FieldId = TEXT("Controller Field Id");
+const FLazyName FRCControllerColumns::ValueTypeSelection = TEXT("Value Type Selection");
+
+void SRCControllerItemRow::Construct(const FArguments& InArgs, const TSharedRef<STableViewBase>& InOwnerTableView, TSharedRef<FRCControllerModel> InControllerItem, TSharedRef<SRCControllerPanelList> InControllerPanelList)
+{
+	ControllerItem = InControllerItem;
+	ControllerPanelListWeak = InControllerPanelList.ToWeakPtr();
+
+	FSuperRowType::Construct(FSuperRowType::FArguments()
+		.Style(InArgs._Style)
+		.Padding(4.5f)
+		.OnCanAcceptDrop(this, &SRCControllerItemRow::CanAcceptDrop)
+		.OnAcceptDrop(this, &SRCControllerItemRow::OnAcceptDrop)
+		, InOwnerTableView);
+}
+
+TSharedRef<SWidget> SRCControllerItemRow::GenerateWidgetForColumn(const FName& InColumnName)
+{
+	if (!ensure(ControllerItem.IsValid()))
+	{
+		return SNullWidget::NullWidget;
+	}
+
+	if (InColumnName == FRCControllerColumns::TypeColor)
+	{
+		if (URCController* Controller = Cast<URCController>(ControllerItem->GetVirtualProperty()))
+		{
+			if (const FProperty* Property = Controller->GetProperty())
+			{
+				return UE::RCUIHelpers::GetTypeColorWidget(Property);
+			}
+		}
+	}
+	else if (InColumnName == FRCControllerColumns::FieldId)
+	{
+		return ControllerItem->GetFieldIdWidget();
+	}
+	else if (InColumnName == FRCControllerColumns::ValueTypeSelection)
+	{
+		return ControllerItem->GetTypeSelectionWidget();
+	}
+	else if (InColumnName == FRCControllerColumns::ControllerId)
+	{
+		return ControllerItem->GetNameWidget();
+	}
+	else if (InColumnName == FRCControllerColumns::Description)
+	{
+		return ControllerItem->GetDescriptionWidget();
+	}
+	else if (InColumnName == FRCControllerColumns::Value)
+	{
+		return ControllerItem->GetWidget();
+	}
+	else if (const TSharedPtr<SRCControllerPanelList> ControlPanelList = ControllerPanelListWeak.Pin())
+	{
+		if (ControlPanelList->GetCustomColumns().Contains(InColumnName))
+		{
+			return ControllerItem->GetControllerExtensionWidget(InColumnName);
+		}
+	}
+	return SNullWidget::NullWidget;
+}
+
+TOptional<EItemDropZone> SRCControllerItemRow::CanAcceptDrop(const FDragDropEvent& InDragDropEvent, EItemDropZone InDropZone, TSharedRef<FRCControllerModel> InItem)
+{
+	const TSharedPtr<SRCControllerPanelList> ControllerPanelList = ControllerPanelListWeak.Pin();
+	if (!ControllerPanelList.IsValid() || !ControllerItem.IsValid())
+	{
+		return TOptional<EItemDropZone>();
+	}
+
+	// Fetch the Exposed Entities
+	URemoteControlPreset* const Preset = ControllerPanelList->GetPreset();
+	if (!Preset)
+	{
+		return TOptional<EItemDropZone>();
+	}
+
+	// Dragging Controllers onto Controllers (Reordering). Only allow Above/Below re-ordering (no-onto support)
+	if (TSharedPtr<FRCControllerDragDrop> ControllerDragDrop = InDragDropEvent.GetOperationAs<FRCControllerDragDrop>())
+	{
+		if (InDropZone == EItemDropZone::OntoItem)
+		{
+			return TOptional<EItemDropZone>();
+		}
+		return InDropZone;
+	}
+
+	if (TSharedPtr<FExposedEntityDragDrop> EntityDragDrop = InDragDropEvent.GetOperationAs<FExposedEntityDragDrop>())
+	{
+		// Support when adding entities onto this controller, ensure that at least one entity can be added via bind action 
+		if (InDropZone == EItemDropZone::OntoItem)
+		{
+			for (const FGuid& ExposedEntityId : EntityDragDrop->GetSelectedFieldsId())
+			{
+				if (TSharedPtr<const FRemoteControlField> RemoteControlField = Preset->GetExposedEntity<FRemoteControlField>(ExposedEntityId).Pin())
+				{
+					if (URCController* Controller = Cast<URCController>(ControllerItem->GetVirtualProperty()))
+					{
+						const bool bAllowNumericInputAsStrings = true;
+						if (URCBehaviourBind::CanHaveActionForField(Controller, RemoteControlField.ToSharedRef(), bAllowNumericInputAsStrings))
+						{
+							ControllerPanelList->bIsAnyControllerItemEligibleForDragDrop = true;
+							return InDropZone;
+						}
+					}
+				}
+			}
+		}
+		// When adding entities above/below controllers (i.e. create a new controller above/below this one), at least one has to be supported
+		else
+		{
+			for (const FGuid& ExposedEntityId : EntityDragDrop->GetSelectedFieldsId())
+			{
+				if (TSharedPtr<const FRemoteControlProperty> RemoteControlProperty = Preset->GetExposedEntity<FRemoteControlProperty>(ExposedEntityId).Pin())
+				{
+					if (UE::RCControllers::CanCreateControllerFromEntity(RemoteControlProperty))
+					{
+						ControllerPanelList->bIsAnyControllerItemEligibleForDragDrop = true;
+						return InDropZone;
+					}
+				}
+			}
+		}
+	}
+
+	return TOptional<EItemDropZone>();
+}
+
+FReply SRCControllerItemRow::OnAcceptDrop(const FDragDropEvent& InDragDropEvent, EItemDropZone InDropZone, TSharedRef<FRCControllerModel> InItem)
+{
+	if (!ControllerItem.IsValid())
+	{
+		return FReply::Unhandled();
+	}
+
+	const TSharedPtr<SRCControllerPanelList> ControllerPanelList = ControllerPanelListWeak.Pin();
+	if (!ControllerPanelList.IsValid())
+	{
+		return FReply::Unhandled();
+	}
+
+	ControllerPanelList->bIsAnyControllerItemEligibleForDragDrop = false;
+
+	URemoteControlPreset* Preset = ControllerPanelList->GetPreset();
+	if (!Preset)
+	{
+		return FReply::Unhandled();
+	}
+
+	if (TSharedPtr<FRCControllerDragDrop> ControllerDragDrop = InDragDropEvent.GetOperationAs<FRCControllerDragDrop>())
+	{
+		FScopedTransaction Transaction(LOCTEXT("ReorderControllers", "Reorder Controllers"));
+
+		const int32 DropIndex = ControllerPanelList->GetDropIndex(ControllerItem, InDropZone);
+		if (!ControllerPanelList->ReorderControllerItems(ControllerDragDrop->ResolveControllers(), DropIndex))
+		{
+			Transaction.Cancel();
+		}
+		return FReply::Handled();
+	}
+
+	if (TSharedPtr<FExposedEntityDragDrop> ExposedEntityDragDrop = InDragDropEvent.GetOperationAs<FExposedEntityDragDrop>())
+	{
+		const FDragDropContext DragDropContext
+			{
+				.Preset = Preset,
+				.ControllerPanelList = ControllerPanelList,
+				.Item = InItem,
+				.DropZone = InDropZone,
+			};
+
+		if (InDropZone == EItemDropZone::OntoItem)
+		{
+			CreateBindBehaviorsFromEntities(*ExposedEntityDragDrop, DragDropContext);
+		}
+		else
+		{
+			CreateControllersFromEntities(*ExposedEntityDragDrop, DragDropContext);
+		}
+		return FReply::Handled();
+	}
+
+	return FReply::Unhandled();
+}
+
+void SRCControllerItemRow::CreateBindBehaviorsFromEntities(const FExposedEntityDragDrop& InExposedEntityDragDrop, const FDragDropContext& InContext)
+{
+	FScopedTransaction Transaction(LOCTEXT("CreateBindBehaviors", "Create Bind Behaviors from Drag Drop"));
+
+	bool bModified = false;
+	for (const FGuid& ExposedEntityId : InExposedEntityDragDrop.GetSelectedFieldsId())
+	{
+		if (TSharedPtr<const FRemoteControlProperty> RemoteControlProperty = InContext.Preset->GetExposedEntity<FRemoteControlProperty>(ExposedEntityId).Pin())
+		{
+			if (URCController* Controller = Cast<URCController>(ControllerItem->GetVirtualProperty()))
+			{
+				InContext.ControllerPanelList->CreateBindBehaviourAndAssignTo(Controller, RemoteControlProperty.ToSharedRef(), true);
+				bModified = true;
+			}
+		}
+	}
+
+	if (!bModified)
+	{
+		Transaction.Cancel();
+	}
+}
+
+void SRCControllerItemRow::CreateControllersFromEntities(const FExposedEntityDragDrop& InExposedEntityDragDrop, const FDragDropContext& InContext)
+{
+	FScopedTransaction Transaction(LOCTEXT("AutoBindEntities", "Auto bind entities to controllers"));
+
+	const TArray<FGuid>& SelectedFieldIds = InExposedEntityDragDrop.GetSelectedFieldsId();
+
+	TArray<URCController*, TInlineAllocator<1>> CreatedControllers;
+	TArray<TSharedRef<const FRemoteControlProperty>, TInlineAllocator<1>> SourcePropertyEntities;
+
+	CreatedControllers.Reserve(SelectedFieldIds.Num());
+	SourcePropertyEntities.Reserve(SelectedFieldIds.Num());
+
+	for (const FGuid& ExposedEntityId : SelectedFieldIds)
+	{
+		if (TSharedPtr<const FRemoteControlProperty> RemoteControlProperty = InContext.Preset->GetExposedEntity<FRemoteControlProperty>(ExposedEntityId).Pin())
+		{
+			if (URCController* NewController = UE::RCUIHelpers::CreateControllerFromEntity(InContext.Preset, RemoteControlProperty))
+			{
+				CreatedControllers.Emplace(NewController);
+				SourcePropertyEntities.Add(RemoteControlProperty.ToSharedRef());
+			}
+		}
+	}
+
+	if (CreatedControllers.IsEmpty())
+	{
+		Transaction.Cancel();
+		return;
+	}
+
+	// Store the drop index prior to refresh
+	const int32 DropIndex = InContext.ControllerPanelList->GetDropIndex(ControllerItem, InContext.DropZone);
+
+	// Refresh panel list so that the controller models are created
+	InContext.ControllerPanelList->Reset();
+
+	// Reorder the controller items to match the target drop zone
+	{
+		const TArray<TSharedPtr<FRCControllerModel>> ControllersToMove = InContext.ControllerPanelList->FindControllerItemsByObject(CreatedControllers);
+		InContext.ControllerPanelList->ReorderControllerItems(ControllersToMove, DropIndex);
+	}
+
+	check(CreatedControllers.Num() == SourcePropertyEntities.Num());
+
+	// Create Bind Behaviour and Bind to the property
+	for (int32 Index = 0; Index < CreatedControllers.Num(); ++Index)
+	{
+		constexpr bool bExecuteBind = true;
+		InContext.ControllerPanelList->CreateBindBehaviourAndAssignTo(CreatedControllers[Index], SourcePropertyEntities[Index], bExecuteBind);
+	}
+}
+
+FReply SRCControllerItemRow::OnDragDetected(const FGeometry& InMyGeometry, const FPointerEvent& InMouseEvent)
+{
+	if (FSuperRowType::OnDragDetected(InMyGeometry, InMouseEvent).IsEventHandled())
+	{
+		return FReply::Handled();
+	}
+
+	const TSharedPtr<SRCControllerPanelList> ControllerPanelList = ControllerPanelListWeak.Pin();
+	if (!ControllerPanelList.IsValid())
+	{
+		return FReply::Unhandled();
+	}
+
+	if (InMouseEvent.IsMouseButtonDown(EKeys::LeftMouseButton))
+	{
+		TSharedRef<FRCControllerDragDrop> ControllerDragDrop = MakeShared<FRCControllerDragDrop>(ControllerPanelList->GetSelectedControllers());
+		return FReply::Handled().BeginDragDrop(ControllerDragDrop);
+	}
+	return FReply::Unhandled();
+}
+
+FReply SRCControllerItemRow::OnMouseButtonDoubleClick(const FGeometry& InMyGeometry, const FPointerEvent& InMouseEvent)
+{
+	if (ControllerItem.IsValid())
+	{
+		ControllerItem->EnterDescriptionEditingMode();
+	}
+	return FSuperRowType::OnMouseButtonDoubleClick(InMyGeometry, InMouseEvent);
+}
+
+#undef LOCTEXT_NAMESPACE
diff --git a/Engine/Plugins/VirtualProduction/RemoteControl/Source/RemoteControlUI/Private/UI/Controller/SRCControllerItemRow.h b/Engine/Plugins/VirtualProduction/RemoteControl/Source/RemoteControlUI/Private/UI/Controller/SRCControllerItemRow.h
new file mode 100644
--- /dev/null
+++ b/Engine/Plugins/VirtualProduction/RemoteControl/Source/RemoteControlUI/Private/UI/Controller/SRCControllerItemRow.h
@@ -0,0 +1,63 @@
+// Copyright Epic Games, Inc. All Rights Reserved.
+
+#pragma once
+
+#include "Widgets/Views/SListView.h"
+#include "Widgets/Views/STableRow.h"
+
+class FExposedEntityDragDrop;
+class FRCControllerModel;
+class SRCControllerPanelList;
+class URemoteControlPreset;
+
+struct FRCControllerColumns
+{
+	static const FLazyName TypeColor;
+	static const FLazyName ControllerId;
+	static const FLazyName Description;
+	static const FLazyName Value;
+	static const FLazyName FieldId;
+	static const FLazyName ValueTypeSelection;
+};
+
+class SRCControllerItemRow : public SMultiColumnTableRow<TSharedRef<FRCControllerModel>>
+{
+public:
+	SLATE_BEGIN_ARGS(SRCControllerItemRow) {}
+		SLATE_STYLE_ARGUMENT(FTableRowStyle, Style)
+	SLATE_END_ARGS()
+
+	void Construct(const FArguments& InArgs, const TSharedRef<STableViewBase>& InOwnerTableView
+		, TSharedRef<FRCControllerModel> InControllerItem
+		, TSharedRef<SRCControllerPanelList> InControllerPanelList);
+
+	TSharedRef<SWidget> GenerateWidgetForColumn(const FName& InColumnName) override;
+
+private:
+	/** Called to validate whether a drag drop event is accepted for a controller item row */
+	TOptional<EItemDropZone> CanAcceptDrop(const FDragDropEvent& InDragDropEvent, EItemDropZone InDropZone, TSharedRef<FRCControllerModel> InItem);
+
+	/** Called when a drop happens on this controller item row */
+	FReply OnAcceptDrop(const FDragDropEvent& InDragDropEvent, EItemDropZone InDropZone, TSharedRef<FRCControllerModel> InItem);
+
+	struct FDragDropContext
+	{
+		URemoteControlPreset* Preset;
+		const TSharedPtr<SRCControllerPanelList>& ControllerPanelList;
+		const TSharedRef<FRCControllerModel>& Item;
+		EItemDropZone DropZone;
+	};
+	/** Creates a bind behavior on the existing controller items for all the given entities found in the drag drop */
+	void CreateBindBehaviorsFromEntities(const FExposedEntityDragDrop& InExposedEntityDragDrop, const FDragDropContext& InContext);
+
+	/** Creates new controllers next to this item row (above/below depending on drop zone) and creates a bind behavior for each new controller for all the given entities found in the drag drop */
+	void CreateControllersFromEntities(const FExposedEntityDragDrop& InExposedEntityDragDrop, const FDragDropContext& InContext);
+
+	//~ Begin SWidget
+	virtual FReply OnDragDetected(const FGeometry& InMyGeometry, const FPointerEvent& InMouseEvent) override;
+	virtual FReply OnMouseButtonDoubleClick(const FGeometry& InMyGeometry, const FPointerEvent& InMouseEvent) override;
+	//~ End SWidget
+
+	TSharedPtr<FRCControllerModel> ControllerItem;
+	TWeakPtr<SRCControllerPanelList> ControllerPanelListWeak;
+};
\ No newline at end of file
diff --git a/Engine/Plugins/VirtualProduction/RemoteControl/Source/RemoteControlUI/Private/UI/Controller/SRCControllerPanelList.cpp b/Engine/Plugins/VirtualProduction/RemoteControl/Source/RemoteControlUI/Private/UI/Controller/SRCControllerPanelList.cpp
--- a/Engine/Plugins/VirtualProduction/RemoteControl/Source/RemoteControlUI/Private/UI/Controller/SRCControllerPanelList.cpp
+++ b/Engine/Plugins/VirtualProduction/RemoteControl/Source/RemoteControlUI/Private/UI/Controller/SRCControllerPanelList.cpp
@@ -23,7 +23,9 @@
 #include "RemoteControlField.h"
 #include "RemoteControlPreset.h"
 #include "SDropTarget.h"
+#include "SRCControllerItemRow.h"
 #include "SRCControllerPanel.h"
+#include "ScopedTransaction.h"
 #include "SlateOptMacros.h"
 #include "Styling/RemoteControlStyles.h"
 #include "UI/Action/SRCActionPanelList.h"
@@ -35,307 +37,56 @@
 #include "UI/SRemoteControlPanel.h"
 #include "Widgets/Views/SHeaderRow.h"
 
-#if WITH_EDITOR
-#include "ScopedTransaction.h"
-#endif
-
 #define LOCTEXT_NAMESPACE "SRCControllerPanelList"
 
-namespace UE::RCControllerPanelList
+TArray<TSharedPtr<FRCControllerModel>> FRCControllerDragDrop::ResolveControllers() const
 {
-	enum class EDragDropSupportedModes : uint8
-	{
-		ReorderOnly,
-		All,
-		None
-	};
-
-	namespace Columns
-	{
-		const FName TypeColor = TEXT("TypeColor");
-		const FName ControllerId = TEXT("Controller Id");
-		const FName Description = TEXT("Controller Description");
-		const FName Value = TEXT("Controller Value");
-		const FName DragHandle = TEXT("Drag Handle");
-		const FName FieldId = TEXT("Controller Field Id");
-		const FName ValueTypeSelection = TEXT("Value Type Selection");
-	}
+	TArray<TSharedPtr<FRCControllerModel>> ResolvedControllers;
+	ResolvedControllers.Reserve(ControllersWeak.Num());
 
-	class SControllerItemListRow : public SMultiColumnTableRow<TSharedRef<FRCControllerModel>>
+	for (const TWeakPtr<FRCControllerModel>& WeakController : ControllersWeak)
 	{
-	public:
-		void Construct(const FTableRowArgs& InArgs, const TSharedRef<STableViewBase>& OwnerTableView, TSharedRef<FRCControllerModel> InControllerItem, TSharedRef<SRCControllerPanelList> InControllerPanelList)
-		{
-			ControllerItem = InControllerItem;
-			ControllerPanelListWeakPtr = InControllerPanelList.ToWeakPtr();
-			FSuperRowType::Construct(InArgs, OwnerTableView);
-		}
-
-		TSharedRef<SWidget> GenerateWidgetForColumn(const FName& ColumnName) override
-		{
-			if (!ensure(ControllerItem.IsValid()))
-				return SNullWidget::NullWidget;
-
-			if (ColumnName == UE::RCControllerPanelList::Columns::TypeColor)
-			{
-				if (URCController* Controller = Cast< URCController>(ControllerItem->GetVirtualProperty()))
-				{
-					if (const FProperty* Property = Controller->GetProperty())
-					{
-						return UE::RCUIHelpers::GetTypeColorWidget(Property);
-					}
-				}
-			}
-			else if (ColumnName == UE::RCControllerPanelList::Columns::FieldId)
-			{
-				return WrapWithDropTarget(ControllerItem->GetFieldIdWidget());
-			}
-			else if (ColumnName == UE::RCControllerPanelList::Columns::ValueTypeSelection)
-			{
-				return WrapWithDropTarget(ControllerItem->GetTypeSelectionWidget());
-			}
-			else if (ColumnName == UE::RCControllerPanelList::Columns::ControllerId)
-			{
-				return WrapWithDropTarget(ControllerItem->GetNameWidget());
-			}
-			else if (ColumnName == UE::RCControllerPanelList::Columns::Description)
-			{
-				return WrapWithDropTarget(ControllerItem->GetDescriptionWidget());
-			}
-			else if (ColumnName == UE::RCControllerPanelList::Columns::Value)
-			{
-				return ControllerItem->GetWidget();
-			}
-			else if (ColumnName == UE::RCControllerPanelList::Columns::DragHandle)
-			{
-				SAssignNew(DragDropBorderWidget, SBorder)
-					.BorderImage(FRemoteControlPanelStyle::Get()->GetBrush("RemoteControlPanel.ExposedFieldBorder"));
-
-				TSharedRef<SWidget> DragHandleWidget = 
-					SNew(SBox)
-					.Padding(5.f)
-					[
-						SNew(SRCPanelDragHandle<FRCControllerDragDrop>, ControllerItem->GetId())
-						.Widget(DragDropBorderWidget)
-					];
-
-				return WrapWithDropTarget(DragHandleWidget, EDragDropSupportedModes::ReorderOnly /*Allow reorder operation only*/);
-			}
-			else if (const TSharedPtr<SRCControllerPanelList> ControlPanelList = ControllerPanelListWeakPtr.Pin())
-			{
-				if (ControlPanelList->GetCustomColumns().Contains(ColumnName))
-				{
-					return ControllerItem->GetControllerExtensionWidget(ColumnName);
-				}
-			}
-
-			return SNullWidget::NullWidget;
-		}
-
-		TSharedPtr<SBorder> DragDropBorderWidget;
-
-
-		void OnDragEnter(FGeometry const& MyGeometry, FDragDropEvent const& DragDropEvent) override
-		{
-			bIsDragActive = true;
-		}
-
-
-		void OnDragLeave(FDragDropEvent const& DragDropEvent) override
-		{
-			bIsDragActive = false;
-		}
-
-	protected:
-
-	private:
-
-		TSharedRef<SWidget> WrapWithDropTarget(const TSharedRef<SWidget> InWidget, const EDragDropSupportedModes SupportedMode = EDragDropSupportedModes::All)
-		{
-			return SNew(SDropTarget)
-				.ValidColor(FStyleColors::AccentOrange)
-				.VerticalImage(FRemoteControlPanelStyle::Get()->GetBrush("RemoteControlPanel.VerticalDash"))
-				.HorizontalImage(FRemoteControlPanelStyle::Get()->GetBrush("RemoteControlPanel.HorizontalDash"))
-				.OnDropped_Lambda([this](const FGeometry& InGeometry, const FDragDropEvent& InDragDropEvent) { return SControllerItemListRow::OnControllerItemDragDrop(InDragDropEvent.GetOperation()); })
-				.OnAllowDrop(this, &SControllerItemListRow::OnAllowDrop, SupportedMode)
-				.OnIsRecognized(this, &SControllerItemListRow::OnAllowDrop, SupportedMode)
-				[
-					InWidget
-				];
-		}
-
-		FReply OnControllerItemDragDrop(TSharedPtr<FDragDropOperation> DragDropOperation)
-		{
-			check(ControllerPanelListWeakPtr.IsValid());
-			const TSharedPtr<SRCControllerPanelList> ControllerPanelList = ControllerPanelListWeakPtr.Pin();
-			
-			bIsDragActive = false;
-			ControllerPanelList->bIsAnyControllerItemEligibleForDragDrop = false;
-
-			if (!DragDropOperation)
-			{
-				return FReply::Handled();
-			}
-
-			if (DragDropOperation->IsOfType<FRCControllerDragDrop>())
-			{
-				if (TSharedPtr<FRCControllerDragDrop> DragDropOp = StaticCastSharedPtr<FRCControllerDragDrop>(DragDropOperation))
-				{
-					const FGuid DragDropControllerId = DragDropOp->GetId();
-
-					if (ControllerPanelList.IsValid())
-					{
-						TSharedPtr<FRCControllerModel> DragDropControllerItem = ControllerPanelList->FindControllerItemById(DragDropControllerId);
-
-						if (ensure(ControllerItem && DragDropControllerItem))
-						{
-							ControllerPanelList->ReorderControllerItem(DragDropControllerItem.ToSharedRef(), ControllerItem.ToSharedRef());
-						}
-					}
-				}
-			}
-			else if (DragDropOperation->IsOfType<FExposedEntityDragDrop>())
-			{
-				if (TSharedPtr<FExposedEntityDragDrop> DragDropOp = StaticCastSharedPtr<FExposedEntityDragDrop>(DragDropOperation))
-				{
-					// Fetch the Exposed Entity
-					const TArray<FGuid>& ExposedEntitiesIds = DragDropOp->GetSelectedFieldsId();
-
-					if (ExposedEntitiesIds.Num() == 1)
-					{
-						if (URemoteControlPreset* Preset = ControllerPanelList->GetPreset())
-						{
-							if (TSharedPtr<const FRemoteControlProperty> RemoteControlProperty = Preset->GetExposedEntity<FRemoteControlProperty>(ExposedEntitiesIds[0]).Pin())
-							{
-								if (URCController* Controller = Cast<URCController>(ControllerItem->GetVirtualProperty()))
-								{
-									ControllerPanelList->CreateBindBehaviourAndAssignTo(Controller, RemoteControlProperty.ToSharedRef(), true);
-								}
-							}
-						}
-					}
-				}
-			}
-
-			return FReply::Handled();
-		}
-
-		bool OnAllowDrop(TSharedPtr<FDragDropOperation> DragDropOperation, const EDragDropSupportedModes SupportedMode)
-		{
-			if (!ensure(ControllerPanelListWeakPtr.IsValid()))
-			{
-				return false;
-			}
-			
-			const TSharedPtr<SRCControllerPanelList> ControllerPanelList = ControllerPanelListWeakPtr.Pin();
-
-			if (DragDropOperation && ControllerItem)
-			{
-				// Dragging Controllers onto Controllers (Reordering)
-				if (DragDropOperation->IsOfType<FRCControllerDragDrop>())
-				{
-					if (TSharedPtr<FRCControllerDragDrop> DragDropOp = StaticCastSharedPtr<FRCControllerDragDrop>(DragDropOperation))
-					{
-						return true;
-					}
-				}
-				else
-				{
-					if (SupportedMode == EDragDropSupportedModes::ReorderOnly)
-					{
-						return false; // This widget only supports Reordering operation as a drag-drop action
-					}
-
-					if (DragDropOperation->IsOfType<FExposedEntityDragDrop>())
-					{
-						if (TSharedPtr<FExposedEntityDragDrop> DragDropOp = StaticCastSharedPtr<FExposedEntityDragDrop>(DragDropOperation))
-						{
-							// Fetch the Exposed Entity
-							const TArray<FGuid>& ExposedEntitiesIds = DragDropOp->GetSelectedFieldsId();
-
-							if (ExposedEntitiesIds.Num() == 1)
-							{
-								if (URemoteControlPreset* Preset = ControllerPanelList->GetPreset())
-								{
-									if (TSharedPtr<const FRemoteControlField> RemoteControlField = Preset->GetExposedEntity<FRemoteControlField>(ExposedEntitiesIds[0]).Pin())
-									{
-										if (URCController* Controller = Cast<URCController>(ControllerItem->GetVirtualProperty()))
-										{
-											const bool bAllowNumericInputAsStrings = true;
-
-											const bool bAllowDrop = URCBehaviourBind::CanHaveActionForField(Controller, RemoteControlField.ToSharedRef(), bAllowNumericInputAsStrings);
-
-											ControllerPanelList->bIsAnyControllerItemEligibleForDragDrop |= (bAllowDrop && bIsDragActive);
-
-											return bAllowDrop;
-										}
-									}
-								}
-							}
-						}
-					}
-				}
-			}
-
-			return false;
-		}
-
-		//~ SWidget Interface
-		virtual FReply OnMouseButtonDoubleClick(const FGeometry& InMyGeometry, const FPointerEvent& InMouseEvent) override
+		if (TSharedPtr<FRCControllerModel> ResolvedController = WeakController.Pin())
 		{
-			if (ControllerItem.IsValid())
-			{
-				ControllerItem->EnterDescriptionEditingMode();
-			}
-
-			return FSuperRowType::OnMouseButtonDoubleClick(InMyGeometry, InMouseEvent);
+			ResolvedControllers.Add(ResolvedController);
 		}
-
-private:
-		TSharedPtr<FRCControllerModel> ControllerItem;
-		TWeakPtr<SRCControllerPanelList> ControllerPanelListWeakPtr;
-		bool bIsDragActive = false;
-	};
-} 
+	}
+	return ResolvedControllers;
+}
 
 void SRCControllerPanelList::Construct(const FArguments& InArgs, const TSharedRef<SRCControllerPanel> InControllerPanel, const TSharedRef<SRemoteControlPanel> InRemoteControlPanel)
 {
 	SRCLogicPanelListBase::Construct(SRCLogicPanelListBase::FArguments(), InControllerPanel, InRemoteControlPanel);
 	
 	ControllerPanelWeakPtr = InControllerPanel;
-	
+
 	RCPanelStyle = &FRemoteControlPanelStyle::Get()->GetWidgetStyle<FRCPanelStyle>("RemoteControlPanel.LogicControllersPanel");
 
 	ListView = SNew(SListView<TSharedPtr<FRCControllerModel>>)
 		.ListItemsSource(&ControllerItems)
 		.OnSelectionChanged(this, &SRCControllerPanelList::OnTreeSelectionChanged)
 		.OnGenerateRow(this, &SRCControllerPanelList::OnGenerateWidgetForList)
-		.SelectionMode(ESelectionMode::Single) // Current setup supports only single selection (and related display) of a Controller in the list
+		.SelectionMode(ESelectionMode::Multi)
 		.OnContextMenuOpening(this, &SRCLogicPanelListBase::GetContextMenuWidget)
 		.HeaderRow(
 			SAssignNew(ControllersHeaderRow, SHeaderRow)
 			.Style(&RCPanelStyle->HeaderRowStyle)
 
-			+ SHeaderRow::Column(UE::RCControllerPanelList::Columns::TypeColor)
+			+ SHeaderRow::Column(FRCControllerColumns::TypeColor)
 			.DefaultLabel(FText())
 			.FixedWidth(15)
 			.HeaderContentPadding(RCPanelStyle->HeaderRowPadding)
 
-			+ SHeaderRow::Column(UE::RCControllerPanelList::Columns::DragHandle)
-			.DefaultLabel(FText())
-			.FixedWidth(15)
-			.HeaderContentPadding(RCPanelStyle->HeaderRowPadding)
-
-			+ SHeaderRow::Column(UE::RCControllerPanelList::Columns::ControllerId)
+			+ SHeaderRow::Column(FRCControllerColumns::ControllerId)
 			.DefaultLabel(LOCTEXT("ControllerIdColumnName", "Controller Id"))
 			.FillWidth(0.2f)
 			.HeaderContentPadding(RCPanelStyle->HeaderRowPadding)
 
-			+ SHeaderRow::Column(UE::RCControllerPanelList::Columns::Description)
+			+ SHeaderRow::Column(FRCControllerColumns::Description)
 			.DefaultLabel(LOCTEXT("ControllerNameColumnDescription", "Description"))
 			.FillWidth(0.35f)
 
-			+ SHeaderRow::Column(UE::RCControllerPanelList::Columns::Value)
+			+ SHeaderRow::Column(FRCControllerColumns::Value)
 			.DefaultLabel(LOCTEXT("ControllerValueColumnName", "Input"))
 			.FillWidth(0.45f)
 			.HeaderContentPadding(RCPanelStyle->HeaderRowPadding)
@@ -370,7 +121,7 @@ void SRCControllerPanelList::Construct(const FArguments& InArgs, const TSharedRe
 	Args.bShouldShowHiddenProperties = true;
 	Args.NotifyHook = this;
 	PropertyRowGenerator = FModuleManager::GetModuleChecked<FPropertyEditorModule>("PropertyEditor").CreatePropertyRowGenerator(Args);
-	
+
 	Reset();
 }
 
@@ -545,11 +296,8 @@ void SRCControllerPanelList::Reset()
 
 TSharedRef<ITableRow> SRCControllerPanelList::OnGenerateWidgetForList(TSharedPtr<FRCControllerModel> InItem, const TSharedRef<STableViewBase>& OwnerTable)
 {
-	typedef UE::RCControllerPanelList::SControllerItemListRow SControllerRowType;
-
-	return SNew(SControllerRowType, OwnerTable, InItem.ToSharedRef(), SharedThis(this))
-		.Style(&RCPanelStyle->TableRowStyle)
-		.Padding(FMargin(4.5f));
+	return SNew(SRCControllerItemRow, OwnerTable, InItem.ToSharedRef(), SharedThis(this))
+		.Style(&RCPanelStyle->TableRowStyle);
 }
 
 void SRCControllerPanelList::OnTreeSelectionChanged(TSharedPtr<FRCControllerModel> InItem, ESelectInfo::Type InSelectInfo)
@@ -558,7 +306,9 @@ void SRCControllerPanelList::OnTreeSelectionChanged(TSharedPtr<FRCControllerMode
 	{
 		if (const TSharedPtr<SRemoteControlPanel> RemoteControlPanel = ControllerPanel->GetRemoteControlPanel())
 		{
-			if (InItem != SelectedControllerItemWeakPtr.Pin())
+			// Tree selection gets updated and notifies with the first element in the selection set.
+			// Controller Behavior / Actions only support viewing one behavior at a time
+			if (InItem != SelectedControllerItemWeakPtr)
 			{
 				SelectedControllerItemWeakPtr = InItem;
 				RemoteControlPanel->OnControllerSelectionChanged.Broadcast(InItem, InSelectInfo);
@@ -732,14 +482,14 @@ void SRCControllerPanelList::DeleteSelectedPanelItems()
 	DeleteItemsFromLogicPanel<FRCControllerModel>(ControllerItems, ListView->GetSelectedItems());
 }
 
+TArray<TSharedPtr<FRCControllerModel>> SRCControllerPanelList::GetSelectedControllers() const
+{
+	return ListView->GetSelectedItems();
+}
+
 TArray<TSharedPtr<FRCLogicModeBase>> SRCControllerPanelList::GetSelectedLogicItems()
 {
-	// Controllers don't support multi selection
-	if (const TSharedPtr<FRCLogicModeBase> SelectedControllerItemPtr = SelectedControllerItemWeakPtr.Pin())
-	{
-		return { SelectedControllerItemPtr };
-	}
-	return TArray<TSharedPtr<FRCLogicModeBase>>();
+	return TArray<TSharedPtr<FRCLogicModeBase>>(ListView->GetSelectedItems());
 }
 
 void SRCControllerPanelList::NotifyPreChange(FEditPropertyChain* PropertyAboutToChange)
@@ -776,27 +526,98 @@ TSharedPtr<FRCControllerModel> SRCControllerPanelList::FindControllerItemById(co
 	return nullptr;
 }
 
-void SRCControllerPanelList::ReorderControllerItem(TSharedRef<FRCControllerModel> ItemToMove, TSharedRef<FRCControllerModel> AnchorItem)
+TArray<TSharedPtr<FRCControllerModel>> SRCControllerPanelList::FindControllerItemsById(TConstArrayView<FGuid> InIds)
+{
+	TArray<TSharedPtr<FRCControllerModel>> FoundControllerItems;
+	FoundControllerItems.Reserve(InIds.Num());
+
+	for (const TSharedPtr<FRCControllerModel>& ControllerItem : ControllerItems)
+	{
+		if (ControllerItem.IsValid() && InIds.Contains(ControllerItem->GetId()))
+		{
+			FoundControllerItems.Add(ControllerItem);
+		}
+	}
+
+	return FoundControllerItems;
+}
+
+TArray<TSharedPtr<FRCControllerModel>> SRCControllerPanelList::FindControllerItemsByObject(TConstArrayView<URCController*> InControllers)
+{
+	TArray<TSharedPtr<FRCControllerModel>> FoundControllerItems;
+	FoundControllerItems.Reserve(InControllers.Num());
+
+	for (const TSharedPtr<FRCControllerModel>& ControllerItem : ControllerItems)
+	{
+		if (ControllerItem.IsValid() && InControllers.Contains(ControllerItem->GetVirtualProperty()))
+		{
+			FoundControllerItems.Add(ControllerItem);
+		}
+	}
+
+	return FoundControllerItems;
+}
+
+int32 SRCControllerPanelList::GetDropIndex(const TSharedPtr<FRCControllerModel>& InItem, EItemDropZone InDropZone) const
 {
-	int32 Index = ControllerItems.Find(AnchorItem);
+	int32 Index = ControllerItems.Find(InItem);
+	if (Index == INDEX_NONE)
+	{
+		return INDEX_NONE;
+	}
 
-	// Update UI list
-	ControllerItems.RemoveSingle(ItemToMove);
-	ControllerItems.Insert(ItemToMove, Index);
+	if (InDropZone == EItemDropZone::BelowItem)
+	{
+		++Index;
+	}
+	return Index;
+}
+
+bool SRCControllerPanelList::ReorderControllerItems(TConstArrayView<TSharedPtr<FRCControllerModel>> InItemsToMove, int32 InTargetIndex)
+{
+	// Only allow valid indices and the index 1 more of the end index (i.e. index == num)
+	if (InTargetIndex == INDEX_NONE || InTargetIndex > ControllerItems.Num())
+	{
+		return false;
+	}
+
+	// Remove the items to move.
+	// if these items come before the target index, the target index needs to shift to compensate for the removed items
+	for (TArray<TSharedPtr<FRCControllerModel>>::TIterator Iter(ControllerItems); Iter; ++Iter)
+	{
+		if (InItemsToMove.Contains(*Iter))
+		{
+			if (Iter.GetIndex() < InTargetIndex)
+			{
+				--InTargetIndex;
+			}
+			Iter.RemoveCurrent();
+		}
+	}
+
+	// Update UI list. Remove all the items to move to then insert them at the desired index
+	ControllerItems.RemoveAll([&InItemsToMove](const TSharedPtr<FRCControllerModel>& InController)
+		{
+			return InItemsToMove.Contains(InController);
+		});
+
+	ControllerItems.Insert(InItemsToMove.GetData(), InItemsToMove.Num(), InTargetIndex);
 
 	// Update display indices
-	for (int32 i = 0; i < ControllerItems.Num(); i++)
+	for (int32 Index = 0; Index < ControllerItems.Num(); Index++)
 	{
-		if (ControllerItems[i])
+		if (ControllerItems[Index])
 		{
-			if (URCVirtualPropertyBase* Controller = ControllerItems[i]->GetVirtualProperty())
+			if (URCVirtualPropertyBase* Controller = ControllerItems[Index]->GetVirtualProperty())
 			{
-				Controller->SetDisplayIndex(i);
+				Controller->Modify();
+				Controller->SetDisplayIndex(Index);
 			}
 		}
 	}
 
 	ListView->RequestListRefresh();
+	return true;
 }
 
 static TSharedPtr<FExposedEntityDragDrop> GetExposedEntityDragDrop(TSharedPtr<FDragDropOperation> DragDropOperation)
@@ -818,16 +639,7 @@ bool SRCControllerPanelList::IsEntitySupported(const FGuid ExposedEntityId)
 	{
 		if (const TSharedPtr<const FRemoteControlProperty>& RemoteControlProperty = Preset->GetExposedEntity<FRemoteControlProperty>(ExposedEntityId).Pin())
 		{
-			if (!RemoteControlProperty->IsEditable())
-			{
-				// Property with error(s)
-				return false;
-			}
-
-			if (RemoteControlProperty->FieldType == EExposedFieldType::Property)
-			{
-				return UE::RCControllers::CanCreateControllerFromProperty(RemoteControlProperty->GetProperty());
-			}
+			return UE::RCControllers::CanCreateControllerFromEntity(RemoteControlProperty);
 		}
 	}
 
@@ -894,57 +706,17 @@ FReply SRCControllerPanelList::OnControllerListViewDragDrop(TSharedPtr<FDragDrop
 	return FReply::Handled();
 }
 
-bool SRCControllerPanelList::CreateAutoBindForProperty(TSharedPtr<const FRemoteControlProperty> RemoteControlProperty)
+bool SRCControllerPanelList::CreateAutoBindForProperty(TSharedPtr<const FRemoteControlProperty> InRemoteControlProperty)
 {
-	if (URemoteControlPreset* Preset = GetPreset())
+	if (URCController* NewController = UE::RCUIHelpers::CreateControllerFromEntity(GetPreset(), InRemoteControlProperty))
 	{
-		// Derive the input data needed for creating a new Controller
-		FProperty* Property = RemoteControlProperty->GetProperty();
-		EPropertyBagPropertyType PropertyBagType = EPropertyBagPropertyType::None;
-		UObject* StructObject = nullptr;
-
-		// In the Logic realm we use a single type like (eg: String / Int) to represent various related types (String/Name/Text, Int32, Int64, etc)
-		// For this reason explicit mapping conversion is required between a given FProperty type and the desired Controller type
-		bool bSuccess = URCBehaviourBind::GetPropertyBagTypeFromFieldProperty(Property, PropertyBagType, StructObject);
-		if(bSuccess)
-		{
-			// Preparation step, in case we are dealing with a custom controller
-			FString CustomControllerName = TEXT("");
-			if (StructObject == UTexture::StaticClass() || StructObject == UTexture2D::StaticClass())
-			{
-				if (PropertyBagType == EPropertyBagPropertyType::String)
-				{
-					StructObject = nullptr;
-					CustomControllerName = UE::RCCustomControllers::CustomTextureControllerName;
-				}
-			}
-			
-			// Step 1. Create a Controller of matching type
-			URCController* NewController = Cast<URCController>(Preset->AddController(URCController::StaticClass(), PropertyBagType, StructObject, RemoteControlProperty->FieldName));
-			NewController->DisplayIndex = Preset->GetNumControllers() - 1;
-
-			// Add metadata to this controller, if this is a custom controller
-			if (!CustomControllerName.IsEmpty())
-			{
-				const TMap<FName, FString>& CustomControllerMetaData = UE::RCCustomControllers::GetCustomControllerMetaData(CustomControllerName);
-				for (const TPair<FName, FString>& Pair : CustomControllerMetaData)
-				{
-					NewController->SetMetadataValue(Pair.Key, Pair.Value);
-				}
-			}
-
-			// Transfer property value from Exposed Property to the New Controller.
-			// The goal is to keep values synced for a Controller newly created via "Auto Bind"
-			URCBehaviourBind::CopyPropertyValueToController(NewController, RemoteControlProperty.ToSharedRef());
-
-			// Step 2. Refresh UI
-			Reset();
+		// Refresh UI
+		Reset();
 
-			// Step 3. Create Bind Behaviour and Bind to the property
-			constexpr bool bExecuteBind = true;
-			CreateBindBehaviourAndAssignTo(NewController, RemoteControlProperty.ToSharedRef(), bExecuteBind);
-			return true;
-		}
+		// Create Bind Behaviour and Bind to the property
+		constexpr bool bExecuteBind = true;
+		CreateBindBehaviourAndAssignTo(NewController, InRemoteControlProperty.ToSharedRef(), bExecuteBind);
+		return true;
 	}
 	return false;
 }
@@ -995,6 +767,8 @@ void SRCControllerPanelList::CreateBindBehaviourAndAssignTo(URCController* Contr
 
 	if (!BindBehaviour)
 	{
+		Controller->Modify();
+
 		URCBehaviourBind* NewBindBehaviour = Cast<URCBehaviourBind>(Controller->AddBehaviour(URCBehaviourBindNode::StaticClass()));
 
 		// If this is a new Bind Behaviour and we are attempting to link unrelated (but compatible types), via numeric conversion, then we apply the bAllowNumericInputAsStrings flag
@@ -1015,6 +789,8 @@ void SRCControllerPanelList::CreateBindBehaviourAndAssignTo(URCController* Contr
 	if (ensure(BindBehaviour))
 	{
 		URCBehaviourBind* BindBehaviourMutable = const_cast<URCBehaviourBind*>(BindBehaviour);
+
+		BindBehaviourMutable->Modify();
 		URCAction* BindAction = BindBehaviourMutable->AddPropertyBindAction(InRemoteControlProperty);
 
 		if (bExecuteBind)
@@ -1036,14 +812,14 @@ void SRCControllerPanelList::ShowValueTypeHeaderColumn(bool bInShowColumn)
 {
 	if (ControllersHeaderRow.IsValid())
 	{
-		const bool bColumnIsGenerated = ControllersHeaderRow->IsColumnGenerated(UE::RCControllerPanelList::Columns::ValueTypeSelection);
+		const bool bColumnIsGenerated = ControllersHeaderRow->IsColumnGenerated(FRCControllerColumns::ValueTypeSelection);
 		if (bInShowColumn)
 		{
 			if (!bColumnIsGenerated)
 			{
 				ControllersHeaderRow->AddColumn(
 					SHeaderRow::FColumn::FArguments()
-					.ColumnId(UE::RCControllerPanelList::Columns::ValueTypeSelection)
+					.ColumnId(FRCControllerColumns::ValueTypeSelection)
 					.DefaultLabel(LOCTEXT("ControllerValueTypeColumnName", "Value Type"))
 					.FillWidth(0.2f)
 					.HeaderContentPadding(RCPanelStyle->HeaderRowPadding)
@@ -1052,7 +828,7 @@ void SRCControllerPanelList::ShowValueTypeHeaderColumn(bool bInShowColumn)
 		}
 		else if (bColumnIsGenerated)
 		{
-			ControllersHeaderRow->RemoveColumn(UE::RCControllerPanelList::Columns::ValueTypeSelection);
+			ControllersHeaderRow->RemoveColumn(FRCControllerColumns::ValueTypeSelection);
 		}
 	}
 }
@@ -1061,14 +837,14 @@ void SRCControllerPanelList::ShowFieldIdHeaderColumn(bool bInShowColumn)
 {	
 	if (ControllersHeaderRow.IsValid())
 	{
-		const bool bColumnIsGenerated = ControllersHeaderRow->IsColumnGenerated(UE::RCControllerPanelList::Columns::FieldId);
+		const bool bColumnIsGenerated = ControllersHeaderRow->IsColumnGenerated(FRCControllerColumns::FieldId);
 		if (bInShowColumn)
 		{
 			if (!bColumnIsGenerated)
 			{
 				ControllersHeaderRow->InsertColumn(
 					SHeaderRow::FColumn::FArguments()
-					.ColumnId(UE::RCControllerPanelList::Columns::FieldId)
+					.ColumnId(FRCControllerColumns::FieldId)
 					.DefaultLabel(LOCTEXT("ControllerNameColumnFieldId", "Field Id"))
 					.FillWidth(0.2f)
 					.HeaderContentPadding(RCPanelStyle->HeaderRowPadding),
@@ -1077,7 +853,7 @@ void SRCControllerPanelList::ShowFieldIdHeaderColumn(bool bInShowColumn)
 		}
 		else if (bColumnIsGenerated)
 		{
-			ControllersHeaderRow->RemoveColumn(UE::RCControllerPanelList::Columns::FieldId);
+			ControllersHeaderRow->RemoveColumn(FRCControllerColumns::FieldId);
 		}
 	}
 }
diff --git a/Engine/Plugins/VirtualProduction/RemoteControl/Source/RemoteControlUI/Private/UI/Controller/SRCControllerPanelList.h b/Engine/Plugins/VirtualProduction/RemoteControl/Source/RemoteControlUI/Private/UI/Controller/SRCControllerPanelList.h
--- a/Engine/Plugins/VirtualProduction/RemoteControl/Source/RemoteControlUI/Private/UI/Controller/SRCControllerPanelList.h
+++ b/Engine/Plugins/VirtualProduction/RemoteControl/Source/RemoteControlUI/Private/UI/Controller/SRCControllerPanelList.h
@@ -9,24 +9,24 @@
 #include "Widgets/Layout/SBorder.h"
 
 class FDragDropOperation;
-struct FRCPanelStyle;
 class FEditPropertyChain;
 class FRCControllerModel;
 class FRCLogicModeBase;
 class IPropertyRowGenerator;
-class ITableRow;
 class ITableBase;
+class ITableRow;
 class SDropTarget;
 class SRCControllerPanel;
 class SRemoteControlPanel;
 class STableViewBase;
 class URCController;
 class URemoteControlPreset;
+enum class EItemDropZone;
+struct FRCPanelStyle;
 template <typename ItemType> class SListView;
 
 /*
 * ~ FRCControllerDragDrop ~
-*
 * Facilitates drag-drop operation for Controller row drag handles
 */
 class FRCControllerDragDrop final : public FDecoratedDragDropOp
@@ -34,17 +34,12 @@ class FRCControllerDragDrop final : public FDecoratedDragDropOp
 public:
 	DRAG_DROP_OPERATOR_TYPE(FRCControllerDragDropOp, FDecoratedDragDropOp)
 
-	using WidgetType = SWidget;
-
-	FRCControllerDragDrop(TSharedPtr<SWidget> InWidget, const FGuid& InId)
-		: Id(InId)
+	explicit FRCControllerDragDrop(const TArray<TSharedPtr<FRCControllerModel>>& InControllers)
+		: ControllersWeak(InControllers)
 	{
 	}
 
-	FGuid GetId() const
-	{
-		return Id;
-	}
+	TArray<TSharedPtr<FRCControllerModel>> ResolveControllers() const;
 
 	virtual void OnDrop(bool bDropWasHandled, const FPointerEvent& MouseEvent) override
 	{
@@ -52,7 +47,7 @@ public:
 	}
 
 private:
-	FGuid Id;
+	TArray<TWeakPtr<FRCControllerModel>> ControllersWeak;
 };
 
 /*
@@ -88,6 +83,9 @@ public:
 	/** Deletes currently selected items from the list view */
 	virtual void DeleteSelectedPanelItems() override;
 
+	/** Returns the controllers selected by the user (if any). */
+	TArray<TSharedPtr<FRCControllerModel>> GetSelectedControllers() const;
+
 	/** Returns the UI items currently selected by the user (if any). */
 	virtual TArray<TSharedPtr<FRCLogicModeBase>> GetSelectedLogicItems() override;
 
@@ -105,15 +103,20 @@ public:
 	/** Finds a Controller UI model by unique Id */
 	TSharedPtr<FRCControllerModel> FindControllerItemById(const FGuid& InId) const;
 
-	/** Given an item to move and an anchor row this function moves the item to the position of the anchor
-	* and pushes all other rows below */
-	void ReorderControllerItem(TSharedRef<FRCControllerModel> ItemToMove, TSharedRef<FRCControllerModel> AnchorItem);
+	/** Finds the Controller UI models by unique Id */
+	TArray<TSharedPtr<FRCControllerModel>> FindControllerItemsById(TConstArrayView<FGuid> InIds);
 
-	/** Returns the currently selected Controller UI Item */
-	TSharedPtr<FRCControllerModel> GetSelectedControllerItem() const
-	{
-		return SelectedControllerItemWeakPtr.Pin();
-	}
+	/** Finds the Controller UI models by their controllers */
+	TArray<TSharedPtr<FRCControllerModel>> FindControllerItemsByObject(TConstArrayView<URCController*> InControllers);
+
+	/** Finds the index of the controller item */
+	int32 GetDropIndex(const TSharedPtr<FRCControllerModel>& InItem, EItemDropZone InDropZone) const;
+
+	/**
+	 * Called to moves the given controller items to the given target index, pushing all other rows below
+	 * returns true if the operation was successful
+	 */
+	bool ReorderControllerItems(TConstArrayView<TSharedPtr<FRCControllerModel>> InItemsToMove, int32 InTargetIndex);
 
 	/** Requests the panel to refresh its contents from the latest list of Controllers */
 	void RequestRefresh()
@@ -150,8 +153,6 @@ public:
 	* This flag is set if any Controller has active drag-drop focus, in which case we disable the first drag-drop zone. This is purely for visual clarity*/
 	bool bIsAnyControllerItemEligibleForDragDrop = false;
 
-private:
-
 	/** OnGenerateRow delegate for the Actions List View */
 	TSharedRef<ITableRow> OnGenerateWidgetForList( TSharedPtr<FRCControllerModel> InItem, const TSharedRef<STableViewBase>& OwnerTable );
 	
@@ -189,7 +190,7 @@ private:
 	void OnControllerValueChanged(TSharedPtr<FRCControllerModel> InControllerModel, bool bInIsMultiController);
 
 	/** Creates a new Controller for the given Remote Control Property and also binds to it. Returns true if operation succeeded, false if nothing happened */
-	bool CreateAutoBindForProperty(TSharedPtr<const FRemoteControlProperty> RemoteControlProperty);
+	bool CreateAutoBindForProperty(TSharedPtr<const FRemoteControlProperty> InRemoteControlProperty);
 
 	/** The row generator used to represent each Controller as a row, when used with SListView */
 	TSharedPtr<IPropertyRowGenerator> PropertyRowGenerator;
diff --git a/Engine/Plugins/VirtualProduction/RemoteControl/Source/RemoteControlUI/Private/UI/RCUIHelpers.cpp b/Engine/Plugins/VirtualProduction/RemoteControl/Source/RemoteControlUI/Private/UI/RCUIHelpers.cpp
--- a/Engine/Plugins/VirtualProduction/RemoteControl/Source/RemoteControlUI/Private/UI/RCUIHelpers.cpp
+++ b/Engine/Plugins/VirtualProduction/RemoteControl/Source/RemoteControlUI/Private/UI/RCUIHelpers.cpp
@@ -2,10 +2,13 @@
 
 #include "RCUIHelpers.h"
 
+#include "Behaviour/Builtin/Bind/RCBehaviourBind.h"
 #include "Controller/RCController.h"
-#include "Editor.h"
-#include "EdGraphSchema_K2.h"
+#include "Controller/RCCustomControllerUtilities.h"
 #include "EdGraph/EdGraphPin.h"
+#include "EdGraphSchema_K2.h"
+#include "Editor.h"
+#include "Engine/Texture2D.h"
 #include "Framework/Application/SlateApplication.h"
 #include "GraphEditorSettings.h"
 #include "IDetailTreeNode.h"
@@ -14,6 +17,8 @@
 #include "PropertyEditorModule.h"
 #include "PropertyHandle.h"
 #include "RCVirtualProperty.h"
+#include "RemoteControlField.h"
+#include "RemoteControlPreset.h"
 #include "TimerManager.h"
 #include "Widgets/Layout/SBorder.h"
 #include "Widgets/Layout/SBox.h"
@@ -213,3 +218,53 @@ TSharedRef<SWidget> UE::RCUIHelpers::GetTypeColorWidget(const FProperty* InPrope
 
 	return TypeColorWidget;
 }
+
+URCController* UE::RCUIHelpers::CreateControllerFromEntity(URemoteControlPreset* InPreset, const TSharedPtr<const FRemoteControlProperty>& InRemoteControlProperty)
+{
+	if (!InPreset || !InRemoteControlProperty.IsValid())
+	{
+		return nullptr;
+	}
+
+	// Derive the input data needed for creating a new Controller
+	FProperty* Property = InRemoteControlProperty->GetProperty();
+	EPropertyBagPropertyType PropertyBagType = EPropertyBagPropertyType::None;
+	UObject* StructObject = nullptr;
+
+	// In the Logic realm we use a single type like (eg: String / Int) to represent various related types (String/Name/Text, Int32, Int64, etc)
+	// For this reason explicit mapping conversion is required between a given FProperty type and the desired Controller type
+	if (!URCBehaviourBind::GetPropertyBagTypeFromFieldProperty(Property, PropertyBagType, StructObject))
+	{
+		return nullptr;
+	}
+
+	// Preparation step, in case we are dealing with a custom controller
+	FString CustomControllerName = TEXT("");
+	if (StructObject == UTexture::StaticClass() || StructObject == UTexture2D::StaticClass())
+	{
+		if (PropertyBagType == EPropertyBagPropertyType::String)
+		{
+			StructObject = nullptr;
+			CustomControllerName = UE::RCCustomControllers::CustomTextureControllerName;
+		}
+	}
+
+	// Create a Controller of matching type
+	URCController* NewController = CastChecked<URCController>(InPreset->AddController(URCController::StaticClass(), PropertyBagType, StructObject, InRemoteControlProperty->FieldName));
+	NewController->DisplayIndex = InPreset->GetNumControllers() - 1;
+
+	// Add metadata to this controller, if this is a custom controller
+	if (!CustomControllerName.IsEmpty())
+	{
+		const TMap<FName, FString>& CustomControllerMetaData = UE::RCCustomControllers::GetCustomControllerMetaData(CustomControllerName);
+		for (const TPair<FName, FString>& Pair : CustomControllerMetaData)
+		{
+			NewController->SetMetadataValue(Pair.Key, Pair.Value);
+		}
+	}
+
+	// Transfer property value from Exposed Property to the New Controller.
+	// The goal is to keep values synced for a Controller newly created via "Auto Bind"
+	URCBehaviourBind::CopyPropertyValueToController(NewController, InRemoteControlProperty.ToSharedRef());
+	return NewController;
+}
diff --git a/Engine/Plugins/VirtualProduction/RemoteControl/Source/RemoteControlUI/Private/UI/RCUIHelpers.h b/Engine/Plugins/VirtualProduction/RemoteControl/Source/RemoteControlUI/Private/UI/RCUIHelpers.h
--- a/Engine/Plugins/VirtualProduction/RemoteControl/Source/RemoteControlUI/Private/UI/RCUIHelpers.h
+++ b/Engine/Plugins/VirtualProduction/RemoteControl/Source/RemoteControlUI/Private/UI/RCUIHelpers.h
@@ -3,6 +3,7 @@
 #pragma once
 
 #include "CoreMinimal.h"
+#include "Templates/SharedPointerFwd.h"
 #include "UObject/ObjectPtr.h"
 
 class IDetailTreeNode;
@@ -11,6 +12,8 @@ class IPropertyRowGenerator;
 class SWidget;
 class URCController;
 class URCVirtualPropertySelfContainer;
+class URemoteControlPreset;
+struct FRemoteControlProperty;
 
 namespace UE::RCUIHelpers
 {
@@ -58,4 +61,7 @@ namespace UE::RCUIHelpers
 	* Used to provide a consistent look for the type color widget across various Remote Control Logic panels
 	*/
 	TSharedRef<SWidget> GetTypeColorWidget(const FProperty* InProperty);
+
+	/** Automatically creates a controller for a given remote control property */
+	URCController* CreateControllerFromEntity(URemoteControlPreset* InPreset, const TSharedPtr<const FRemoteControlProperty>& InRemoteControlProperty);
 }
\ No newline at end of file
