Remote Control: controller panel to refresh on undo

diff --git a/Engine/Plugins/VirtualProduction/RemoteControl/Source/RemoteControlUI/Private/UI/Controller/SRCControllerPanel.cpp b/Engine/Plugins/VirtualProduction/RemoteControl/Source/RemoteControlUI/Private/UI/Controller/SRCControllerPanel.cpp
--- a/Engine/Plugins/VirtualProduction/RemoteControl/Source/RemoteControlUI/Private/UI/Controller/SRCControllerPanel.cpp
+++ b/Engine/Plugins/VirtualProduction/RemoteControl/Source/RemoteControlUI/Private/UI/Controller/SRCControllerPanel.cpp
@@ -85,6 +85,14 @@ void SRCControllerPanel::Construct(const FArguments& InArgs, const TSharedRef<SR
 		];
 }
 
+void SRCControllerPanel::Refresh()
+{
+	if (ControllerPanelList.IsValid())
+	{
+		ControllerPanelList->Reset();
+	}
+}
+
 END_SLATE_FUNCTION_BUILD_OPTIMIZATION
 
 SRCControllerPanel::FRCControllerPropertyInfo SRCControllerPanel::GetPropertyInfoForCustomType(const FName& InType)
diff --git a/Engine/Plugins/VirtualProduction/RemoteControl/Source/RemoteControlUI/Private/UI/Controller/SRCControllerPanel.h b/Engine/Plugins/VirtualProduction/RemoteControl/Source/RemoteControlUI/Private/UI/Controller/SRCControllerPanel.h
--- a/Engine/Plugins/VirtualProduction/RemoteControl/Source/RemoteControlUI/Private/UI/Controller/SRCControllerPanel.h
+++ b/Engine/Plugins/VirtualProduction/RemoteControl/Source/RemoteControlUI/Private/UI/Controller/SRCControllerPanel.h
@@ -55,6 +55,9 @@ public:
 	/** Constructs this widget with InArgs */
 	void Construct(const FArguments& InArgs, const TSharedRef<SRemoteControlPanel>& InPanel);
 
+	/** Called to refresh the controller list */
+	void Refresh();
+
 	/** Whether the Controller list widget currently has focus. Used for Delete Item UI command */
 	bool IsListFocused() const;
 
diff --git a/Engine/Plugins/VirtualProduction/RemoteControl/Source/RemoteControlUI/Private/UI/SRemoteControlPanel.cpp b/Engine/Plugins/VirtualProduction/RemoteControl/Source/RemoteControlUI/Private/UI/SRemoteControlPanel.cpp
--- a/Engine/Plugins/VirtualProduction/RemoteControl/Source/RemoteControlUI/Private/UI/SRemoteControlPanel.cpp
+++ b/Engine/Plugins/VirtualProduction/RemoteControl/Source/RemoteControlUI/Private/UI/SRemoteControlPanel.cpp
@@ -1427,6 +1427,11 @@ void SRemoteControlPanel::Refresh()
 	}
 
 	EntityList->Refresh();
+
+	if (ControllerPanel.IsValid())
+	{
+		ControllerPanel->Refresh();
+	}
 }
 
 void SRemoteControlPanel::AddToolbarWidget(TSharedRef<SWidget> Widget)
