Remote Control: fixed issue where the protocols kept running even after the RC asset window was closed

diff --git a/Engine/Plugins/VirtualProduction/RemoteControl/Source/RemoteControlProtocol/Private/RemoteControlProtocolModule.cpp b/Engine/Plugins/VirtualProduction/RemoteControl/Source/RemoteControlProtocol/Private/RemoteControlProtocolModule.cpp
--- a/Engine/Plugins/VirtualProduction/RemoteControl/Source/RemoteControlProtocol/Private/RemoteControlProtocolModule.cpp
+++ b/Engine/Plugins/VirtualProduction/RemoteControl/Source/RemoteControlProtocol/Private/RemoteControlProtocolModule.cpp
@@ -104,35 +104,66 @@ void FRemoteControlProtocolModule::EmptyProtocols()
 	Protocols.Empty();
 }
 
-void FRemoteControlProtocolModule::OnPostLoadRemoteControlPreset(URemoteControlPreset* InPreset) const
+void FRemoteControlProtocolModule::ApplyProtocolBindings(URemoteControlPreset* InPreset) const
 {
-	// Avoid protocol binding if the preset is not valid or while cooking.
-	if (!ensure(InPreset) || GIsCookerLoadingPackage)
+	if (!InPreset)
 	{
 		return;
 	}
-	
-	for(TWeakPtr<FRemoteControlProperty> ExposedPropertyWeakPtr : InPreset->GetExposedEntities<FRemoteControlProperty>())
+
+	for (TWeakPtr<FRemoteControlProperty> ExposedPropertyWeakPtr : InPreset->GetExposedEntities<FRemoteControlProperty>())
 	{
 		if (TSharedPtr<FRemoteControlProperty> ExposedPropertyPtr = ExposedPropertyWeakPtr.Pin())
 		{
-			for(FRemoteControlProtocolBinding& Binding : ExposedPropertyPtr->ProtocolBindings)
+			for (const FRemoteControlProtocolBinding& Binding : ExposedPropertyPtr->ProtocolBindings)
 			{
-				const TSharedPtr<IRemoteControlProtocol> Protocol = GetProtocolByName(Binding.GetProtocolName());
 				// Supporting plugin needs to be loaded/protocol available.
-				if(Protocol.IsValid())
+				if (const TSharedPtr<IRemoteControlProtocol> Protocol = GetProtocolByName(Binding.GetProtocolName()))
 				{
-					Protocol->Bind(Binding.GetRemoteControlProtocolEntityPtr());
+					if (TSharedPtr<TStructOnScope<FRemoteControlProtocolEntity>> ProtocolEntity = Binding.GetRemoteControlProtocolEntityPtr())
+					{
+						Protocol->Unbind(ProtocolEntity);
+						Protocol->Bind(ProtocolEntity);
+					}
 				}
-			}	
+			}
 		}
-		else
+	}
+}
+
+void FRemoteControlProtocolModule::UnapplyProtocolBindings(URemoteControlPreset* InPreset) const
+{
+	if (!InPreset)
+	{
+		return;
+	}
+
+	for (TWeakPtr<FRemoteControlProperty> ExposedPropertyWeakPtr : InPreset->GetExposedEntities<FRemoteControlProperty>())
+	{
+		if (TSharedPtr<FRemoteControlProperty> ExposedPropertyPtr = ExposedPropertyWeakPtr.Pin())
 		{
-			ensure(false);
+			for (const FRemoteControlProtocolBinding& Binding : ExposedPropertyPtr->ProtocolBindings)
+			{
+				// Supporting plugin needs to be loaded/protocol available.
+				if (const TSharedPtr<IRemoteControlProtocol> Protocol = GetProtocolByName(Binding.GetProtocolName()))
+				{
+					Protocol->Unbind(Binding.GetRemoteControlProtocolEntityPtr());
+				}
+			}
 		}
 	}
 }
 
+void FRemoteControlProtocolModule::OnPostLoadRemoteControlPreset(URemoteControlPreset* InPreset) const
+{
+	// Avoid protocol binding if the preset is not valid or while cooking.
+	if (!ensure(InPreset) || GIsCookerLoadingPackage)
+	{
+		return;
+	}
+
+	ApplyProtocolBindings(InPreset);
+}
 
 IMPLEMENT_MODULE(FRemoteControlProtocolModule, RemoteControlProtocol);
 
diff --git a/Engine/Plugins/VirtualProduction/RemoteControl/Source/RemoteControlProtocol/Private/RemoteControlProtocolModule.h b/Engine/Plugins/VirtualProduction/RemoteControl/Source/RemoteControlProtocol/Private/RemoteControlProtocolModule.h
--- a/Engine/Plugins/VirtualProduction/RemoteControl/Source/RemoteControlProtocol/Private/RemoteControlProtocolModule.h
+++ b/Engine/Plugins/VirtualProduction/RemoteControl/Source/RemoteControlProtocol/Private/RemoteControlProtocolModule.h
@@ -24,6 +24,8 @@ public:
 	virtual void RemoveProtocol(FName InProtocolName, TSharedRef<IRemoteControlProtocol> InProtocol) override;
 	virtual void EmptyProtocols() override;
 	virtual bool IsRCProtocolsDisable() const override { return bRCProtocolsDisable; };
+	virtual void ApplyProtocolBindings(URemoteControlPreset* InPreset) const override;
+	virtual void UnapplyProtocolBindings(URemoteControlPreset* InPreset) const override;
 	//~ End IRemoteControlProtocolModule interface
 
 private:
diff --git a/Engine/Plugins/VirtualProduction/RemoteControl/Source/RemoteControlProtocol/Public/IRemoteControlProtocolModule.h b/Engine/Plugins/VirtualProduction/RemoteControl/Source/RemoteControlProtocol/Public/IRemoteControlProtocolModule.h
--- a/Engine/Plugins/VirtualProduction/RemoteControl/Source/RemoteControlProtocol/Public/IRemoteControlProtocolModule.h
+++ b/Engine/Plugins/VirtualProduction/RemoteControl/Source/RemoteControlProtocol/Public/IRemoteControlProtocolModule.h
@@ -7,6 +7,8 @@
 
 #include "IRemoteControlProtocol.h"
 
+class URemoteControlPreset;
+
 /**
  * A Remote Control Protocol module. That is base for all remote control protcols
  */
@@ -49,6 +51,12 @@ public:
 	 */
 	virtual void RemoveProtocol(FName InProtocolName, TSharedRef<IRemoteControlProtocol> InProtocol) = 0;
 
+	/** Applies all the protocol bindings found within the given preset */
+	virtual void ApplyProtocolBindings(URemoteControlPreset* InPreset) const = 0;
+
+	/** Unapplies all the protocol bindings found within the given preset */
+	virtual void UnapplyProtocolBindings(URemoteControlPreset* InPreset) const = 0;
+
 	/** Unregister all protocols */
 	virtual void EmptyProtocols() = 0;
 
diff --git a/Engine/Plugins/VirtualProduction/RemoteControl/Source/RemoteControlUI/Private/UI/SRemoteControlPanel.cpp b/Engine/Plugins/VirtualProduction/RemoteControl/Source/RemoteControlUI/Private/UI/SRemoteControlPanel.cpp
--- a/Engine/Plugins/VirtualProduction/RemoteControl/Source/RemoteControlUI/Private/UI/SRemoteControlPanel.cpp
+++ b/Engine/Plugins/VirtualProduction/RemoteControl/Source/RemoteControlUI/Private/UI/SRemoteControlPanel.cpp
@@ -19,6 +19,7 @@
 #include "Framework/MultiBox/MultiBoxBuilder.h"
 #include "GameFramework/Actor.h"
 #include "IRemoteControlModule.h"
+#include "IRemoteControlProtocolModule.h"
 #include "IRemoteControlProtocolWidgetsModule.h"
 #include "IRemoteControlUIModule.h"
 #include "ISettingsModule.h"
@@ -107,6 +108,16 @@ void SRemoteControlPanel::AddReferencedObjects(FReferenceCollector& Collector)
 	Collector.AddReferencedObjects(LogicClipboardItems);
 }
 
+void SRemoteControlPanel::ApplyProtocolBindings()
+{
+	IRemoteControlProtocolModule::Get().ApplyProtocolBindings(Preset.Get());
+}
+
+void SRemoteControlPanel::UnapplyProtocolBindings()
+{
+	IRemoteControlProtocolModule::Get().UnapplyProtocolBindings(Preset.Get());
+}
+
 namespace RemoteControlPanelUtils
 {
 	bool IsExposableActor(AActor* Actor)
@@ -200,6 +211,8 @@ void SRemoteControlPanel::Construct(const FArguments& InArgs, URemoteControlPres
 	TArray<TSharedRef<SWidget>> ExtensionWidgets;
 	FRemoteControlUIModule::Get().GetExtensionGenerators().Broadcast(ExtensionWidgets);
 
+	ApplyProtocolBindings();
+
 	BindRemoteControlCommands();
 
 	GenerateAuxiliaryToolbar();
@@ -462,6 +475,8 @@ void SRemoteControlPanel::Construct(const FArguments& InArgs, URemoteControlPres
 
 SRemoteControlPanel::~SRemoteControlPanel()
 {
+	UnapplyProtocolBindings();
+
 	SaveSettings();
 	UnregisterEvents();
 
diff --git a/Engine/Plugins/VirtualProduction/RemoteControl/Source/RemoteControlUI/Private/UI/SRemoteControlPanel.h b/Engine/Plugins/VirtualProduction/RemoteControl/Source/RemoteControlUI/Private/UI/SRemoteControlPanel.h
--- a/Engine/Plugins/VirtualProduction/RemoteControl/Source/RemoteControlUI/Private/UI/SRemoteControlPanel.h
+++ b/Engine/Plugins/VirtualProduction/RemoteControl/Source/RemoteControlUI/Private/UI/SRemoteControlPanel.h
@@ -174,6 +174,12 @@ public:
 	virtual void AddReferencedObjects(FReferenceCollector& Collector) override;
 
 private:
+	/** Applies all the protocol bindings within this widget's owning preset to their respective protocols */
+	void ApplyProtocolBindings();
+
+	/** Unapplies all the protocol bindings within this widget's owning preset to their respective protocols */
+	void UnapplyProtocolBindings();
+
 	/** Returns a Helper widget for entity details view and protocol details view. */
 	static TSharedRef<SBox> CreateNoneSelectedWidget();
 
diff --git a/Engine/Plugins/VirtualProduction/RemoteControlProtocolDMX/Source/RemoteControlProtocolDMX/Private/Library/RemoteControlDMXLibraryProxy.cpp b/Engine/Plugins/VirtualProduction/RemoteControlProtocolDMX/Source/RemoteControlProtocolDMX/Private/Library/RemoteControlDMXLibraryProxy.cpp
--- a/Engine/Plugins/VirtualProduction/RemoteControlProtocolDMX/Source/RemoteControlProtocolDMX/Private/Library/RemoteControlDMXLibraryProxy.cpp
+++ b/Engine/Plugins/VirtualProduction/RemoteControlProtocolDMX/Source/RemoteControlProtocolDMX/Private/Library/RemoteControlDMXLibraryProxy.cpp
@@ -178,7 +178,7 @@ TArray<UDMXEntityFixturePatch*> URemoteControlDMXLibraryProxy::FindPatchesThatEx
 #endif // WITH_EDITOR
 
 void URemoteControlDMXLibraryProxy::UpdatePropertyPatches()
-{	
+{
 	URemoteControlDMXUserData& DMXUserData = GetMutableDMXUserDataChecked();
 	UDMXLibrary* DMXLibrary = DMXUserData.GetDMXLibrary();
 	if (!ensureMsgf(DMXLibrary, TEXT("Cannot create remote control DMX patches, library is invalid.")))
@@ -186,24 +186,35 @@ void URemoteControlDMXLibraryProxy::UpdatePropertyPatches()
 		return;
 	}
 
+	const TSharedPtr<FRemoteControlProtocolDMX> DMXProtocol = StaticCastSharedPtr<FRemoteControlProtocolDMX>(IRemoteControlProtocolModule::Get().GetProtocolByName(FRemoteControlProtocolDMX::ProtocolName));
+	if (!ensureMsgf(DMXProtocol.IsValid(), TEXT("Cannot create remote control DMX patches, DMX protocol is not available")))
+	{
+		return;
+	}
+
+	// Checks that for a given binding, its protocol is DMX, and is bound to the DMX protocol.
+	auto IsValidDMXBinding =
+		[ProtocolBindings = DMXProtocol->GetProtocolBindings()](const FRemoteControlProtocolBinding& InBinding)->bool
+		{
+			return InBinding.GetProtocolName() == FRemoteControlProtocolDMX::ProtocolName
+				&& ProtocolBindings.Contains(InBinding.GetRemoteControlProtocolEntityPtr());
+		};
+
 	// Create DMX Controlled Properties
 	TArray<TSharedRef<FRemoteControlProperty>> ExposedProperties = GetExposedProperties();
 
 	TArray<TSharedRef<FRemoteControlDMXControlledProperty>> DMXControlledProperties;
 	Algo::TransformIf(ExposedProperties, DMXControlledProperties,
-		[](const TSharedRef<FRemoteControlProperty>& ExposedProperty)
+		[&IsValidDMXBinding](const TSharedRef<FRemoteControlProperty>& ExposedProperty)
 		{			
 			const bool bValidPreset = ExposedProperty->GetOwner() != nullptr;
 			const bool bValidBoundObject = ExposedProperty->GetBoundObject() != nullptr;
-			const bool bHasAnyDMXBindings = Algo::AnyOf(ExposedProperty->ProtocolBindings, [](const FRemoteControlProtocolBinding& Binding)
-				{
-					return Binding.GetProtocolName() == FRemoteControlProtocolDMX::ProtocolName;
-				});
+			const bool bHasValidDMXBindings = Algo::AnyOf(ExposedProperty->ProtocolBindings, IsValidDMXBinding);
 
 			return 
 				bValidPreset &&
 				bValidBoundObject &&
-				bHasAnyDMXBindings;
+				bHasValidDMXBindings;
 		},
 		[](const TSharedRef<FRemoteControlProperty>& ExposedProperty)
 		{
diff --git a/Engine/Plugins/VirtualProduction/RemoteControlProtocolDMX/Source/RemoteControlProtocolDMX/Private/RemoteControlProtocolDMX.cpp b/Engine/Plugins/VirtualProduction/RemoteControlProtocolDMX/Source/RemoteControlProtocolDMX/Private/RemoteControlProtocolDMX.cpp
--- a/Engine/Plugins/VirtualProduction/RemoteControlProtocolDMX/Source/RemoteControlProtocolDMX/Private/RemoteControlProtocolDMX.cpp
+++ b/Engine/Plugins/VirtualProduction/RemoteControlProtocolDMX/Source/RemoteControlProtocolDMX/Private/RemoteControlProtocolDMX.cpp
@@ -217,7 +217,7 @@ void FRemoteControlProtocolDMX::Bind(FRemoteControlProtocolEntityPtr InRemoteCon
 			return false;
 		});
 
-	if (ExistingProtocolBindings)
+	if (!ExistingProtocolBindings)
 	{
 		WeakProtocolsBindings.Emplace(InRemoteControlProtocolEntityPtr);
 	}
diff --git a/Engine/Plugins/VirtualProduction/RemoteControlProtocolDMX/Source/RemoteControlProtocolDMX/Public/RemoteControlProtocolDMX.h b/Engine/Plugins/VirtualProduction/RemoteControlProtocolDMX/Source/RemoteControlProtocolDMX/Public/RemoteControlProtocolDMX.h
--- a/Engine/Plugins/VirtualProduction/RemoteControlProtocolDMX/Source/RemoteControlProtocolDMX/Public/RemoteControlProtocolDMX.h
+++ b/Engine/Plugins/VirtualProduction/RemoteControlProtocolDMX/Source/RemoteControlProtocolDMX/Public/RemoteControlProtocolDMX.h
@@ -197,7 +197,12 @@ public:
 	FRemoteControlProtocolDMX()
 		: FRemoteControlProtocol(ProtocolName)
 	{}
-	
+
+	TConstArrayView<FRemoteControlProtocolEntityWeakPtr> GetProtocolBindings() const
+	{
+		return WeakProtocolsBindings;
+	}
+
 	//~ Begin IRemoteControlProtocol interface
 	virtual void Bind(TSharedPtr<TStructOnScope<FRemoteControlProtocolEntity>> InRemoteControlProtocolEntity) override;
 	virtual void Unbind(TSharedPtr<TStructOnScope<FRemoteControlProtocolEntity>> InRemoteControlProtocolEntity) override;
