Remote Control: Add struct referenced objects to protocol bindings to consider protocol entity

diff --git a/Engine/Plugins/VirtualProduction/RemoteControl/Source/RemoteControl/Private/RemoteControlProtocolBinding.cpp b/Engine/Plugins/VirtualProduction/RemoteControl/Source/RemoteControl/Private/RemoteControlProtocolBinding.cpp
--- a/Engine/Plugins/VirtualProduction/RemoteControl/Source/RemoteControl/Private/RemoteControlProtocolBinding.cpp
+++ b/Engine/Plugins/VirtualProduction/RemoteControl/Source/RemoteControl/Private/RemoteControlProtocolBinding.cpp
@@ -825,6 +825,17 @@ TSharedPtr<FStructOnScope> FRemoteControlProtocolBinding::GetStructOnScope() con
 	return RemoteControlProtocolEntityPtr;
 }
 
+void FRemoteControlProtocolBinding::AddStructReferencedObjects(FReferenceCollector& InCollector)
+{
+	// Check that the shared ptr and the underlying struct on scope are valid
+	if (RemoteControlProtocolEntityPtr.IsValid() && RemoteControlProtocolEntityPtr->IsValid())
+	{
+		// The struct should be guaranteed a script struct (i.e. since it's a TStructOnScope<FRemoteControlProtocolEntity>)
+		const UScriptStruct* EntityType = CastChecked<UScriptStruct>(RemoteControlProtocolEntityPtr->GetStruct());
+		InCollector.AddPropertyReferencesWithStructARO(EntityType, RemoteControlProtocolEntityPtr->Get());
+	}
+}
+
 FRemoteControlProtocolEntity* FRemoteControlProtocolBinding::GetRemoteControlProtocolEntity()
 {
 	if (RemoteControlProtocolEntityPtr.IsValid())
diff --git a/Engine/Plugins/VirtualProduction/RemoteControl/Source/RemoteControl/Public/RemoteControlProtocolBinding.h b/Engine/Plugins/VirtualProduction/RemoteControl/Source/RemoteControl/Public/RemoteControlProtocolBinding.h
--- a/Engine/Plugins/VirtualProduction/RemoteControl/Source/RemoteControl/Public/RemoteControlProtocolBinding.h
+++ b/Engine/Plugins/VirtualProduction/RemoteControl/Source/RemoteControl/Public/RemoteControlProtocolBinding.h
@@ -699,7 +699,9 @@ public:
 	/** Custom struct serialize */
 	bool Serialize(FArchive& Ar);
 	friend FArchive& operator<<(FArchive& Ar, FRemoteControlProtocolBinding& InProtocolBinding);
-	
+
+	void AddStructReferencedObjects(FReferenceCollector& InCollector);
+
 private:
 	/** Return FRemoteControlProtocolEntity pointer from RemoteControlProtocolEntityPtr */
 	FRemoteControlProtocolEntity* GetRemoteControlProtocolEntity();
@@ -735,5 +737,6 @@ struct TStructOpsTypeTraits<FRemoteControlProtocolBinding> : public TStructOpsTy
 	enum
 	{
 		WithSerializer = true,
+		WithAddStructReferencedObjects = true,
 	};
 };
diff --git a/Engine/Plugins/VirtualProduction/RemoteControlProtocolDMX/Source/RemoteControlProtocolDMX/Internal/Library/RemoteControlDMXLibraryProxy.h b/Engine/Plugins/VirtualProduction/RemoteControlProtocolDMX/Source/RemoteControlProtocolDMX/Internal/Library/RemoteControlDMXLibraryProxy.h
--- a/Engine/Plugins/VirtualProduction/RemoteControlProtocolDMX/Source/RemoteControlProtocolDMX/Internal/Library/RemoteControlDMXLibraryProxy.h
+++ b/Engine/Plugins/VirtualProduction/RemoteControlProtocolDMX/Source/RemoteControlProtocolDMX/Internal/Library/RemoteControlDMXLibraryProxy.h
@@ -54,8 +54,7 @@ public:
 #if WITH_EDITOR
 	/**
 	 * Clears Fixture Patches. This should be called after a DMX Library asset in use was force deleted.
-	 * This is due to the fact that remote control protocol entities do not follow GC as expected, being left with stale
-	 * DMX Libraries in their Fixture Patch Ref objects otherwise.
+	 * Remote Control Protocol Entities keep a 'hard' reference to the DMX Library object, so need to be cleared here
 	 */
 	void ClearFixturePatches();
 
diff --git a/Engine/Plugins/VirtualProduction/RemoteControlProtocolDMX/Source/RemoteControlProtocolDMX/Private/RemoteControlDMXUserData.cpp b/Engine/Plugins/VirtualProduction/RemoteControlProtocolDMX/Source/RemoteControlProtocolDMX/Private/RemoteControlDMXUserData.cpp
--- a/Engine/Plugins/VirtualProduction/RemoteControlProtocolDMX/Source/RemoteControlProtocolDMX/Private/RemoteControlDMXUserData.cpp
+++ b/Engine/Plugins/VirtualProduction/RemoteControlProtocolDMX/Source/RemoteControlProtocolDMX/Private/RemoteControlDMXUserData.cpp
@@ -159,8 +159,7 @@ void URemoteControlDMXUserData::EnsureValidDMXLibrary()
 #if WITH_EDITOR
 		// It is important to clear the fixture patches for the DMX Library proxy when the DMX Library is reset,
 		// this is to handle cases where a DMX Library was force deleted. 
-		// As of 5.5 Remote Control Protocol Entities do not follow GC as expected, leaving their 
-		// fixture patch references with stale DMX Library object pointers if not explicitly cleared here.
+		// Remote Control Protocol Entities keep a 'hard' reference to the DMX Library object, so need to be cleared here
 		check(DMXLibraryProxy);
 		DMXLibraryProxy->ClearFixturePatches();
 #endif // WITH_EDITOR
