Remote Control: protocol entities to unbind when entity gets unexposed/deleted

diff --git a/Engine/Plugins/VirtualProduction/RemoteControl/Source/RemoteControlProtocolWidgets/Private/ViewModels/ProtocolEntityViewModel.cpp b/Engine/Plugins/VirtualProduction/RemoteControl/Source/RemoteControlProtocolWidgets/Private/ViewModels/ProtocolEntityViewModel.cpp
--- a/Engine/Plugins/VirtualProduction/RemoteControl/Source/RemoteControlProtocolWidgets/Private/ViewModels/ProtocolEntityViewModel.cpp
+++ b/Engine/Plugins/VirtualProduction/RemoteControl/Source/RemoteControlProtocolWidgets/Private/ViewModels/ProtocolEntityViewModel.cpp
@@ -211,6 +211,14 @@ void FProtocolEntityViewModel::OnEntityUnexposed(URemoteControlPreset* InPreset,
 	{
 		if (GetId() == InEntityId)
 		{
+			for (const TSharedPtr<FProtocolBindingViewModel>& Binding : Bindings)
+			{
+				if (Binding.IsValid())
+				{
+					RemoveProtocolInternal(InPreset, Commands::FAddRemoveProtocolArgs{GetId(), NAME_None, Binding->GetId()});
+				}
+			}
+
 			Bindings.Empty();
 			OnChanged().Broadcast();
 		}
