Remote Control: Fixed issue with field view models not updating their field index when a sibling is removed

diff --git a/Engine/Plugins/VirtualProduction/RemoteControl/Source/RemoteControlUI/Private/UI/Signature/Items/RCSignatureTreeFieldItem.cpp b/Engine/Plugins/VirtualProduction/RemoteControl/Source/RemoteControlUI/Private/UI/Signature/Items/RCSignatureTreeFieldItem.cpp
--- a/Engine/Plugins/VirtualProduction/RemoteControl/Source/RemoteControlUI/Private/UI/Signature/Items/RCSignatureTreeFieldItem.cpp
+++ b/Engine/Plugins/VirtualProduction/RemoteControl/Source/RemoteControlUI/Private/UI/Signature/Items/RCSignatureTreeFieldItem.cpp
@@ -131,6 +131,18 @@ int32 FRCSignatureTreeFieldItem::RemoveFromRegistry()
 	FScopedTransaction Transaction(LOCTEXT("RemoveField", "Remove Field"));
 	Registry->Modify();
 	Signature->Fields.RemoveAt(FieldIndex);
+
+	// Update the field indices of sibling field items that have a greater field index than the index being removed 
+	ForEachSiblingFieldItem(
+		[RemovedFieldIndex = FieldIndex](const TSharedPtr<FRCSignatureTreeFieldItem>& InSiblingItem)
+		{
+			if (InSiblingItem->FieldIndex > RemovedFieldIndex)
+			{
+				--InSiblingItem->FieldIndex;
+				InSiblingItem->RebuildPath();
+			}
+		});
+
 	return 1;
 }
 
@@ -183,4 +195,21 @@ FRCSignature* FRCSignatureTreeFieldItem::FindParentSignature(URCSignatureRegistr
 	return Signature;
 }
 
+void FRCSignatureTreeFieldItem::ForEachSiblingFieldItem(TFunctionRef<void(const TSharedPtr<FRCSignatureTreeFieldItem>&)> InCallable)
+{
+	if (TSharedPtr<FRCSignatureTreeSignatureItem> ParentSignatureItem = GetParentSignatureItem())
+	{
+		ParentSignatureItem->VisitChildren(
+			[&InCallable](const TSharedPtr<FRCSignatureTreeItemBase>& InChildItem)
+			{
+				if (const TSharedPtr<FRCSignatureTreeFieldItem> FieldItem = InChildItem->MutableCast<FRCSignatureTreeFieldItem>())
+				{
+					InCallable(FieldItem);
+				}
+				return true; // Continue iteration
+			}
+			, /*bRecursive*/ false);
+	}
+}
+
 #undef LOCTEXT_NAMESPACE
diff --git a/Engine/Plugins/VirtualProduction/RemoteControl/Source/RemoteControlUI/Private/UI/Signature/Items/RCSignatureTreeFieldItem.h b/Engine/Plugins/VirtualProduction/RemoteControl/Source/RemoteControlUI/Private/UI/Signature/Items/RCSignatureTreeFieldItem.h
--- a/Engine/Plugins/VirtualProduction/RemoteControl/Source/RemoteControlUI/Private/UI/Signature/Items/RCSignatureTreeFieldItem.h
+++ b/Engine/Plugins/VirtualProduction/RemoteControl/Source/RemoteControlUI/Private/UI/Signature/Items/RCSignatureTreeFieldItem.h
@@ -40,5 +40,8 @@ private:
 
 	FRCSignature* FindParentSignature(URCSignatureRegistry** OutRegistry);
 
+	/** Iterates each field item sibling */
+	void ForEachSiblingFieldItem(TFunctionRef<void(const TSharedPtr<FRCSignatureTreeFieldItem>&)> InCallable);
+
 	int32 FieldIndex;
 };
diff --git a/Engine/Plugins/VirtualProduction/RemoteControl/Source/RemoteControlUI/Private/UI/Signature/Items/RCSignatureTreeItemBase.cpp b/Engine/Plugins/VirtualProduction/RemoteControl/Source/RemoteControlUI/Private/UI/Signature/Items/RCSignatureTreeItemBase.cpp
--- a/Engine/Plugins/VirtualProduction/RemoteControl/Source/RemoteControlUI/Private/UI/Signature/Items/RCSignatureTreeItemBase.cpp
+++ b/Engine/Plugins/VirtualProduction/RemoteControl/Source/RemoteControlUI/Private/UI/Signature/Items/RCSignatureTreeItemBase.cpp
@@ -86,6 +86,20 @@ void FRCSignatureTreeItemBase::VisitChildren(TFunctionRef<bool(const TSharedPtr<
 	}
 }
 
+void FRCSignatureTreeItemBase::RebuildPath()
+{
+	Path = BuildPath();
+
+	// Rebuilds path for each item recursively (VisitChildren isn't setup to be recursive, but RebuildPath is)
+	VisitChildren(
+		[](const TSharedPtr<FRCSignatureTreeItemBase>& InChild)
+		{
+			InChild->RebuildPath();
+			return true; // continue iteration
+		}
+		, /*bRecursive*/false);
+}
+
 void FRCSignatureTreeItemBase::Initialize(const TSharedPtr<FRCSignatureTreeItemBase>& InParent)
 {
 	ParentWeak = InParent;
diff --git a/Engine/Plugins/VirtualProduction/RemoteControl/Source/RemoteControlUI/Private/UI/Signature/Items/RCSignatureTreeItemBase.h b/Engine/Plugins/VirtualProduction/RemoteControl/Source/RemoteControlUI/Private/UI/Signature/Items/RCSignatureTreeItemBase.h
--- a/Engine/Plugins/VirtualProduction/RemoteControl/Source/RemoteControlUI/Private/UI/Signature/Items/RCSignatureTreeItemBase.h
+++ b/Engine/Plugins/VirtualProduction/RemoteControl/Source/RemoteControlUI/Private/UI/Signature/Items/RCSignatureTreeItemBase.h
@@ -164,6 +164,9 @@ public:
 	void VisitChildren(TFunctionRef<bool(const TSharedPtr<FRCSignatureTreeItemBase>&)> InCallable, bool bInRecursive);
 
 protected:
+	/** Called when needing to recalculate a new path to the item */
+	void RebuildPath();
+
 	virtual void BuildPathSegment(FStringBuilderBase& InBuilder) const
 	{
 	}
