Remote Control: deprecated active masks from entities

diff --git a/Engine/Plugins/VirtualProduction/RemoteControl/Source/RemoteControl/Private/RemoteControlField.cpp b/Engine/Plugins/VirtualProduction/RemoteControl/Source/RemoteControl/Private/RemoteControlField.cpp
--- a/Engine/Plugins/VirtualProduction/RemoteControl/Source/RemoteControl/Private/RemoteControlField.cpp
+++ b/Engine/Plugins/VirtualProduction/RemoteControl/Source/RemoteControl/Private/RemoteControlField.cpp
@@ -150,17 +150,30 @@ bool FRemoteControlField::CanBindObject(const UObject* InObjectToBind) const
 
 void FRemoteControlField::ClearMask(ERCMask InMaskBit)
 {
-	ActiveMasks &= ~(uint8)InMaskBit;
+PRAGMA_DISABLE_DEPRECATION_WARNINGS
+	ActiveMasks_DEPRECATED &= ~(uint8)InMaskBit;
+PRAGMA_ENABLE_DEPRECATION_WARNINGS
 }
 
 void FRemoteControlField::EnableMask(ERCMask InMaskBit)
 {
-	ActiveMasks |= (uint8)InMaskBit;
+PRAGMA_DISABLE_DEPRECATION_WARNINGS
+	ActiveMasks_DEPRECATED |= (uint8)InMaskBit;
+PRAGMA_ENABLE_DEPRECATION_WARNINGS
 }
 
 bool FRemoteControlField::HasMask(ERCMask InMaskBit) const
 {
-	return (ActiveMasks & (uint8)InMaskBit) != (uint8)ERCMask::NoMask;
+PRAGMA_DISABLE_DEPRECATION_WARNINGS
+	return (ActiveMasks_DEPRECATED & (uint8)InMaskBit) != (uint8)ERCMask::NoMask;
+PRAGMA_ENABLE_DEPRECATION_WARNINGS
+}
+
+ERCMask FRemoteControlField::GetActiveMasks() const
+{
+PRAGMA_DISABLE_DEPRECATION_WARNINGS
+	return (ERCMask)ActiveMasks_DEPRECATED;
+PRAGMA_ENABLE_DEPRECATION_WARNINGS
 }
 
 void FRemoteControlField::PostSerialize(const FArchive& Ar)
diff --git a/Engine/Plugins/VirtualProduction/RemoteControl/Source/RemoteControl/Private/RemoteControlProtocolEntityProcessor.cpp b/Engine/Plugins/VirtualProduction/RemoteControl/Source/RemoteControl/Private/RemoteControlProtocolEntityProcessor.cpp
--- a/Engine/Plugins/VirtualProduction/RemoteControl/Source/RemoteControl/Private/RemoteControlProtocolEntityProcessor.cpp
+++ b/Engine/Plugins/VirtualProduction/RemoteControl/Source/RemoteControl/Private/RemoteControlProtocolEntityProcessor.cpp
@@ -483,10 +483,7 @@ namespace UE::RemoteControl::ProtocolEntityProcessor
 				const TTuple<FRemoteControlProtocolEntity*, double>* EntityToValuePtr = Algo::FindByPredicate(PropertyData.EntityToValueMap,
 					[&PropertyName, &PropertyData](const TTuple<FRemoteControlProtocolEntity*, double>& EntityToValuePair)
 					{
-						const ERCMask Mask = EntityToValuePair.Key->GetOverridenMask() == ERCMask::NoMask ?
-							PropertyData.Property->GetActiveMasks() :
-							EntityToValuePair.Key->GetOverridenMask();
-
+						const ERCMask Mask = EntityToValuePair.Key->GetOverridenMask();
 						return Masking::IsInMask<VectorType>(PropertyName, Mask);
 					});
 
@@ -766,11 +763,6 @@ namespace UE::RemoteControl::ProtocolEntityProcessor
 				continue;
 			}
 
-			if (RemoteControlProperty->GetActiveMasks() == ERCMask::NoMask)
-			{
-				continue;
-			}
-
 			FProperty* Property = RemoteControlProperty->GetProperty();
 			if (!Property)
 			{
diff --git a/Engine/Plugins/VirtualProduction/RemoteControl/Source/RemoteControl/Public/RemoteControlField.h b/Engine/Plugins/VirtualProduction/RemoteControl/Source/RemoteControl/Public/RemoteControlField.h
--- a/Engine/Plugins/VirtualProduction/RemoteControl/Source/RemoteControl/Public/RemoteControlField.h
+++ b/Engine/Plugins/VirtualProduction/RemoteControl/Source/RemoteControl/Public/RemoteControlField.h
@@ -32,6 +32,13 @@ struct REMOTECONTROL_API FRemoteControlField : public FRemoteControlEntity
 
 	FRemoteControlField() = default;
 
+	PRAGMA_DISABLE_DEPRECATION_WARNINGS
+	FRemoteControlField(const FRemoteControlField&) = default;
+	FRemoteControlField(FRemoteControlField&&) = default;
+	FRemoteControlField& operator=(const FRemoteControlField&) = default;
+	FRemoteControlField& operator=(FRemoteControlField&&) = default;
+	PRAGMA_ENABLE_DEPRECATION_WARNINGS
+
 	/**
 	 * Resolve the field's owners using the section's top level objects.
 	 * @param SectionObjects The top level objects of the section.
@@ -46,16 +53,24 @@ struct REMOTECONTROL_API FRemoteControlField : public FRemoteControlEntity
 	//~ End FRemoteControlEntity interface
 
 	/** Disables the given mask. */
+	UE_DEPRECATED(5.6, "ClearMask is deprecated. Use per-protocol-binding masking instead.")
 	virtual void ClearMask(ERCMask InMaskBit);
+
 	/** Enables the given mask. */
+	UE_DEPRECATED(5.6, "EnableMask is deprecated. Use per-protocol-binding masking instead.")
 	virtual void EnableMask(ERCMask InMaskBit);
+
 	/** Returns true if the given mask is enabled, false otherwise. */
+	UE_DEPRECATED(5.6, "HasMask is deprecated. Use per-protocol-binding masking instead.")
 	virtual bool HasMask(ERCMask InMaskBit) const;
+
 	/** Returns true if this field supports masking. */
+	UE_DEPRECATED(5.6, "SupportsMasking is deprecated. Use per-protocol-binding masking instead.")
 	virtual bool SupportsMasking() const { return false; }
 
 	/** Retrieves the enabled masks. */
-	ERCMask GetActiveMasks() const { return (ERCMask)ActiveMasks; }
+	UE_DEPRECATED(5.6, "GetActiveMasks is deprecated. Use per-protocol-binding masking instead.")
+	ERCMask GetActiveMasks() const;
 
 	/**	Returns whether the field is only available in editor. */
 	bool IsEditorOnly() const { return bIsEditorOnly; }
@@ -109,10 +124,11 @@ protected:
 	/** Whether the field is only available in editor. */
 	UPROPERTY()
 	bool bIsEditorOnly = false;
-	
+
 	/** Holds the actively enabled masks. */
+	UE_DEPRECATED(5.6, "ActiveMasks is deprecated. Use per-protocol-binding masking instead.")
 	UPROPERTY()
-	uint8 ActiveMasks = (uint8)RC_AllMasks;
+	uint8 ActiveMasks_DEPRECATED = (uint8)RC_AllMasks;
 
 protected:
 	FRemoteControlField(URemoteControlPreset* InPreset, EExposedFieldType InType, FName InLabel, FRCFieldPathInfo InFieldPathInfo, const TArray<URemoteControlBinding*> InBindings);
@@ -152,6 +168,7 @@ public:
 	//~ End FRemoteControlEntity interface
 	
 	//~ Begin FRemoteControlField interface
+	UE_DEPRECATED(5.6, "SupportsMasking is deprecated. Use per-protocol-binding masking instead.")
 	virtual bool SupportsMasking() const override;
 	//~ End FRemoteControlField interface
 
diff --git a/Engine/Plugins/VirtualProduction/RemoteControl/Source/RemoteControlProtocolWidgets/Private/Widgets/Masks/SRCProtocolMask.cpp b/Engine/Plugins/VirtualProduction/RemoteControl/Source/RemoteControlProtocolWidgets/Private/Widgets/Masks/SRCProtocolMask.cpp
--- a/Engine/Plugins/VirtualProduction/RemoteControl/Source/RemoteControlProtocolWidgets/Private/Widgets/Masks/SRCProtocolMask.cpp
+++ b/Engine/Plugins/VirtualProduction/RemoteControl/Source/RemoteControlProtocolWidgets/Private/Widgets/Masks/SRCProtocolMask.cpp
@@ -27,11 +27,12 @@ void SRCProtocolMask::Construct(const FArguments& InArgs, TWeakPtr<FRemoteContro
 
 bool SRCProtocolMask::CanBeMasked() const
 {
+PRAGMA_DISABLE_DEPRECATION_WARNINGS
 	if (TSharedPtr<FRemoteControlField> RCField = WeakField.Pin())
 	{
 		return RCField->SupportsMasking();
 	}
-
+PRAGMA_ENABLE_DEPRECATION_WARNINGS
 	return false;
 }
 
@@ -75,10 +76,12 @@ bool SRCProtocolMask::HasOptionalMask() const
 
 ECheckBoxState SRCProtocolMask::IsMaskEnabled(ERCMask InMaskBit) const
 {
+PRAGMA_DISABLE_DEPRECATION_WARNINGS
 	if (TSharedPtr<FRemoteControlField> RCField = WeakField.Pin())
 	{
 		return RCField->HasMask((ERCMask)InMaskBit) ? ECheckBoxState::Checked : ECheckBoxState::Unchecked;
 	}
+PRAGMA_ENABLE_DEPRECATION_WARNINGS
 
 	return SRCProtocolMaskTriplet::IsMaskEnabled(InMaskBit);
 }
@@ -87,6 +90,7 @@ void SRCProtocolMask::SetMaskEnabled(ECheckBoxState NewState, ERCMask NewMaskBit
 {
 	if (TSharedPtr<FRemoteControlField> RCField = WeakField.Pin())
 	{
+PRAGMA_DISABLE_DEPRECATION_WARNINGS
 		if (NewState == ECheckBoxState::Checked)
 		{
 			RCField->EnableMask((ERCMask)NewMaskBit);
@@ -95,6 +99,7 @@ void SRCProtocolMask::SetMaskEnabled(ECheckBoxState NewState, ERCMask NewMaskBit
 		{
 			RCField->ClearMask((ERCMask)NewMaskBit);
 		}
+PRAGMA_ENABLE_DEPRECATION_WARNINGS
 	}
 	else
 	{
diff --git a/Engine/Plugins/VirtualProduction/RemoteControl/Source/RemoteControlProtocolWidgets/Public/Widgets/Masks/SRCProtocolMask.h b/Engine/Plugins/VirtualProduction/RemoteControl/Source/RemoteControlProtocolWidgets/Public/Widgets/Masks/SRCProtocolMask.h
--- a/Engine/Plugins/VirtualProduction/RemoteControl/Source/RemoteControlProtocolWidgets/Public/Widgets/Masks/SRCProtocolMask.h
+++ b/Engine/Plugins/VirtualProduction/RemoteControl/Source/RemoteControlProtocolWidgets/Public/Widgets/Masks/SRCProtocolMask.h
@@ -11,6 +11,7 @@ struct FRemoteControlField;
 /**
  * A widget that represents a mask.
  */
+class UE_DEPRECATED(5.6, "SRCProtocolMask deprecated. Prefer using SRCProtocolMaskTriplet for mask widgets") SRCProtocolMask;
 class REMOTECONTROLPROTOCOLWIDGETS_API SRCProtocolMask : public SRCProtocolMaskTriplet, public IHasMaskExtensibility
 {
 public:
diff --git a/Engine/Plugins/VirtualProduction/RemoteControl/Source/RemoteControlUI/Private/UI/SRCPanelExposedEntitiesList.cpp b/Engine/Plugins/VirtualProduction/RemoteControl/Source/RemoteControlUI/Private/UI/SRCPanelExposedEntitiesList.cpp
--- a/Engine/Plugins/VirtualProduction/RemoteControl/Source/RemoteControlUI/Private/UI/SRCPanelExposedEntitiesList.cpp
+++ b/Engine/Plugins/VirtualProduction/RemoteControl/Source/RemoteControlUI/Private/UI/SRCPanelExposedEntitiesList.cpp
@@ -42,8 +42,6 @@
 #define LOCTEXT_NAMESPACE "RemoteControlPanelEntitiesList"
 
 TSet<FName> SRCPanelExposedEntitiesList::DefaultProtocolColumns = {
-//	RemoteControlPresetColumns::LinkIdentifier,
-	RemoteControlPresetColumns::Mask,
 	RemoteControlPresetColumns::Status,
 	RemoteControlPresetColumns::BindingStatus,
 };
@@ -121,10 +119,6 @@ public:
 			{
 				HorizontalAlignment = HAlign_Center;
 			}
-			else if (InColumnName == RemoteControlPresetColumns::Mask)
-			{
-				HorizontalAlignment = HAlign_Left;
-			}
 
 			HorizontalBox->AddSlot()
 				.FillWidth(1.f)
@@ -2080,11 +2074,7 @@ SHeaderRow::FColumn::FArguments SRCPanelExposedEntitiesList::CreateColumn(const
 
 int32 SRCPanelExposedEntitiesList::GetColumnIndex(const FName& ForColumn) const
 {
-	if (ForColumn == RemoteControlPresetColumns::Mask)
-	{
-		return GetColumnIndex_Internal(ForColumn, RemoteControlPresetColumns::Description, ERCColumn::ERC_After);
-	}
-	else if (ForColumn == RemoteControlPresetColumns::Status)
+	if (ForColumn == RemoteControlPresetColumns::Status)
 	{
 		return GetColumnIndex_Internal(ForColumn, RemoteControlPresetColumns::Value, ERCColumn::ERC_After);
 	}
@@ -2156,10 +2146,6 @@ FText SRCPanelExposedEntitiesList::GetColumnLabel(const FName& ForColumn) const
 	{
 		return LOCTEXT("RCPresetPropertyIdColumnHeader_Label", "Property ID");
 	}
-	else if (ForColumn == RemoteControlPresetColumns::Mask)
-	{
-		return LOCTEXT("RCPresetMaskColumnHeader", "Mask");
-	}
 	else if (ForColumn == RemoteControlPresetColumns::Status)
 	{
 		return FText();
diff --git a/Engine/Plugins/VirtualProduction/RemoteControl/Source/RemoteControlUI/Private/UI/SRCPanelExposedField.cpp b/Engine/Plugins/VirtualProduction/RemoteControl/Source/RemoteControlUI/Private/UI/SRCPanelExposedField.cpp
--- a/Engine/Plugins/VirtualProduction/RemoteControl/Source/RemoteControlUI/Private/UI/SRCPanelExposedField.cpp
+++ b/Engine/Plugins/VirtualProduction/RemoteControl/Source/RemoteControlUI/Private/UI/SRCPanelExposedField.cpp
@@ -408,11 +408,7 @@ TSharedRef<SWidget> SRCPanelExposedField::GetWidget(const FName ForColumnName, c
 {
 	if (HasProtocolExtension())
 	{
-		if (ForColumnName == RemoteControlPresetColumns::Mask)
-		{
-			return SNew(SRCProtocolMask, WeakField);
-		}
-		else if (ForColumnName == RemoteControlPresetColumns::Status)
+		if (ForColumnName == RemoteControlPresetColumns::Status)
 		{
 			return SNew(STextBlock)
 				.Text(LOCTEXT("StatusText", "!"))
diff --git a/Engine/Plugins/VirtualProduction/RemoteControl/Source/RemoteControlUI/Private/UI/SRCPanelTreeNode.h b/Engine/Plugins/VirtualProduction/RemoteControl/Source/RemoteControlUI/Private/UI/SRCPanelTreeNode.h
--- a/Engine/Plugins/VirtualProduction/RemoteControl/Source/RemoteControlUI/Private/UI/SRCPanelTreeNode.h
+++ b/Engine/Plugins/VirtualProduction/RemoteControl/Source/RemoteControlUI/Private/UI/SRCPanelTreeNode.h
@@ -21,7 +21,6 @@ namespace RemoteControlPresetColumns
 	static FName OwnerName = TEXT("OwnerName");
 	static FName SubobjectPath = TEXT("Subobject Path");
 	static FName Description = TEXT("Description");
-	static FName Mask = TEXT("Mask");
 	static FName Value = TEXT("Value");
 	static FName BindingStatus = TEXT("BindingStatus");
 	static FName Status = TEXT("Status");
