Remote Control: fixed bug where properties with custom reset to default not working as expected

diff --git a/Engine/Plugins/VirtualProduction/RemoteControl/Source/RemoteControlUI/Private/UI/SRCPanelExposedField.cpp b/Engine/Plugins/VirtualProduction/RemoteControl/Source/RemoteControlUI/Private/UI/SRCPanelExposedField.cpp
--- a/Engine/Plugins/VirtualProduction/RemoteControl/Source/RemoteControlUI/Private/UI/SRCPanelExposedField.cpp
+++ b/Engine/Plugins/VirtualProduction/RemoteControl/Source/RemoteControlUI/Private/UI/SRCPanelExposedField.cpp
@@ -5,6 +5,8 @@
 #include "Algo/Transform.h"
 #include "Components/MeshComponent.h"
 #include "Components/SceneComponent.h"
+#include "DetailTreeNode.h"
+#include "DetailWidgetRow.h"
 #include "Editor/UnrealEdEngine.h"
 #include "EditorFontGlyphs.h"
 #include "Factories/IRCDefaultValueFactory.h"
@@ -455,6 +457,7 @@ TSharedRef<SWidget> SRCPanelExposedField::ConstructWidget()
 						TSharedPtr<SWidget> ValueWidget = SNullWidget::NullWidget;
 						TSharedPtr<SWidget> EditConditionWidget = SNullWidget::NullWidget;
 						TSharedPtr<IPropertyHandle> PropertyHandle;
+						TOptional<FResetToDefaultOverride> ResetToDefaultOverride;
 						bool bIsChildPropertyWidgetCreated = false;
 
 						if (!ExposedFieldUtils::IsFieldPathMatchingWithNodePath(Node, Field))
@@ -476,6 +479,14 @@ TSharedRef<SWidget> SRCPanelExposedField::ConstructWidget()
 							ExposedFieldUtils::CreateWidgets(Node, ValueWidget, EditConditionWidget);
 
 							PropertyHandle = Node->CreatePropertyHandle();
+
+							// This is doing similar operations to what IDetailTreeNode::CreateNodeWidgets already does (via ExposedFieldUtils::CreateWidgets)
+							// But this accesses the widget row that is not exposed by IDetailTreeNode::CreateNodeWidgets.
+							// Widget Row is needed to access custom behavior like custom reset to default.
+							FDetailWidgetRow WidgetRow;
+							TSharedPtr<FDetailTreeNode> DetailTreeNode = StaticCastSharedPtr<FDetailTreeNode>(Node);
+							DetailTreeNode->GenerateStandaloneWidget(WidgetRow);
+							ResetToDefaultOverride = MoveTemp(WidgetRow.CustomResetToDefault);
 						}
 
 						TSharedPtr<SHorizontalBox> FieldWidget = SNew(SHorizontalBox);
@@ -522,9 +533,16 @@ TSharedRef<SWidget> SRCPanelExposedField::ConstructWidget()
 									.ButtonStyle(&RCPanelStyle->FlatButtonStyle)
 									.ContentPadding(RCPanelStyle->PanelPadding)
 									.Visibility_Lambda(IsVisible)
-									.OnClicked_Lambda([PropertyHandle]()
+									.OnClicked_Lambda([PropertyHandle, ResetToDefaultOverride]()
 									{
-										PropertyHandle->ResetToDefault();
+										if (ResetToDefaultOverride.IsSet())
+										{
+											ResetToDefaultOverride->OnResetToDefaultClicked(PropertyHandle);
+										}
+										else
+										{
+											PropertyHandle->ResetToDefault();
+										}
 										return FReply::Handled();
 									})
 									.Content()
