Remote Control: trivial fix for static analysis warning

diff --git a/Engine/Plugins/VirtualProduction/RemoteControl/Source/RemoteControlUI/Private/UI/SRCPanelExposedEntitiesList.cpp b/Engine/Plugins/VirtualProduction/RemoteControl/Source/RemoteControlUI/Private/UI/SRCPanelExposedEntitiesList.cpp
--- a/Engine/Plugins/VirtualProduction/RemoteControl/Source/RemoteControlUI/Private/UI/SRCPanelExposedEntitiesList.cpp
+++ b/Engine/Plugins/VirtualProduction/RemoteControl/Source/RemoteControlUI/Private/UI/SRCPanelExposedEntitiesList.cpp
@@ -2074,14 +2074,7 @@ SHeaderRow::FColumn::FArguments SRCPanelExposedEntitiesList::CreateColumn(const
 
 int32 SRCPanelExposedEntitiesList::GetColumnIndex(const FName& ForColumn) const
 {
-	if (ForColumn == RemoteControlPresetColumns::Status)
-	{
-		return GetColumnIndex_Internal(ForColumn, RemoteControlPresetColumns::Value, ERCColumn::ERC_After);
-	}
-	else
-	{
-		return GetColumnIndex_Internal(ForColumn, RemoteControlPresetColumns::Value, ERCColumn::ERC_After);
-	}
+	return GetColumnIndex_Internal(ForColumn, RemoteControlPresetColumns::Value, ERCColumn::ERC_After);
 }
 
 int32 SRCPanelExposedEntitiesList::GetColumnIndex_Internal(const FName& ForColumn, const FName& ExistingColumnName, ERCColumn::Position InPosition) const
