Remote Control: Fixed generate unique label to now properly find a suitable name for a duplicate controller.

diff --git a/Engine/Plugins/VirtualProduction/RemoteControl/Source/RemoteControl/Private/RCVirtualPropertyContainer.cpp b/Engine/Plugins/VirtualProduction/RemoteControl/Source/RemoteControl/Private/RCVirtualPropertyContainer.cpp
--- a/Engine/Plugins/VirtualProduction/RemoteControl/Source/RemoteControl/Private/RCVirtualPropertyContainer.cpp
+++ b/Engine/Plugins/VirtualProduction/RemoteControl/Source/RemoteControl/Private/RCVirtualPropertyContainer.cpp
@@ -2,10 +2,9 @@
 
 #include "RCVirtualPropertyContainer.h"
 
-#include "UObject/StructOnScope.h"
-#include "Templates/SubclassOf.h"
-
 #include "RCVirtualProperty.h"
+#include "Templates/SubclassOf.h"
+#include "UObject/StructOnScope.h"
 
 #if WITH_EDITOR
 #include "ScopedTransaction.h"
@@ -13,6 +12,68 @@
 
 #define LOCTEXT_NAMESPACE "URCVirtualPropertyInContainer"
 
+namespace UE::RemoteControl::Private
+{
+	void SplitNameAndNumber(FString& InOutName, FString& OutNumber)
+	{
+		for (int32 Index = InOutName.Len() - 1; Index >= 0; Index--)
+		{
+			// Stop on first index that isn't a digit
+			if (!FChar::IsDigit(InOutName[Index]))
+			{
+				const FStringView NumberView = FStringView(InOutName).RightChop(Index + 1);
+				if (NumberView.Len() > 0)
+				{
+					// Process number (right-side) first before name gets modified inlined to only contain the left side
+					OutNumber = NumberView;
+					InOutName.LeftInline(Index + 1);
+					return;
+				}
+				// Fallback in case there was no digit found
+				OutNumber = TEXT("1");
+				return;
+			}
+		}
+
+		// The entire string is a number (never stopped at the first non-digit index)
+		OutNumber = MoveTemp(InOutName);
+		InOutName.Reset();
+	}
+
+	const FString::ElementType* IncrementNumber(FString& InNumber)
+	{
+		int32 LeadingZeros = 0;
+
+		// Count leading zeros
+		for (const FString::ElementType& Digit : InNumber)
+		{
+			if (Digit == TEXT('0'))
+			{
+				++LeadingZeros;
+			}
+			else
+			{
+				break;
+			}
+		}
+
+		// If the first non-zero digit is a 9, it means the incremented number will take up the space of a leading zero
+		if (LeadingZeros < InNumber.Len() && InNumber[LeadingZeros] == TEXT('9'))
+		{
+			--LeadingZeros;
+		}
+
+		TArray<FString::ElementType> LeadingZeroString;
+		if (LeadingZeros > 0)
+		{
+			LeadingZeroString.Init(TEXT('0'), LeadingZeros);
+		}
+
+		InNumber = FString(MoveTemp(LeadingZeroString)) + FString::FromInt(FCString::Atoi(*InNumber) + 1);
+		return *InNumber;
+	}
+}
+
 void URCVirtualPropertyContainerBase::AddVirtualProperty(URCVirtualPropertyBase* InVirtualProperty)
 {
 	if(ensure(InVirtualProperty))
@@ -306,62 +367,38 @@ FName URCVirtualPropertyContainerBase::GenerateUniquePropertyName(const FName& I
 
 FName URCVirtualPropertyContainerBase::GenerateUniquePropertyName(const FName& InPropertyName, const URCVirtualPropertyContainerBase* InContainer)
 {
-	auto GetFinalName = [](const FString& InPrefix, const int32 InIndex = 0)
-	{
-		FString FinalName = InPrefix;
-
-		if (InIndex > 0)
-		{
-			FinalName += TEXT("_") + FString::FromInt(InIndex);
-		}
-
-		return FinalName;
-	};
+	using namespace UE::RemoteControl;
 
-	int32 Index = 0;
-	const FString InitialName = InPropertyName.ToString();
-	FString FinalName = InitialName;
+	FString PropertyName = InPropertyName.ToString();
+	FString Prefix = PropertyName;
+	FString Number;
+	Private::SplitNameAndNumber(Prefix, Number);
 
-	// Recursively search for an available name by incrementing suffix till we find one.
-	const FPropertyBagPropertyDesc* PropertyDesc = InContainer->Bag.FindPropertyDescByName(*FinalName);
-	while (PropertyDesc)
+	// Recursively search for an available name by incrementing number until a unique name is found
+	while (InContainer->Bag.FindPropertyDescByName(*PropertyName))
 	{
-		++Index;
-		FinalName = GetFinalName(InitialName, Index);
-		PropertyDesc = InContainer->Bag.FindPropertyDescByName(*FinalName);
+		PropertyName = FString::Printf(TEXT("%s%s"), *Prefix, Private::IncrementNumber(Number));
 	}
 
-	return *FinalName;
+	return *PropertyName;
 }
 
 FName URCVirtualPropertyContainerBase::GenerateUniqueDisplayName(const FName& InPropertyName, const URCVirtualPropertyContainerBase* InContainer)
 {
-	auto GetFinalName = [](const FName& InPrefix, const int32 InIndex = 0)
-	{
-		FName FinalName = InPrefix;
-
-		if (InIndex > 0)
-		{
-			FinalName = FName(FinalName.ToString() + TEXT("_") + FString::FromInt(InIndex));
-		}
-
-		return FinalName;
-	};
+	using namespace UE::RemoteControl;
 
-	int32 Index = 0;
-	const FName InitialName = InPropertyName;
-	FName FinalName = InitialName;
+	FString DisplayName = InPropertyName.ToString();
+	FString Prefix = DisplayName;
+	FString Number;
+	Private::SplitNameAndNumber(Prefix, Number);
 
-	// Recursively search for an available name by incrementing suffix till we find one.
-	bool bControllerExist = InContainer->ControllerLabelToIdCache.Contains(FinalName);
-	while (bControllerExist)
+	// Recursively search for an available name by incrementing number until a unique name is found
+	while (InContainer->ControllerLabelToIdCache.Contains(*DisplayName))
 	{
-		++Index;
-		FinalName = GetFinalName(InitialName, Index);
-		bControllerExist = InContainer->ControllerLabelToIdCache.Contains(FinalName);
+		DisplayName = FString::Printf(TEXT("%s%s"), *Prefix, Private::IncrementNumber(Number));
 	}
 
-	return FinalName;
+	return *DisplayName;
 }
 
 void URCVirtualPropertyContainerBase::UpdateEntityIds(const TMap<FGuid, FGuid>& InEntityIdMap)
