Remote Control: removed unneeded checks leading to potential crashes

diff --git a/Engine/Plugins/VirtualProduction/RemoteControl/Source/RemoteControlProtocolWidgets/Private/ViewModels/ProtocolBindingViewModel.cpp b/Engine/Plugins/VirtualProduction/RemoteControl/Source/RemoteControlProtocolWidgets/Private/ViewModels/ProtocolBindingViewModel.cpp
--- a/Engine/Plugins/VirtualProduction/RemoteControl/Source/RemoteControlProtocolWidgets/Private/ViewModels/ProtocolBindingViewModel.cpp
+++ b/Engine/Plugins/VirtualProduction/RemoteControl/Source/RemoteControlProtocolWidgets/Private/ViewModels/ProtocolBindingViewModel.cpp
@@ -567,10 +567,15 @@ TSharedPtr<IRemoteControlProtocol> FProtocolBindingViewModel::GetProtocol() cons
 FRemoteControlProtocolMapping* FProtocolBindingViewModel::GetRangesMapping(const FGuid& InRangeId) const
 {
 	check(IsValid());
-	
-	FRemoteControlProtocolMapping* Mapping = GetBinding()->FindMapping(InRangeId);
-	check(InRangeId == Mapping->GetId());
-	return Mapping;
+
+	if (FRemoteControlProtocolBinding* Binding = GetBinding())
+	{
+		FRemoteControlProtocolMapping* Mapping = Binding->FindMapping(InRangeId);
+		ensure(Mapping && InRangeId == Mapping->GetId());
+		return Mapping;
+	}
+
+	return nullptr;
 }
 
 void FProtocolBindingViewModel::Remove() const
diff --git a/Engine/Plugins/VirtualProduction/RemoteControl/Source/RemoteControlProtocolWidgets/Private/ViewModels/ProtocolRangeViewModel.cpp b/Engine/Plugins/VirtualProduction/RemoteControl/Source/RemoteControlProtocolWidgets/Private/ViewModels/ProtocolRangeViewModel.cpp
--- a/Engine/Plugins/VirtualProduction/RemoteControl/Source/RemoteControlProtocolWidgets/Private/ViewModels/ProtocolRangeViewModel.cpp
+++ b/Engine/Plugins/VirtualProduction/RemoteControl/Source/RemoteControlProtocolWidgets/Private/ViewModels/ProtocolRangeViewModel.cpp
@@ -154,7 +154,10 @@ void FProtocolRangeViewModel::CopyInputValue(const TSharedPtr<IPropertyHandle>&
 {
 	check(IsValid());
 
-	GetRangesData()->CopyRawRangeData<FProperty>(InDstHandle);
+	if (FRemoteControlProtocolMapping* RangesData = GetRangesData())
+	{
+		RangesData->CopyRawRangeData<FProperty>(InDstHandle);
+	}
 }
 
 void FProtocolRangeViewModel::SetInputData(const TSharedPtr<IPropertyHandle>& InSrcHandle) const
