Remote Control: fixed issues with some mapped actions not working as expected with remote control panel

diff --git a/Engine/Plugins/VirtualProduction/RemoteControl/Source/RemoteControlUI/Private/RemoteControlUIModule.cpp b/Engine/Plugins/VirtualProduction/RemoteControl/Source/RemoteControlUI/Private/RemoteControlUIModule.cpp
--- a/Engine/Plugins/VirtualProduction/RemoteControl/Source/RemoteControlUI/Private/RemoteControlUIModule.cpp
+++ b/Engine/Plugins/VirtualProduction/RemoteControl/Source/RemoteControlUI/Private/RemoteControlUIModule.cpp
@@ -18,7 +18,6 @@
 #include "Engine/StaticMesh.h"
 #include "Framework/MultiBox/MultiBoxBuilder.h"
 #include "HAL/PlatformApplicationMisc.h"
-#include "Interfaces/IMainFrameModule.h"
 #include "IRemoteControlModule.h"
 #include "Kismet2/ComponentEditorUtils.h"
 #include "LevelEditorSubsystem.h"
@@ -696,29 +695,6 @@ void FRemoteControlUIModule::BindRemoteControlCommands()
 
 void FRemoteControlUIModule::UnbindRemoteControlCommands()
 {
-	if (FModuleManager::Get().IsModuleLoaded("MainFrame"))
-	{
-		const FRemoteControlCommands& Commands = FRemoteControlCommands::Get();
-
-		IMainFrameModule& MainFrame = FModuleManager::Get().LoadModuleChecked<IMainFrameModule>("MainFrame");
-
-		FUICommandList& ActionList = *MainFrame.GetMainFrameCommandBindings();
-
-		ActionList.UnmapAction(Commands.SavePreset);
-		ActionList.UnmapAction(FGlobalEditorCommonCommands::Get().FindInContentBrowser);
-		ActionList.UnmapAction(Commands.ActivateLogicMode);
-		ActionList.UnmapAction(Commands.ActivateDetailsMode);
-		ActionList.UnmapAction(Commands.ActivateSignatureMode);
-		ActionList.UnmapAction(Commands.ActivateProtocolsMode);
-		ActionList.UnmapAction(Commands.ActivateOutputLogMode);
-		ActionList.UnmapAction(Commands.DeleteEntity);
-		ActionList.UnmapAction(Commands.RenameEntity);
-		ActionList.UnmapAction(Commands.CopyItem);
-		ActionList.UnmapAction(Commands.PasteItem);
-		ActionList.UnmapAction(Commands.DuplicateItem);
-		ActionList.UnmapAction(Commands.UpdateValue);
-	}
-
 	FRemoteControlCommands::Unregister();
 }
 
diff --git a/Engine/Plugins/VirtualProduction/RemoteControl/Source/RemoteControlUI/Private/UI/Action/RCActionModel.cpp b/Engine/Plugins/VirtualProduction/RemoteControl/Source/RemoteControlUI/Private/UI/Action/RCActionModel.cpp
--- a/Engine/Plugins/VirtualProduction/RemoteControl/Source/RemoteControlUI/Private/UI/Action/RCActionModel.cpp
+++ b/Engine/Plugins/VirtualProduction/RemoteControl/Source/RemoteControlUI/Private/UI/Action/RCActionModel.cpp
@@ -6,7 +6,6 @@
 #include "Commands/RemoteControlCommands.h"
 #include "Controller/RCController.h"
 #include "IDetailTreeNode.h"
-#include "Interfaces/IMainFrameModule.h"
 #include "IPropertyRowGenerator.h"
 #include "RCVirtualProperty.h"
 #include "RCVirtualPropertyContainer.h"
diff --git a/Engine/Plugins/VirtualProduction/RemoteControl/Source/RemoteControlUI/Private/UI/Action/SRCVirtualPropertyWidget.cpp b/Engine/Plugins/VirtualProduction/RemoteControl/Source/RemoteControlUI/Private/UI/Action/SRCVirtualPropertyWidget.cpp
--- a/Engine/Plugins/VirtualProduction/RemoteControl/Source/RemoteControlUI/Private/UI/Action/SRCVirtualPropertyWidget.cpp
+++ b/Engine/Plugins/VirtualProduction/RemoteControl/Source/RemoteControlUI/Private/UI/Action/SRCVirtualPropertyWidget.cpp
@@ -3,7 +3,6 @@
 #include "SRCVirtualPropertyWidget.h"
 
 #include "Framework/MultiBox/MultiBoxBuilder.h"
-#include "Interfaces/IMainFrameModule.h"
 #include "PropertyHandle.h"
 #include "RCVirtualProperty.h"
 #include "SlateOptMacros.h"
diff --git a/Engine/Plugins/VirtualProduction/RemoteControl/Source/RemoteControlUI/Private/UI/BaseLogicUI/SRCLogicPanelListBase.cpp b/Engine/Plugins/VirtualProduction/RemoteControl/Source/RemoteControlUI/Private/UI/BaseLogicUI/SRCLogicPanelListBase.cpp
--- a/Engine/Plugins/VirtualProduction/RemoteControl/Source/RemoteControlUI/Private/UI/BaseLogicUI/SRCLogicPanelListBase.cpp
+++ b/Engine/Plugins/VirtualProduction/RemoteControl/Source/RemoteControlUI/Private/UI/BaseLogicUI/SRCLogicPanelListBase.cpp
@@ -3,7 +3,6 @@
 #include "SRCLogicPanelListBase.h"
 
 #include "Commands/RemoteControlCommands.h"
-#include "Interfaces/IMainFrameModule.h"
 #include "SRCLogicPanelBase.h"
 #include "UI/SRemoteControlPanel.h"
 
@@ -13,13 +12,12 @@ void SRCLogicPanelListBase::Construct(const FArguments& InArgs, const TSharedRef
 {
 	LogicPanelWeakPtr = InLogicParentPanel;
 	RemoteControlPanelWeakPtr = InPanel;
+	CommandList = InPanel->GetCommandList();
 }
 
 TSharedPtr<SWidget> SRCLogicPanelListBase::GetContextMenuWidget()
 {
-	IMainFrameModule& MainFrame = FModuleManager::Get().LoadModuleChecked<IMainFrameModule>("MainFrame");
-
-	FMenuBuilder MenuBuilder(true, MainFrame.GetMainFrameCommandBindings());
+	FMenuBuilder MenuBuilder(/*bShouldCloseWindowAfterMenuSelection*/true, CommandList);
 
 	// Special menu options
 	MenuBuilder.BeginSection("Advanced");
diff --git a/Engine/Plugins/VirtualProduction/RemoteControl/Source/RemoteControlUI/Private/UI/BaseLogicUI/SRCLogicPanelListBase.h b/Engine/Plugins/VirtualProduction/RemoteControl/Source/RemoteControlUI/Private/UI/BaseLogicUI/SRCLogicPanelListBase.h
--- a/Engine/Plugins/VirtualProduction/RemoteControl/Source/RemoteControlUI/Private/UI/BaseLogicUI/SRCLogicPanelListBase.h
+++ b/Engine/Plugins/VirtualProduction/RemoteControl/Source/RemoteControlUI/Private/UI/BaseLogicUI/SRCLogicPanelListBase.h
@@ -12,6 +12,7 @@
 #include "UI/Controller/RCControllerModel.h"
 #include "Widgets/SCompoundWidget.h"
 
+class FUICommandList;
 class SRCLogicPanelBase;
 class SRemoteControlPanel;
 
@@ -126,6 +127,9 @@ protected:
 		}
 	}
 
+	/** Command list to build context menus for */
+	TSharedPtr<FUICommandList> CommandList;
+
 	/** The parent Remote Control Panel widget*/
 	TWeakPtr<SRemoteControlPanel> RemoteControlPanelWeakPtr;
 
diff --git a/Engine/Plugins/VirtualProduction/RemoteControl/Source/RemoteControlUI/Private/UI/Controller/SRCControllerPanelList.cpp b/Engine/Plugins/VirtualProduction/RemoteControl/Source/RemoteControlUI/Private/UI/Controller/SRCControllerPanelList.cpp
--- a/Engine/Plugins/VirtualProduction/RemoteControl/Source/RemoteControlUI/Private/UI/Controller/SRCControllerPanelList.cpp
+++ b/Engine/Plugins/VirtualProduction/RemoteControl/Source/RemoteControlUI/Private/UI/Controller/SRCControllerPanelList.cpp
@@ -14,7 +14,6 @@
 #include "IDetailTreeNode.h"
 #include "IPropertyRowGenerator.h"
 #include "IRemoteControlModule.h"
-#include "Interfaces/IMainFrameModule.h"
 #include "Materials/MaterialInterface.h"
 #include "RCControllerModel.h"
 #include "RCMultiController.h"
diff --git a/Engine/Plugins/VirtualProduction/RemoteControl/Source/RemoteControlUI/Private/UI/SRCPanelExposedActor.cpp b/Engine/Plugins/VirtualProduction/RemoteControl/Source/RemoteControlUI/Private/UI/SRCPanelExposedActor.cpp
--- a/Engine/Plugins/VirtualProduction/RemoteControl/Source/RemoteControlUI/Private/UI/SRCPanelExposedActor.cpp
+++ b/Engine/Plugins/VirtualProduction/RemoteControl/Source/RemoteControlUI/Private/UI/SRCPanelExposedActor.cpp
@@ -17,7 +17,9 @@
 
 TSharedPtr<SRCPanelTreeNode> SRCPanelExposedActor::MakeInstance(const FGenerateWidgetArgs& Args)
 {
-	return SNew(SRCPanelExposedActor, StaticCastSharedPtr<FRemoteControlActor>(Args.Entity), Args.Preset, Args.ColumnSizeData).LiveMode(Args.bIsInLiveMode);
+	return SNew(SRCPanelExposedActor, StaticCastSharedPtr<FRemoteControlActor>(Args.Entity), Args.Preset, Args.ColumnSizeData)
+		.LiveMode(Args.bIsInLiveMode)
+		.CommandList(Args.CommandList);
 }
 
 void SRCPanelExposedActor::Construct(const FArguments& InArgs, TWeakPtr<FRemoteControlActor> InWeakActor, URemoteControlPreset* InPreset, FRCColumnSizeData InColumnSizeData)
@@ -27,8 +29,13 @@ void SRCPanelExposedActor::Construct(const FArguments& InArgs, TWeakPtr<FRemoteC
 	const TSharedPtr<FRemoteControlActor> Actor = InWeakActor.Pin();
 	if (ensure(Actor))
 	{
-		Initialize(Actor->GetId(), InPreset, InArgs._LiveMode);
-		
+		FInitParams InitParams;
+		InitParams.Preset = InPreset;
+		InitParams.EntityId = Actor->GetId();
+		InitParams.bLiveMode = InArgs._LiveMode;
+		InitParams.CommandList = InArgs._CommandList;
+		Initialize(InitParams);
+
 		ColumnSizeData = MoveTemp(InColumnSizeData);
 		WeakPreset = InPreset;
 		bLiveMode = InArgs._LiveMode;
diff --git a/Engine/Plugins/VirtualProduction/RemoteControl/Source/RemoteControlUI/Private/UI/SRCPanelExposedActor.h b/Engine/Plugins/VirtualProduction/RemoteControl/Source/RemoteControlUI/Private/UI/SRCPanelExposedActor.h
--- a/Engine/Plugins/VirtualProduction/RemoteControl/Source/RemoteControlUI/Private/UI/SRCPanelExposedActor.h
+++ b/Engine/Plugins/VirtualProduction/RemoteControl/Source/RemoteControlUI/Private/UI/SRCPanelExposedActor.h
@@ -25,6 +25,7 @@ struct SRCPanelExposedActor : public SRCPanelExposedEntity
 		: _LiveMode(false)
 	{}
 		SLATE_ATTRIBUTE(bool, LiveMode)
+		SLATE_ARGUMENT(TSharedPtr<FUICommandList>, CommandList)
 	SLATE_END_ARGS()
 
 	static TSharedPtr<SRCPanelTreeNode> MakeInstance(const FGenerateWidgetArgs& Args);
diff --git a/Engine/Plugins/VirtualProduction/RemoteControl/Source/RemoteControlUI/Private/UI/SRCPanelExposedEntitiesList.cpp b/Engine/Plugins/VirtualProduction/RemoteControl/Source/RemoteControlUI/Private/UI/SRCPanelExposedEntitiesList.cpp
--- a/Engine/Plugins/VirtualProduction/RemoteControl/Source/RemoteControlUI/Private/UI/SRCPanelExposedEntitiesList.cpp
+++ b/Engine/Plugins/VirtualProduction/RemoteControl/Source/RemoteControlUI/Private/UI/SRCPanelExposedEntitiesList.cpp
@@ -238,6 +238,8 @@ private:
 
 void SRCPanelExposedEntitiesList::Construct(const FArguments& InArgs, URemoteControlPreset* InPreset, TWeakPtr<FRCPanelWidgetRegistry> InWidgetRegistry)
 {
+	CommandList = InArgs._CommandList;
+
 	bIsInLiveMode = InArgs._LiveMode;
 	bIsInProtocolsMode = InArgs._ProtocolsMode;
 	Preset = TStrongObjectPtr<URemoteControlPreset>(InPreset);
@@ -1051,6 +1053,7 @@ void SRCPanelExposedEntitiesList::GenerateListWidgets()
 			Args.WidgetRegistry = WidgetRegistry;
 			Args.ColumnSizeData = ColumnSizeData;
 			Args.bIsInLiveMode = bIsInLiveMode;
+			Args.CommandList = CommandList;
 
 			FieldWidgetMap.Add(Entity->GetId(), FRemoteControlUIModule::Get().GenerateEntityWidget(Args));
 		}
@@ -1116,7 +1119,8 @@ void SRCPanelExposedEntitiesList::RefreshGroups()
 			.OnFieldDropEvent(this, &SRCPanelExposedEntitiesList::OnDropOnGroup)
 			.OnGetGroupId(this, &SRCPanelExposedEntitiesList::GetGroupId)
 			.OnDeleteGroup(this, &SRCPanelExposedEntitiesList::OnDeleteGroup)
-			.LiveMode(bIsInLiveMode);
+			.LiveMode(bIsInLiveMode)
+			.CommandList(CommandList);
 		
 		FieldGroups.Add(FieldGroup);
 		FieldGroup->GetNodes().Reserve(RCGroup.GetFields().Num());
@@ -1750,6 +1754,7 @@ void SRCPanelExposedEntitiesList::OnEntityAdded(const FGuid& InEntityId)
 	Args.ColumnSizeData = ColumnSizeData;
 	Args.bIsInLiveMode = bIsInLiveMode;
 	Args.Entity = Preset->GetExposedEntity(InEntityId).Pin();
+	Args.CommandList = CommandList;
 
 	ExposeEntity(FRemoteControlUIModule::Get().GenerateEntityWidget(Args));
 }
@@ -1786,7 +1791,8 @@ void SRCPanelExposedEntitiesList::OnGroupAdded(const FRemoteControlPresetGroup&
 		.OnFieldDropEvent(this, &SRCPanelExposedEntitiesList::OnDropOnGroup)
 		.OnGetGroupId(this, &SRCPanelExposedEntitiesList::GetGroupId)
 		.OnDeleteGroup(this, &SRCPanelExposedEntitiesList::OnDeleteGroup)
-		.LiveMode(bIsInLiveMode);
+		.LiveMode(bIsInLiveMode)
+		.CommandList(CommandList);
 	
 	FieldGroups.Add(FieldGroup);
 	
diff --git a/Engine/Plugins/VirtualProduction/RemoteControl/Source/RemoteControlUI/Private/UI/SRCPanelExposedEntitiesList.h b/Engine/Plugins/VirtualProduction/RemoteControl/Source/RemoteControlUI/Private/UI/SRCPanelExposedEntitiesList.h
--- a/Engine/Plugins/VirtualProduction/RemoteControl/Source/RemoteControlUI/Private/UI/SRCPanelExposedEntitiesList.h
+++ b/Engine/Plugins/VirtualProduction/RemoteControl/Source/RemoteControlUI/Private/UI/SRCPanelExposedEntitiesList.h
@@ -16,26 +16,27 @@
 #include "Widgets/SCompoundWidget.h"
 #include "Widgets/Views/STreeView.h"
 
-struct FRCPanelGroup;
-class FRCPanelWidgetRegistry;
 class FDragDropOperation;
-struct FGuid;
+class FRCPanelWidgetRegistry;
 class FReply;
-struct FRemoteControlProperty;
-struct FRemoteControlPresetGroup;
-struct FRemoteControlFunction;
 class ITableRow;
 class SComboButton;
 class SRCHeaderRow;
 class SRCPanelFilter;
 class SRCPanelGroup;
-struct SRCPanelTreeNode;
+class SRemoteControlPanel;
 class SRemoteControlTarget;
-struct SRCPanelExposedField;
 class SSearchBox;
 class STableViewBase;
 class URemoteControlPreset;
+struct FGuid;
+struct FRCPanelGroup;
 struct FRCPanelStyle;
+struct FRemoteControlFunction;
+struct FRemoteControlPresetGroup;
+struct FRemoteControlProperty;
+struct SRCPanelExposedField;
+struct SRCPanelTreeNode;
 
 enum class EEntitiesListMode : uint8
 {
@@ -74,6 +75,7 @@ public:
 		SLATE_ATTRIBUTE(TSharedPtr<SWidget>, ExposeFunctionsComboButton)
 		SLATE_ATTRIBUTE(TSharedPtr<SWidget>, ExposeActorsComboButton)
 		SLATE_EVENT(FSimpleDelegate, OnEntityListUpdated)
+		SLATE_ARGUMENT(TSharedPtr<FUICommandList>, CommandList)
 	SLATE_END_ARGS()
 
 	void Construct(const FArguments& InArgs, URemoteControlPreset* InPreset, TWeakPtr<FRCPanelWidgetRegistry> InWidgetRegistry);
@@ -267,6 +269,8 @@ private:
 	void RecallListSettings();
 
 private:
+	/** Command list to use if specified */
+	TSharedPtr<FUICommandList> CommandList;
 	/** Holds the Groups list view. */
 	TSharedPtr<SListView<TSharedPtr<SRCPanelTreeNode>>> GroupsListView;
 	/** Holds the Fields list view. */
diff --git a/Engine/Plugins/VirtualProduction/RemoteControl/Source/RemoteControlUI/Private/UI/SRCPanelExposedEntity.cpp b/Engine/Plugins/VirtualProduction/RemoteControl/Source/RemoteControlUI/Private/UI/SRCPanelExposedEntity.cpp
--- a/Engine/Plugins/VirtualProduction/RemoteControl/Source/RemoteControlUI/Private/UI/SRCPanelExposedEntity.cpp
+++ b/Engine/Plugins/VirtualProduction/RemoteControl/Source/RemoteControlUI/Private/UI/SRCPanelExposedEntity.cpp
@@ -14,7 +14,6 @@
 #include "Framework/MultiBox/MultiBoxBuilder.h"
 #include "GameFramework/Actor.h"
 #include "IRemoteControlModule.h"
-#include "Interfaces/IMainFrameModule.h"
 #include "Modules/ModuleManager.h"
 #include "RemoteControlBinding.h"
 #include "RemoteControlEntity.h"
@@ -82,9 +81,7 @@ TSharedPtr<FRemoteControlEntity> SRCPanelExposedEntity::GetEntity() const
 
 TSharedPtr<SWidget> SRCPanelExposedEntity::GetContextMenu()
 {
-	IMainFrameModule& MainFrame = FModuleManager::Get().LoadModuleChecked<IMainFrameModule>("MainFrame");
-
-	FMenuBuilder MenuBuilder(true, MainFrame.GetMainFrameCommandBindings());
+	FMenuBuilder MenuBuilder(/*bShouldCloseWindowAfterMenuSelection*/true, CommandList);
 
 	MenuBuilder.BeginSection("Common");
 
@@ -153,21 +150,27 @@ void SRCPanelExposedEntity::Refresh()
 {
 	if (EntityId.IsValid() && Preset.IsValid())
 	{
-		Initialize(EntityId, Preset.Get(), bLiveMode);
+		FInitParams InitParams;
+		InitParams.Preset = Preset.Get();
+		InitParams.EntityId = EntityId;
+		InitParams.bLiveMode = bLiveMode;
+		InitParams.CommandList = CommandList;
+		Initialize(InitParams);
 	}
 }
 
-void SRCPanelExposedEntity::Initialize(const FGuid& InEntityId, URemoteControlPreset* InPreset, const TAttribute<bool>& InbLiveMode)
+void SRCPanelExposedEntity::Initialize(const FInitParams& InParams)
 {
-	EntityId = InEntityId;
-	Preset = InPreset;
-	bLiveMode = InbLiveMode;
+	EntityId = InParams.EntityId;
+	Preset = InParams.Preset;
+	bLiveMode = InParams.bLiveMode;
+	CommandList = InParams.CommandList;
 
 	RCPanelStyle = &FRemoteControlPanelStyle::Get()->GetWidgetStyle<FRCPanelStyle>("RemoteControlPanel.MinorPanel");
 
-	if (ensure(InPreset))
+	if (ensure(InParams.Preset))
 	{
-		if (const TSharedPtr<FRemoteControlEntity> RCEntity = InPreset->GetExposedEntity(InEntityId).Pin())
+		if (const TSharedPtr<FRemoteControlEntity> RCEntity = InParams.Preset->GetExposedEntity(InParams.EntityId).Pin())
 		{
 			const FString BindingPath = RCEntity->GetLastBindingPath().ToString();
 			CachedLabel = RCEntity->GetLabel();
diff --git a/Engine/Plugins/VirtualProduction/RemoteControl/Source/RemoteControlUI/Private/UI/SRCPanelExposedEntity.h b/Engine/Plugins/VirtualProduction/RemoteControl/Source/RemoteControlUI/Private/UI/SRCPanelExposedEntity.h
--- a/Engine/Plugins/VirtualProduction/RemoteControl/Source/RemoteControlUI/Private/UI/SRCPanelExposedEntity.h
+++ b/Engine/Plugins/VirtualProduction/RemoteControl/Source/RemoteControlUI/Private/UI/SRCPanelExposedEntity.h
@@ -39,8 +39,15 @@ struct SRCPanelExposedEntity : public SRCPanelTreeNode
 	virtual void Refresh() override;
 	
 protected:
-	void Initialize(const FGuid& InEntityId, URemoteControlPreset* InPreset, const TAttribute<bool>& InbLiveMode);
-	
+	struct FInitParams
+	{
+		FGuid EntityId;
+		URemoteControlPreset* Preset;
+		TAttribute<bool> bLiveMode;
+		TSharedPtr<FUICommandList> CommandList;
+	};
+	void Initialize(const FInitParams& InParams);
+
 	/** Create a widget that displays the rebind button. */
 	TSharedRef<SWidget> CreateInvalidWidget(const FText& InErrorText = FText::GetEmpty());
 
@@ -60,6 +67,8 @@ protected:
 	TWeakObjectPtr<URemoteControlPreset> Preset;
 	/** Whether the panel is in live mode or not. */
 	TAttribute<bool> bLiveMode;
+	/** Command list to use for the context menu */
+	TSharedPtr<FUICommandList> CommandList;
 	/** Name of the entity's owner (What actor it's bound to). */
 	FName CachedOwnerName;
 	/** Path from the entity owner to the object holding the exposed property. */
diff --git a/Engine/Plugins/VirtualProduction/RemoteControl/Source/RemoteControlUI/Private/UI/SRCPanelExposedField.cpp b/Engine/Plugins/VirtualProduction/RemoteControl/Source/RemoteControlUI/Private/UI/SRCPanelExposedField.cpp
--- a/Engine/Plugins/VirtualProduction/RemoteControl/Source/RemoteControlUI/Private/UI/SRCPanelExposedField.cpp
+++ b/Engine/Plugins/VirtualProduction/RemoteControl/Source/RemoteControlUI/Private/UI/SRCPanelExposedField.cpp
@@ -131,14 +131,14 @@ TSharedPtr<SRCPanelTreeNode> SRCPanelExposedField::MakeInstance(const FGenerateW
 	return SNew(SRCPanelExposedField, StaticCastSharedPtr<FRemoteControlField>(Args.Entity), Args.ColumnSizeData, Args.WidgetRegistry)
 		.Preset(Args.Preset)
 		.LiveMode(Args.bIsInLiveMode)
-		.HighlightText(Args.HighlightText);
+		.HighlightText(Args.HighlightText)
+		.CommandList(Args.CommandList);
 }
 
 void SRCPanelExposedField::Construct(const FArguments& InArgs, TWeakPtr<FRemoteControlField> InField, FRCColumnSizeData InColumnSizeData, TWeakPtr<FRCPanelWidgetRegistry> InWidgetRegistry)
 {
 	TRACE_CPUPROFILER_EVENT_SCOPE(SRCPanelExposedField::Construct);
 
-
 	WeakField = MoveTemp(InField);
 
 	ColumnSizeData = MoveTemp(InColumnSizeData);
@@ -150,7 +150,12 @@ void SRCPanelExposedField::Construct(const FArguments& InArgs, TWeakPtr<FRemoteC
 
 	if (TSharedPtr<FRemoteControlField> FieldPtr = WeakField.Pin())
 	{
-		Initialize(FieldPtr->GetId(), InArgs._Preset.Get(), InArgs._LiveMode);
+		FInitParams InitParams;
+		InitParams.Preset = InArgs._Preset.Get();
+		InitParams.EntityId = FieldPtr->GetId();
+		InitParams.bLiveMode = InArgs._LiveMode;
+		InitParams.CommandList = InArgs._CommandList;
+		Initialize(InitParams);
 
 		CachedLabel = FieldPtr->GetLabel();
 		PropertyIdLabel = FieldPtr->PropertyId;
diff --git a/Engine/Plugins/VirtualProduction/RemoteControl/Source/RemoteControlUI/Private/UI/SRCPanelExposedField.h b/Engine/Plugins/VirtualProduction/RemoteControl/Source/RemoteControlUI/Private/UI/SRCPanelExposedField.h
--- a/Engine/Plugins/VirtualProduction/RemoteControl/Source/RemoteControlUI/Private/UI/SRCPanelExposedField.h
+++ b/Engine/Plugins/VirtualProduction/RemoteControl/Source/RemoteControlUI/Private/UI/SRCPanelExposedField.h
@@ -36,6 +36,7 @@ struct SRCPanelExposedField : public SRCPanelExposedEntity
 		SLATE_ATTRIBUTE(bool, LiveMode)
 		SLATE_ATTRIBUTE(FText, HighlightText)
 		SLATE_ATTRIBUTE(URemoteControlPreset*, Preset)
+		SLATE_ARGUMENT(TSharedPtr<FUICommandList>, CommandList)
 	SLATE_END_ARGS()
 
 	static TSharedPtr<SRCPanelTreeNode> MakeInstance(const FGenerateWidgetArgs& Args);
diff --git a/Engine/Plugins/VirtualProduction/RemoteControl/Source/RemoteControlUI/Private/UI/SRCPanelFieldGroup.cpp b/Engine/Plugins/VirtualProduction/RemoteControl/Source/RemoteControlUI/Private/UI/SRCPanelFieldGroup.cpp
--- a/Engine/Plugins/VirtualProduction/RemoteControl/Source/RemoteControlUI/Private/UI/SRCPanelFieldGroup.cpp
+++ b/Engine/Plugins/VirtualProduction/RemoteControl/Source/RemoteControlUI/Private/UI/SRCPanelFieldGroup.cpp
@@ -5,7 +5,6 @@
 #include "Commands/RemoteControlCommands.h"
 #include "Framework/MultiBox/MultiBoxBuilder.h"
 #include "IRemoteControlModule.h"
-#include "Interfaces/IMainFrameModule.h"
 #include "RemoteControlPanelStyle.h"
 #include "RemoteControlPreset.h"
 #include "SDropTarget.h"
@@ -139,9 +138,7 @@ void SRCPanelGroup::GetNodeChildren(TArray<TSharedPtr<SRCPanelTreeNode>>& OutChi
 
 TSharedPtr<SWidget> SRCPanelGroup::GetContextMenu()
 {
-	IMainFrameModule& MainFrame = FModuleManager::Get().LoadModuleChecked<IMainFrameModule>("MainFrame");
-
-	FMenuBuilder MenuBuilder(true, MainFrame.GetMainFrameCommandBindings());
+	FMenuBuilder MenuBuilder(/*bShouldCloseWindowAfterMenuSelection*/true, CommandList);
 
 	MenuBuilder.BeginSection("Common");
 
diff --git a/Engine/Plugins/VirtualProduction/RemoteControl/Source/RemoteControlUI/Private/UI/SRCPanelFieldGroup.h b/Engine/Plugins/VirtualProduction/RemoteControl/Source/RemoteControlUI/Private/UI/SRCPanelFieldGroup.h
--- a/Engine/Plugins/VirtualProduction/RemoteControl/Source/RemoteControlUI/Private/UI/SRCPanelFieldGroup.h
+++ b/Engine/Plugins/VirtualProduction/RemoteControl/Source/RemoteControlUI/Private/UI/SRCPanelFieldGroup.h
@@ -41,7 +41,7 @@ public:
 		SLATE_EVENT(FOnGetGroupId, OnGetGroupId)
 		SLATE_EVENT(FOnDeleteGroup, OnDeleteGroup)
 		SLATE_ATTRIBUTE(bool, LiveMode)
-		
+		SLATE_ARGUMENT(TSharedPtr<FUICommandList>, CommandList)
 	SLATE_END_ARGS()
 
 	void Tick(const FGeometry&, const double, const float);
@@ -88,6 +88,8 @@ private:
 	const FSlateBrush* HandleGroupColor() const;
 
 private:
+	/** Command list to use for context menu actions */
+	TSharedPtr<FUICommandList> CommandList;
 	/** Event called when something is dropped on this group. */
 	FOnFieldDropEvent OnFieldDropEvent;
 	/** Getter for this group's name. */
diff --git a/Engine/Plugins/VirtualProduction/RemoteControl/Source/RemoteControlUI/Private/UI/SRemoteControlPanel.cpp b/Engine/Plugins/VirtualProduction/RemoteControl/Source/RemoteControlUI/Private/UI/SRemoteControlPanel.cpp
--- a/Engine/Plugins/VirtualProduction/RemoteControl/Source/RemoteControlUI/Private/UI/SRemoteControlPanel.cpp
+++ b/Engine/Plugins/VirtualProduction/RemoteControl/Source/RemoteControlUI/Private/UI/SRemoteControlPanel.cpp
@@ -25,7 +25,6 @@
 #include "ISettingsModule.h"
 #include "IStructureDetailsView.h"
 #include "Input/Reply.h"
-#include "Interfaces/IMainFrameModule.h"
 #include "Kismet/BlueprintFunctionLibrary.h"
 #include "Layout/Visibility.h"
 #include "Materials/Material.h"
@@ -142,61 +141,6 @@ namespace RemoteControlPanelUtils
 	}
 }
 
-class FRemoteControlPanelInputProcessor : public IInputProcessor
-{
-public:
-	FRemoteControlPanelInputProcessor(TSharedPtr<SRemoteControlPanel> InOwner)
-	{
-		Owner = InOwner;
-	}
-
-	void SetOwner(TSharedPtr<SRemoteControlPanel> InOwner)
-	{
-		Owner = InOwner;
-	}
-
-	virtual ~FRemoteControlPanelInputProcessor() = default;
-
-	virtual void Tick(const float DeltaTime, FSlateApplication& SlateApp, TSharedRef<ICursor> Cursor) override
-	{
-	}
-
-	virtual bool HandleKeyDownEvent(FSlateApplication& SlateApp, const FKeyEvent& InKeyEvent) override
-	{
-		if (const TSharedPtr<SRemoteControlPanel> PinnedOwner = Owner.Pin())
-		{
-			if (PinnedOwner->IsModeActive(ERCPanelMode::Live))
-			{
-				return false;
-			}
-
-			if (!SlateApp.HasFocusedDescendants(PinnedOwner.ToSharedRef()))
-			{
-				return false;
-			}
-
-			if (InKeyEvent.GetKey() == EKeys::Delete)
-			{
-				PinnedOwner->DeleteEntity();
-				return true;
-			}
-
-			if (InKeyEvent.GetKey() == EKeys::F2)
-			{
-				PinnedOwner->RenameEntity();
-				return true;
-			}
-		}
-
-		return false;
-	}
-
-	virtual const TCHAR* GetDebugName() const override { return TEXT("RemoteControlInputInterceptor"); }
-
-private:
-	TWeakPtr<SRemoteControlPanel> Owner;
-};
-
 void SRemoteControlPanel::Construct(const FArguments& InArgs, URemoteControlPreset* InPreset, TSharedPtr<IToolkitHost> InToolkitHost)
 {
 	OnLiveModeChange = InArgs._OnLiveModeChange;
@@ -268,7 +212,8 @@ void SRemoteControlPanel::Construct(const FArguments& InArgs, URemoteControlPres
 				UpdateEntityDetailsView(EntityList->GetSelectedEntity());
 				UpdateRebindButtonVisibility();
 				CachedExposedPropertyArgs.Reset();
-			});
+			})
+		.CommandList(CommandList);
 
 	EntityList->OnSelectionChange().AddSP(this, &SRemoteControlPanel::UpdateEntityDetailsView);
 
@@ -464,9 +409,6 @@ void SRemoteControlPanel::Construct(const FArguments& InArgs, URemoteControlPres
 			]
 		];
 
-	InputProcessor = MakeShared<FRemoteControlPanelInputProcessor>(SharedThis(this));
-	FSlateApplication::Get().RegisterInputPreProcessor(InputProcessor);
-
 	RegisterEvents();
 	CacheLevelClasses();
 	Refresh();
@@ -480,8 +422,6 @@ SRemoteControlPanel::~SRemoteControlPanel()
 	SaveSettings();
 	UnregisterEvents();
 
-	FSlateApplication::Get().UnregisterInputPreProcessor(InputProcessor);
-
 	// Clear the log
 	FRemoteControlLogger::Get().ClearLog();
 
@@ -505,6 +445,15 @@ void SRemoteControlPanel::Tick(const FGeometry& AllottedGeometry, const double I
 	}
 }
 
+FReply SRemoteControlPanel::OnKeyDown(const FGeometry& InGeometry, const FKeyEvent& InKeyEvent)
+{
+	if (CommandList.IsValid() && CommandList->ProcessCommandBindings(InKeyEvent))
+	{
+		return FReply::Handled();
+	}
+	return FReply::Unhandled();
+}
+
 bool SRemoteControlPanel::IsExposed(const FRCExposesPropertyArgs& InPropertyArgs)
 {
 	if (!ensure(InPropertyArgs.IsValid()))
@@ -1225,32 +1174,30 @@ TSharedRef<SWidget> SRemoteControlPanel::BuildSignaturesModeContent()
 
 void SRemoteControlPanel::BindRemoteControlCommands()
 {
-	const FRemoteControlCommands& Commands = FRemoteControlCommands::Get();
+	CommandList = MakeShared<FUICommandList>();
 
-	IMainFrameModule& MainFrame = FModuleManager::Get().LoadModuleChecked<IMainFrameModule>("MainFrame");
-
-	FUICommandList& ActionList = *MainFrame.GetMainFrameCommandBindings();
+	const FRemoteControlCommands& Commands = FRemoteControlCommands::Get();
 
-	ActionList.MapAction(
+	CommandList->MapAction(
 		Commands.SavePreset,
 		FExecuteAction::CreateSP(this, &SRemoteControlPanel::SaveAsset),
 		FCanExecuteAction(),
 		FGetActionCheckState(),
 		FIsActionButtonVisible::CreateSP(this, &SRemoteControlPanel::CanSaveAsset));
 
-	ActionList.MapAction(
+	CommandList->MapAction(
 		FGlobalEditorCommonCommands::Get().FindInContentBrowser,
 		FExecuteAction::CreateSP(this, &SRemoteControlPanel::FindInContentBrowser),
 		FCanExecuteAction::CreateSP(this, &SRemoteControlPanel::CanFindInContentBrowser));
 
-	auto MapModeAction = [&ActionList, this](const TSharedPtr<FUICommandInfo>& InCommand, ERCPanelMode InMode)
+	auto MapModeAction = [&CommandList = CommandList, This=this](const TSharedPtr<FUICommandInfo>& InCommand, ERCPanelMode InMode)
 		{
-			ActionList.MapAction(
+			CommandList->MapAction(
 				InCommand,
-				FExecuteAction::CreateSP(this, &SRemoteControlPanel::SetActiveMode, InMode),
-				FCanExecuteAction::CreateSP(this, &SRemoteControlPanel::CanActivateMode, InMode),
-				FIsActionChecked::CreateSP(this, &SRemoteControlPanel::IsModeActive, InMode),
-				FIsActionButtonVisible::CreateSP(this, &SRemoteControlPanel::CanActivateMode, InMode));
+				FExecuteAction::CreateSP(This, &SRemoteControlPanel::SetActiveMode, InMode),
+				FCanExecuteAction::CreateSP(This, &SRemoteControlPanel::CanActivateMode, InMode),
+				FIsActionChecked::CreateSP(This, &SRemoteControlPanel::IsModeActive, InMode),
+				FIsActionButtonVisible::CreateSP(This, &SRemoteControlPanel::CanActivateMode, InMode));
 		};
 
 	MapModeAction(Commands.ActivateLogicMode    , ERCPanelMode::Controller);
@@ -1259,43 +1206,43 @@ void SRemoteControlPanel::BindRemoteControlCommands()
 	MapModeAction(Commands.ActivateProtocolsMode, ERCPanelMode::Protocols);
 	MapModeAction(Commands.ActivateOutputLogMode, ERCPanelMode::OutputLog);
 
-	ActionList.MapAction(
+	CommandList->MapAction(
 		Commands.DeleteEntity,
 		FExecuteAction::CreateSP(this, &SRemoteControlPanel::DeleteEntity_Execute),
 		FCanExecuteAction::CreateSP(this, &SRemoteControlPanel::CanDeleteEntity));
 
-	ActionList.MapAction(
+	CommandList->MapAction(
 		Commands.RenameEntity,
 		FExecuteAction::CreateSP(this, &SRemoteControlPanel::RenameEntity_Execute),
 		FCanExecuteAction::CreateSP(this, &SRemoteControlPanel::CanRenameEntity));
 
-	ActionList.MapAction(
+	CommandList->MapAction(
 		Commands.ChangePropId,
 		FExecuteAction::CreateSP(this, &SRemoteControlPanel::ChangePropertyId_Execute),
 		FCanExecuteAction::CreateSP(this, &SRemoteControlPanel::CanChangePropertyId));
 
-	ActionList.MapAction(
+	CommandList->MapAction(
 		Commands.CopyItem,
 		FExecuteAction::CreateSP(this, &SRemoteControlPanel::CopyItem_Execute),
 		FCanExecuteAction(),
 		FGetActionCheckState(),
 		FIsActionButtonVisible::CreateSP(this, &SRemoteControlPanel::CanCopyItem));
 
-	ActionList.MapAction(
+	CommandList->MapAction(
 		Commands.PasteItem,
 		FExecuteAction::CreateSP(this, &SRemoteControlPanel::PasteItem_Execute),
 		FCanExecuteAction(),
 		FGetActionCheckState(),
 		FIsActionButtonVisible::CreateSP(this, &SRemoteControlPanel::CanPasteItem));
 
-	ActionList.MapAction(
+	CommandList->MapAction(
 		Commands.DuplicateItem,
 		FExecuteAction::CreateSP(this, &SRemoteControlPanel::DuplicateItem_Execute),
 		FCanExecuteAction(),
 		FGetActionCheckState(),
 		FIsActionButtonVisible::CreateSP(this, &SRemoteControlPanel::CanDuplicateItem));
 
-	ActionList.MapAction(
+	CommandList->MapAction(
 		Commands.UpdateValue,
 		FExecuteAction::CreateSP(this, &SRemoteControlPanel::UpdateValue_Execute),
 		FCanExecuteAction(),
@@ -1812,11 +1759,10 @@ void SRemoteControlPanel::GenerateToolbar()
 		];
 
 	UToolMenus* ToolMenus = UToolMenus::Get();
-	IMainFrameModule& MainFrame = FModuleManager::Get().LoadModuleChecked<IMainFrameModule>("MainFrame");
 
 	UToolMenu* GeneratedToolbar = ToolMenus->FindMenu(DefaultRemoteControlPanelToolBarName);
 
-	GeneratedToolbar->Context = FToolMenuContext(MainFrame.GetMainFrameCommandBindings());
+	GeneratedToolbar->Context = FToolMenuContext(CommandList);
 
 	TSharedRef<class SWidget> ToolBarWidget = ToolMenus->GenerateWidget(GeneratedToolbar);
 
@@ -2015,11 +1961,10 @@ void SRemoteControlPanel::GenerateAuxiliaryToolbar()
 		];
 
 	UToolMenus* ToolMenus = UToolMenus::Get();
-	IMainFrameModule& MainFrame = FModuleManager::Get().LoadModuleChecked<IMainFrameModule>("MainFrame");
 
 	UToolMenu* GeneratedToolbar = ToolMenus->FindMenu(AuxiliaryRemoteControlPanelToolBarName);
 
-	GeneratedToolbar->Context = FToolMenuContext(MainFrame.GetMainFrameCommandBindings());
+	GeneratedToolbar->Context = FToolMenuContext(CommandList);
 
 	TSharedRef<class SWidget> ToolBarWidget = ToolMenus->GenerateWidget(GeneratedToolbar);
 
diff --git a/Engine/Plugins/VirtualProduction/RemoteControl/Source/RemoteControlUI/Private/UI/SRemoteControlPanel.h b/Engine/Plugins/VirtualProduction/RemoteControl/Source/RemoteControlUI/Private/UI/SRemoteControlPanel.h
--- a/Engine/Plugins/VirtualProduction/RemoteControl/Source/RemoteControlUI/Private/UI/SRemoteControlPanel.h
+++ b/Engine/Plugins/VirtualProduction/RemoteControl/Source/RemoteControlUI/Private/UI/SRemoteControlPanel.h
@@ -30,38 +30,38 @@ enum class ERCPanelMode : uint8
 };
 
 class AActor;
-struct EVisibility;
-struct FAssetData;
 class FExposedEntityDragDrop;
-struct FListEntry;
-struct FRCPanelDrawerArgs;
+class FRCBehaviourModel;
+class FRCControllerModel;
 class FRCPanelWidgetRegistry;
-struct FRCPanelStyle;
-struct FRemoteControlEntity;
 class FReply;
-class IToolkitHost;
-class IPropertyRowGenerator;
+class FUICommandList;
 class IPropertyHandle;
+class IPropertyRowGenerator;
+class IToolkitHost;
 class SBox;
 class SClassViewer;
 class SComboButton;
-struct SRCPanelTreeNode;
-class SRCPanelFunctionPicker;
 class SRCActionPanel;
 class SRCLogicPanelBase;
 class SRCPanelDrawer;
 class SRCPanelExposedEntitiesList;
 class SRCPanelFilter;
+class SRCPanelFunctionPicker;
 class SRemoteControlPanel;
 class SSearchBox;
 class STextBlock;
-class URemoteControlPreset;
-
-class URCController;
-class URCBehaviour;
 class URCAction;
-class FRCControllerModel;
-class FRCBehaviourModel;
+class URCBehaviour;
+class URCController;
+class URemoteControlPreset;
+struct EVisibility;
+struct FAssetData;
+struct FListEntry;
+struct FRCPanelDrawerArgs;
+struct FRCPanelStyle;
+struct FRemoteControlEntity;
+struct SRCPanelTreeNode;
 
 DECLARE_DELEGATE_TwoParams(FOnLiveModeChange, TSharedPtr<SRemoteControlPanel> /* Panel */, bool /* bEditModeChange */);
 
@@ -94,6 +94,7 @@ public:
 
 	//~ Begin SWidget interface
 	virtual void Tick(const FGeometry& AllottedGeometry, const double InCurrentTime, const float InDeltaTime);
+	virtual FReply OnKeyDown(const FGeometry& InGeometry, const FKeyEvent& InKeyEvent) override;
 	//~ End SWidget interface
 
 	/**
@@ -158,6 +159,11 @@ public:
 		return ActionPanel;
 	}
 
+	TSharedPtr<FUICommandList> GetCommandList() const
+	{
+		return CommandList;
+	}
+
 	/** Retrieves the number of controllers. */
 	int32 NumControllerItems() const;
 
@@ -396,6 +402,8 @@ private:
 	static const FName TargetWorldRemoteControlPanelMenuName;
 	/** Holds the preset asset. */
 	TStrongObjectPtr<URemoteControlPreset> Preset;
+	/** Command list of this panel */
+	TSharedPtr<FUICommandList> CommandList;
 	/** Delegate called when the live mode changes. */
 	FOnLiveModeChange OnLiveModeChange;
 	/** Holds the blueprint library picker */
@@ -450,8 +458,6 @@ private:
 	const FRCPanelStyle* RCPanelStyle;
 	/** Stores the active panel that is drawn. */
 	ERCPanelMode ActiveMode = ERCPanelMode::Controller;
-	/** Input Preprocessor which catches the Delete Key when Docked. */
-	TSharedPtr<IInputProcessor> InputProcessor;
 	/** Currently selected world name */
 	FString SelectedWorldName;
 	// ~ Remote Control Logic Panels ~
diff --git a/Engine/Plugins/VirtualProduction/RemoteControl/Source/RemoteControlUI/Public/IRemoteControlUIModule.h b/Engine/Plugins/VirtualProduction/RemoteControl/Source/RemoteControlUI/Public/IRemoteControlUIModule.h
--- a/Engine/Plugins/VirtualProduction/RemoteControl/Source/RemoteControlUI/Public/IRemoteControlUIModule.h
+++ b/Engine/Plugins/VirtualProduction/RemoteControl/Source/RemoteControlUI/Public/IRemoteControlUIModule.h
@@ -13,12 +13,13 @@
 class FDelegateHandle;
 class FPropertyPath;
 class FRCPanelWidgetRegistry;
+class FUICommandList;
 class IDetailCategoryBuilder;
 class IDetailLayoutBuilder;
 class IPropertyHandle;
+class IRCExposedEntitiesPanelExtender;
 class IRCPanelExposedEntitiesGroupWidgetFactory;
 class IRCPanelExposedEntitiesListSettingsForProtocol;
-class IRCExposedEntitiesPanelExtender;
 class IRCPanelExposedEntityWidgetFactory;
 class IRCSignatureCustomization;
 class URCVirtualPropertyBase;
@@ -43,6 +44,7 @@ struct FGenerateWidgetArgs
 	TSharedPtr<FRemoteControlEntity> Entity;
 	TAttribute<bool> bIsInLiveMode;
 	TWeakPtr<FRCPanelWidgetRegistry> WidgetRegistry;
+	TSharedPtr<FUICommandList> CommandList;
 	FText HighlightText;
 };
 
