Remote Control: Controllers to start numeric from 0001 by default, and to auto-populate its description if being created / bound to an entity (and there's no existing description for it).

diff --git a/Engine/Plugins/VirtualProduction/RemoteControl/Source/RemoteControl/Private/RemoteControlPreset.cpp b/Engine/Plugins/VirtualProduction/RemoteControl/Source/RemoteControl/Private/RemoteControlPreset.cpp
--- a/Engine/Plugins/VirtualProduction/RemoteControl/Source/RemoteControl/Private/RemoteControlPreset.cpp
+++ b/Engine/Plugins/VirtualProduction/RemoteControl/Source/RemoteControl/Private/RemoteControlPreset.cpp
@@ -893,7 +893,9 @@ URCVirtualPropertyInContainer* URemoteControlPreset::AddController(TSubclassOf<U
 	FName NewPropertyName = InPropertyName;
 	if (NewPropertyName.IsNone())
 	{
-		NewPropertyName = URCVirtualPropertyContainerBase::GenerateUniquePropertyName(TEXT(""), InValueType, InValueTypeObject, ControllerContainer);
+		// Property name must always be set. Fallback to generic numeric naming
+		// Ensuring uniqueness is dealt in URCVirtualPropertyContainerBase::AddProperty
+		NewPropertyName = TEXT("0001");
 	}
 
 #if WITH_EDITOR
diff --git a/Engine/Plugins/VirtualProduction/RemoteControl/Source/RemoteControlUI/Private/UI/Controller/RCControllerModel.cpp b/Engine/Plugins/VirtualProduction/RemoteControl/Source/RemoteControlUI/Private/UI/Controller/RCControllerModel.cpp
--- a/Engine/Plugins/VirtualProduction/RemoteControl/Source/RemoteControlUI/Private/UI/Controller/RCControllerModel.cpp
+++ b/Engine/Plugins/VirtualProduction/RemoteControl/Source/RemoteControlUI/Private/UI/Controller/RCControllerModel.cpp
@@ -19,36 +19,43 @@ FRCControllerModel::FRCControllerModel(URCVirtualPropertyBase* InVirtualProperty
 	: FRCLogicModeBase(InRemoteControlPanel)
 	, VirtualPropertyWeakPtr(InVirtualProperty)
 	, DetailTreeNodeWeakPtr(InTreeNode)
+	, CurrentControlValueType(EPropertyBagPropertyType::None)
 {
-	if (ensure(InVirtualProperty))
+	Id = FGuid::NewGuid();
+}
+
+void FRCControllerModel::Initialize()
+{
+	URCVirtualPropertyBase* const VirtualProperty = VirtualPropertyWeakPtr.Get();
+	const TSharedPtr<IDetailTreeNode> TreeNode = DetailTreeNodeWeakPtr.Pin();
+
+	if (ensure(VirtualProperty && TreeNode.IsValid()))
 	{
-		if (InVirtualProperty->DisplayName.IsNone())
+		if (VirtualProperty->DisplayName.IsNone())
 		{
-			InVirtualProperty->DisplayName = InVirtualProperty->PropertyName;
+			VirtualProperty->DisplayName = VirtualProperty->PropertyName;
 		}
 
 		SAssignNew(ControllerNameTextBox, SEditableTextBox)
+			.Text(this, &FRCControllerModel::GetControllerDisplayName)
 			.RevertTextOnEscape(true)
 			.SelectAllTextWhenFocused(true)
-			.Text(FText::FromName(InVirtualProperty->DisplayName))
-			.OnTextCommitted_Raw(this, &FRCControllerModel::OnControllerNameCommitted);
+			.OnTextCommitted(this, &FRCControllerModel::OnControllerNameCommitted);
 
 		SAssignNew(ControllerDescriptionTextBox, SInlineEditableTextBlock)
-			.Text(InVirtualProperty->Description)
+			.Text(this, &FRCControllerModel::GetControllerDescription)
 			.MultiLine(true)
-			.OnTextCommitted_Raw(this, &FRCControllerModel::OnControllerDescriptionCommitted);
+			.OnTextCommitted(this, &FRCControllerModel::OnControllerDescriptionCommitted);
 
 		SAssignNew(ControllerFieldIdTextBox, SInlineEditableTextBlock)
-			.Text(FText::FromName(InVirtualProperty->FieldId))
-			.OnTextCommitted_Raw(this, &FRCControllerModel::OnControllerFieldIdCommitted);
+			.Text(FText::FromName(VirtualProperty->FieldId))
+			.OnTextCommitted(this, &FRCControllerModel::OnControllerFieldIdCommitted);
 
-		if (const TSharedPtr<IPropertyHandle>& PropertyHandle = InTreeNode->CreatePropertyHandle())
+		if (const TSharedPtr<IPropertyHandle>& PropertyHandle = TreeNode->CreatePropertyHandle())
 		{
-			PropertyHandle->SetOnPropertyValueChangedWithData(TDelegate<void(const FPropertyChangedEvent&)>::CreateRaw(this, &FRCControllerModel::OnPropertyValueChanged));
+			PropertyHandle->SetOnPropertyValueChangedWithData(TDelegate<void(const FPropertyChangedEvent&)>::CreateSP(this, &FRCControllerModel::OnPropertyValueChanged));
 		}
 	}
-
-	Id = FGuid::NewGuid();
 }
 
 TSharedRef<SWidget> FRCControllerModel::GetWidget() const
@@ -341,23 +348,21 @@ void FRCControllerModel::EnterDescriptionEditingMode()
 	ControllerDescriptionTextBox->EnterEditingMode();
 }
 
-FName FRCControllerModel::GetControllerDisplayName()
+FText FRCControllerModel::GetControllerDisplayName() const
 {
 	if (URCVirtualPropertyBase* Controller = GetVirtualProperty())
 	{
-		return Controller->DisplayName;
+		return FText::FromName(Controller->DisplayName);
 	}
-
-	return NAME_None;
+	return FText::GetEmpty();
 }
 
-FText FRCControllerModel::GetControllerDescription()
+FText FRCControllerModel::GetControllerDescription() const
 {
 	if (URCVirtualPropertyBase* Controller = GetVirtualProperty())
 	{
 		return Controller->Description;
 	}
-
 	return FText::GetEmpty();
 }
 
diff --git a/Engine/Plugins/VirtualProduction/RemoteControl/Source/RemoteControlUI/Private/UI/Controller/RCControllerModel.h b/Engine/Plugins/VirtualProduction/RemoteControl/Source/RemoteControlUI/Private/UI/Controller/RCControllerModel.h
--- a/Engine/Plugins/VirtualProduction/RemoteControl/Source/RemoteControlUI/Private/UI/Controller/RCControllerModel.h
+++ b/Engine/Plugins/VirtualProduction/RemoteControl/Source/RemoteControlUI/Private/UI/Controller/RCControllerModel.h
@@ -34,7 +34,10 @@ public:
 	DECLARE_MULTICAST_DELEGATE_OneParam(FOnValueChanged, TSharedPtr<FRCControllerModel> /* InControllerModel */);
 	
 	FRCControllerModel(URCVirtualPropertyBase* InVirtualProperty, const TSharedRef<IDetailTreeNode>& InTreeNode, const TSharedPtr<SRemoteControlPanel> InRemoteControlPanel);
-	
+
+	/** Called right after object construction to initialize the model (and allow using SharedThis, etc) */
+	void Initialize();
+
 	/**
 	 * The widget to be rendered for this Controller
 	 * Used to represent a single row when added to the Controllers Panel List
@@ -78,10 +81,10 @@ public:
 	void EnterDescriptionEditingMode();
 
 	/** User-friendly Name of the underlying Controller */
-	FName GetControllerDisplayName();
+	FText GetControllerDisplayName() const;
 
 	/** Description of the underlying Controller */
-	FText GetControllerDescription();
+	FText GetControllerDescription() const;
 
 	/**Fetches the unique Id for this UI item */
 	FGuid GetId() const { return Id; }
diff --git a/Engine/Plugins/VirtualProduction/RemoteControl/Source/RemoteControlUI/Private/UI/Controller/SRCControllerItemRow.cpp b/Engine/Plugins/VirtualProduction/RemoteControl/Source/RemoteControlUI/Private/UI/Controller/SRCControllerItemRow.cpp
--- a/Engine/Plugins/VirtualProduction/RemoteControl/Source/RemoteControlUI/Private/UI/Controller/SRCControllerItemRow.cpp
+++ b/Engine/Plugins/VirtualProduction/RemoteControl/Source/RemoteControlUI/Private/UI/Controller/SRCControllerItemRow.cpp
@@ -204,18 +204,40 @@ FReply SRCControllerItemRow::OnAcceptDrop(const FDragDropEvent& InDragDropEvent,
 
 void SRCControllerItemRow::CreateBindBehaviorsFromEntities(const FExposedEntityDragDrop& InExposedEntityDragDrop, const FDragDropContext& InContext)
 {
-	FScopedTransaction Transaction(LOCTEXT("CreateBindBehaviors", "Create Bind Behaviors from Drag Drop"));
+	const TArray<FGuid>& DroppedFieldIds = InExposedEntityDragDrop.GetSelectedFieldsId();
+
+	if (!ControllerItem.IsValid() || DroppedFieldIds.IsEmpty())
+	{
+		return;
+	}
+
+	URCController* const Controller = Cast<URCController>(ControllerItem->GetVirtualProperty());
+	if (!Controller)
+	{
+		return;
+	}
+
+	FScopedTransaction Transaction(LOCTEXT("BindPropertiesToController", "Bind properties to Controller"));
 
 	bool bModified = false;
-	for (const FGuid& ExposedEntityId : InExposedEntityDragDrop.GetSelectedFieldsId())
+
+	// Update controller description if empty to match the first dragged property
+	if (Controller->Description.IsEmpty())
+	{
+		if (TSharedPtr<const FRemoteControlProperty> RemoteControlProperty = InContext.Preset->GetExposedEntity<FRemoteControlProperty>(DroppedFieldIds[0]).Pin())
+		{
+			Controller->Modify();
+			Controller->Description = UE::RCUIHelpers::GenerateControllerDescriptionFromEntity(RemoteControlProperty);
+			bModified = true;
+		}
+	}
+
+	for (const FGuid& ExposedEntityId : DroppedFieldIds)
 	{
 		if (TSharedPtr<const FRemoteControlProperty> RemoteControlProperty = InContext.Preset->GetExposedEntity<FRemoteControlProperty>(ExposedEntityId).Pin())
 		{
-			if (URCController* Controller = Cast<URCController>(ControllerItem->GetVirtualProperty()))
-			{
-				InContext.ControllerPanelList->CreateBindBehaviourAndAssignTo(Controller, RemoteControlProperty.ToSharedRef(), true);
-				bModified = true;
-			}
+			InContext.ControllerPanelList->CreateBindBehaviourAndAssignTo(Controller, RemoteControlProperty.ToSharedRef(), true);
+			bModified = true;
 		}
 	}
 
diff --git a/Engine/Plugins/VirtualProduction/RemoteControl/Source/RemoteControlUI/Private/UI/Controller/SRCControllerPanelList.cpp b/Engine/Plugins/VirtualProduction/RemoteControl/Source/RemoteControlUI/Private/UI/Controller/SRCControllerPanelList.cpp
--- a/Engine/Plugins/VirtualProduction/RemoteControl/Source/RemoteControlUI/Private/UI/Controller/SRCControllerPanelList.cpp
+++ b/Engine/Plugins/VirtualProduction/RemoteControl/Source/RemoteControlUI/Private/UI/Controller/SRCControllerPanelList.cpp
@@ -212,6 +212,7 @@ void SRCControllerPanelList::Reset()
 					if (ensureAlways(ControllerItems.IsValidIndex(Controller->DisplayIndex)))
 					{
 						const TSharedRef<FRCControllerModel> ControllerModel = MakeShared<FRCControllerModel>(Controller, Child, RemoteControlPanel);
+						ControllerModel->Initialize();
 						ControllerModel->OnValueChanged.AddSP(this, &SRCControllerPanelList::OnControllerValueChanged, bIsMultiController);
 						if (bIsMultiController)
 						{
diff --git a/Engine/Plugins/VirtualProduction/RemoteControl/Source/RemoteControlUI/Private/UI/RCUIHelpers.cpp b/Engine/Plugins/VirtualProduction/RemoteControl/Source/RemoteControlUI/Private/UI/RCUIHelpers.cpp
--- a/Engine/Plugins/VirtualProduction/RemoteControl/Source/RemoteControlUI/Private/UI/RCUIHelpers.cpp
+++ b/Engine/Plugins/VirtualProduction/RemoteControl/Source/RemoteControlUI/Private/UI/RCUIHelpers.cpp
@@ -13,6 +13,7 @@
 #include "GraphEditorSettings.h"
 #include "IDetailTreeNode.h"
 #include "IPropertyRowGenerator.h"
+#include "Kismet/KismetSystemLibrary.h"
 #include "Modules/ModuleManager.h"
 #include "PropertyEditorModule.h"
 #include "PropertyHandle.h"
@@ -219,6 +220,31 @@ TSharedRef<SWidget> UE::RCUIHelpers::GetTypeColorWidget(const FProperty* InPrope
 	return TypeColorWidget;
 }
 
+UObject* UE::RCUIHelpers::GetEntityOwner(const TSharedPtr<const FRemoteControlEntity>& InEntity)
+{
+	if (UObject* BoundObject = InEntity->GetBoundObject())
+	{
+		if (AActor* OwnerActor = BoundObject->GetTypedOuter<AActor>())
+		{
+			return OwnerActor;
+		}
+		return BoundObject;
+	}
+	return nullptr;
+}
+
+FText UE::RCUIHelpers::GenerateControllerDescriptionFromEntity(const TSharedPtr<const FRemoteControlEntity>& InEntity)
+{
+	if (!InEntity.IsValid())
+	{
+		return FText::GetEmpty();
+	}
+
+	return FText::Format(INVTEXT("{0} - {1}")
+		, FText::FromString(UKismetSystemLibrary::GetDisplayName(GetEntityOwner(InEntity)))
+		, FText::FromName(InEntity->GetLabel()));
+}
+
 URCController* UE::RCUIHelpers::CreateControllerFromEntity(URemoteControlPreset* InPreset, const TSharedPtr<const FRemoteControlProperty>& InRemoteControlProperty)
 {
 	if (!InPreset || !InRemoteControlProperty.IsValid())
@@ -250,8 +276,9 @@ URCController* UE::RCUIHelpers::CreateControllerFromEntity(URemoteControlPreset*
 	}
 
 	// Create a Controller of matching type
-	URCController* NewController = CastChecked<URCController>(InPreset->AddController(URCController::StaticClass(), PropertyBagType, StructObject, InRemoteControlProperty->FieldName));
+	URCController* const NewController = CastChecked<URCController>(InPreset->AddController(URCController::StaticClass(), PropertyBagType, StructObject));
 	NewController->DisplayIndex = InPreset->GetNumControllers() - 1;
+	NewController->Description = GenerateControllerDescriptionFromEntity(InRemoteControlProperty);
 
 	// Add metadata to this controller, if this is a custom controller
 	if (!CustomControllerName.IsEmpty())
diff --git a/Engine/Plugins/VirtualProduction/RemoteControl/Source/RemoteControlUI/Private/UI/RCUIHelpers.h b/Engine/Plugins/VirtualProduction/RemoteControl/Source/RemoteControlUI/Private/UI/RCUIHelpers.h
--- a/Engine/Plugins/VirtualProduction/RemoteControl/Source/RemoteControlUI/Private/UI/RCUIHelpers.h
+++ b/Engine/Plugins/VirtualProduction/RemoteControl/Source/RemoteControlUI/Private/UI/RCUIHelpers.h
@@ -13,6 +13,7 @@ class SWidget;
 class URCController;
 class URCVirtualPropertySelfContainer;
 class URemoteControlPreset;
+struct FRemoteControlEntity;
 struct FRemoteControlProperty;
 
 namespace UE::RCUIHelpers
@@ -62,6 +63,12 @@ namespace UE::RCUIHelpers
 	*/
 	TSharedRef<SWidget> GetTypeColorWidget(const FProperty* InProperty);
 
+	/** Retrieves the entity owner for a given entity. If the entity is bound to an actor's subobject, the actor will be returned. */
+	UObject* GetEntityOwner(const TSharedPtr<const FRemoteControlEntity>& InEntity);
+
+	/** Generates the description to use for a new or existing controller, from the given property */
+	FText GenerateControllerDescriptionFromEntity(const TSharedPtr<const FRemoteControlEntity>& InEntity);
+
 	/** Automatically creates a controller for a given remote control property */
 	URCController* CreateControllerFromEntity(URemoteControlPreset* InPreset, const TSharedPtr<const FRemoteControlProperty>& InRemoteControlProperty);
 }
\ No newline at end of file
diff --git a/Engine/Plugins/VirtualProduction/RemoteControl/Source/RemoteControlUI/Private/UI/SRCPanelExposedEntity.cpp b/Engine/Plugins/VirtualProduction/RemoteControl/Source/RemoteControlUI/Private/UI/SRCPanelExposedEntity.cpp
--- a/Engine/Plugins/VirtualProduction/RemoteControl/Source/RemoteControlUI/Private/UI/SRCPanelExposedEntity.cpp
+++ b/Engine/Plugins/VirtualProduction/RemoteControl/Source/RemoteControlUI/Private/UI/SRCPanelExposedEntity.cpp
@@ -14,7 +14,9 @@
 #include "Framework/MultiBox/MultiBoxBuilder.h"
 #include "GameFramework/Actor.h"
 #include "IRemoteControlModule.h"
+#include "Kismet/KismetSystemLibrary.h"
 #include "Modules/ModuleManager.h"
+#include "RCUIHelpers.h"
 #include "RemoteControlBinding.h"
 #include "RemoteControlEntity.h"
 #include "RemoteControlField.h"
@@ -184,25 +186,11 @@ void SRCPanelExposedEntity::Initialize(const FInitParams& InParams)
 
 			FName OwnerFName;
 			// If the binding is valid, display the actor label if possible
-			if (UObject* Object = RCEntity->GetBoundObject())
+			if (UObject* Owner = UE::RCUIHelpers::GetEntityOwner(RCEntity))
 			{
 				bValidBinding = true;
-				
-				if (AActor* OwnerActor = Object->GetTypedOuter<AActor>())
-				{
-					CachedOwnerName = *OwnerActor->GetActorLabel();
-					OwnerFName = OwnerActor->GetFName();
-				}
-				else if (AActor* Actor = Cast<AActor>(Object))
-				{
-					CachedOwnerName = *Actor->GetActorLabel();
-					OwnerFName = Object->GetFName();
-				}
-				else
-				{
-					CachedOwnerName = Object->GetFName();
-					OwnerFName = Object->GetFName();
-				}
+				OwnerFName = Owner->GetFName();
+				CachedOwnerName = *UKismetSystemLibrary::GetDisplayName(Owner);
 			}
 			else
 			{
