Remote Control: added transaction when deleting multiple items.

diff --git a/Engine/Plugins/VirtualProduction/RemoteControl/Source/RemoteControlUI/Private/UI/Action/SRCActionPanelList.h b/Engine/Plugins/VirtualProduction/RemoteControl/Source/RemoteControlUI/Private/UI/Action/SRCActionPanelList.h
--- a/Engine/Plugins/VirtualProduction/RemoteControl/Source/RemoteControlUI/Private/UI/Action/SRCActionPanelList.h
+++ b/Engine/Plugins/VirtualProduction/RemoteControl/Source/RemoteControlUI/Private/UI/Action/SRCActionPanelList.h
@@ -148,7 +148,11 @@ public:
 	/** Deletes currently selected items from the list view */
 	virtual void DeleteSelectedPanelItems() override
 	{
-		DeleteItemsFromLogicPanel<ActionType>(ActionItems, ListView->GetSelectedItems());
+		FScopedTransaction Transaction(LOCTEXT("DeleteSelectedItems", "Delete Selected Items"));
+		if (!DeleteItemsFromLogicPanel<ActionType>(ActionItems, ListView->GetSelectedItems()))
+		{
+			Transaction.Cancel();
+		}
 	}
 
 	/** Returns the UI items currently selected by the user (if any). */
diff --git a/Engine/Plugins/VirtualProduction/RemoteControl/Source/RemoteControlUI/Private/UI/BaseLogicUI/SRCLogicPanelListBase.h b/Engine/Plugins/VirtualProduction/RemoteControl/Source/RemoteControlUI/Private/UI/BaseLogicUI/SRCLogicPanelListBase.h
--- a/Engine/Plugins/VirtualProduction/RemoteControl/Source/RemoteControlUI/Private/UI/BaseLogicUI/SRCLogicPanelListBase.h
+++ b/Engine/Plugins/VirtualProduction/RemoteControl/Source/RemoteControlUI/Private/UI/BaseLogicUI/SRCLogicPanelListBase.h
@@ -90,10 +90,13 @@ private:
 protected:
 	TWeakPtr<SWidget> ContextMenuWidgetCached;
 
-	/** Helper function for handling common Delete Item functionality across all child panels (Actions/Behaviours/Controllers)
-	* Currently invoked from each Panel List child class with appropriate template class*/
+	/**
+	 * Helper function for handling common Delete Item functionality across all child panels (Actions/Behaviours/Controllers)
+	 * Currently invoked from each Panel List child class with appropriate template class
+	 * @return true if any items were deleted
+	 */
 	template<class T>
-	void DeleteItemsFromLogicPanel(TArray<TSharedPtr<T>>& ItemsSource, const TArray<TSharedPtr<T>>& SelectedItems)
+	bool DeleteItemsFromLogicPanel(TArray<TSharedPtr<T>>& ItemsSource, const TArray<TSharedPtr<T>>& SelectedItems)
 	{
 		bool bIsDeleted = false;
 		for (const TSharedPtr<T> SelectedItem : SelectedItems)
@@ -125,6 +128,7 @@ protected:
 			BroadcastOnItemRemoved();
 			Reset();
 		}
+		return bIsDeleted;
 	}
 
 	/** Command list to build context menus for */
diff --git a/Engine/Plugins/VirtualProduction/RemoteControl/Source/RemoteControlUI/Private/UI/Behaviour/SRCBehaviourPanelList.cpp b/Engine/Plugins/VirtualProduction/RemoteControl/Source/RemoteControlUI/Private/UI/Behaviour/SRCBehaviourPanelList.cpp
--- a/Engine/Plugins/VirtualProduction/RemoteControl/Source/RemoteControlUI/Private/UI/Behaviour/SRCBehaviourPanelList.cpp
+++ b/Engine/Plugins/VirtualProduction/RemoteControl/Source/RemoteControlUI/Private/UI/Behaviour/SRCBehaviourPanelList.cpp
@@ -433,7 +433,11 @@ bool SRCBehaviourPanelList::IsListFocused() const
 
 void SRCBehaviourPanelList::DeleteSelectedPanelItems()
 {
-	DeleteItemsFromLogicPanel<FRCBehaviourModel>(BehaviourItems, ListView->GetSelectedItems());
+	FScopedTransaction Transaction(LOCTEXT("DeleteSelectedItems", "Delete Selected Items"));
+	if (!DeleteItemsFromLogicPanel<FRCBehaviourModel>(BehaviourItems, ListView->GetSelectedItems()))
+	{
+		Transaction.Cancel();
+	}
 }
 
 TArray<TSharedPtr<FRCLogicModeBase>> SRCBehaviourPanelList::GetSelectedLogicItems()
diff --git a/Engine/Plugins/VirtualProduction/RemoteControl/Source/RemoteControlUI/Private/UI/Controller/SRCControllerPanelList.cpp b/Engine/Plugins/VirtualProduction/RemoteControl/Source/RemoteControlUI/Private/UI/Controller/SRCControllerPanelList.cpp
--- a/Engine/Plugins/VirtualProduction/RemoteControl/Source/RemoteControlUI/Private/UI/Controller/SRCControllerPanelList.cpp
+++ b/Engine/Plugins/VirtualProduction/RemoteControl/Source/RemoteControlUI/Private/UI/Controller/SRCControllerPanelList.cpp
@@ -479,7 +479,11 @@ bool SRCControllerPanelList::IsListFocused() const
 
 void SRCControllerPanelList::DeleteSelectedPanelItems()
 {
-	DeleteItemsFromLogicPanel<FRCControllerModel>(ControllerItems, ListView->GetSelectedItems());
+	FScopedTransaction Transaction(LOCTEXT("DeleteSelectedItems", "Delete Selected Items"));
+	if (!DeleteItemsFromLogicPanel<FRCControllerModel>(ControllerItems, ListView->GetSelectedItems()))
+	{
+		Transaction.Cancel();
+	}
 }
 
 TArray<TSharedPtr<FRCControllerModel>> SRCControllerPanelList::GetSelectedControllers() const
diff --git a/Engine/Plugins/VirtualProduction/RemoteControl/Source/RemoteControlUI/Private/UI/Signature/SRCSignatureTree.cpp b/Engine/Plugins/VirtualProduction/RemoteControl/Source/RemoteControlUI/Private/UI/Signature/SRCSignatureTree.cpp
--- a/Engine/Plugins/VirtualProduction/RemoteControl/Source/RemoteControlUI/Private/UI/Signature/SRCSignatureTree.cpp
+++ b/Engine/Plugins/VirtualProduction/RemoteControl/Source/RemoteControlUI/Private/UI/Signature/SRCSignatureTree.cpp
@@ -206,7 +206,10 @@ void SRCSignatureTree::DeleteSelectedPanelItems()
 	}
 
 	FScopedTransaction Transaction(LOCTEXT("RemoveSelectedItems", "Remove Selected Items"));
-	DeleteItemsFromLogicPanel<FRCSignatureTreeItemBase>(RootItem->GetChildrenMutable(), SelectedItems);
+	if (!DeleteItemsFromLogicPanel<FRCSignatureTreeItemBase>(RootItem->GetChildrenMutable(), SelectedItems))
+	{
+		Transaction.Cancel();
+	}
 
 	RootItem->GetSelection()->ClearSelection();
 }
