Remote Control: remove meta-data for locks and scrubbing delta as it's not desired behavior for most cases

diff --git a/Engine/Plugins/VirtualProduction/RemoteControl/Source/RemoteControl/Public/RCVirtualProperty.h b/Engine/Plugins/VirtualProduction/RemoteControl/Source/RemoteControl/Public/RCVirtualProperty.h
--- a/Engine/Plugins/VirtualProduction/RemoteControl/Source/RemoteControl/Public/RCVirtualProperty.h
+++ b/Engine/Plugins/VirtualProduction/RemoteControl/Source/RemoteControl/Public/RCVirtualProperty.h
@@ -69,13 +69,23 @@ public:
 	/** Return property bag property type */
 	EPropertyBagPropertyType GetValueType() const;
 
-	/** Get Metadata from the internal PropertyBag */
+	/**
+	 * Get the meta-data matching the given key.
+	 * NOTE: This meta-data is available at runtime and is different from the property bag property's meta-data
+	 * @return the value mapped to the key
+	 */
 	FString GetMetadataValue(FName Key) const;
 
-	/** Add/Set Metadata into the internal PropertyBag */
+	/**
+	 * Adds/sets the meta-data with the given key and value.
+	 * NOTE: This meta-data is available at runtime and is different from the property bag property's meta-data
+	 */
 	void SetMetadataValue(FName Key, FString Data);
 
-	/** Remove Metadata from the internal PropertyBag */
+	/**
+	 * Removes the meta-data matching the given key.
+	 * NOTE: This meta-data is available at runtime and is different from the property bag property's meta-data
+	 */
 	void RemoveMetadataValue(FName Key);
 
 	/** Return pointer to object that defines the Enum, Struct, or Class. */
diff --git a/Engine/Plugins/VirtualProduction/RemoteControl/Source/RemoteControlLogic/Private/Controller/RCControllerContainer.cpp b/Engine/Plugins/VirtualProduction/RemoteControl/Source/RemoteControlLogic/Private/Controller/RCControllerContainer.cpp
--- a/Engine/Plugins/VirtualProduction/RemoteControl/Source/RemoteControlLogic/Private/Controller/RCControllerContainer.cpp
+++ b/Engine/Plugins/VirtualProduction/RemoteControl/Source/RemoteControlLogic/Private/Controller/RCControllerContainer.cpp
@@ -28,25 +28,11 @@ URCVirtualPropertyInContainer* URCControllerContainer::AddProperty(const FName&
 {
 	if (InValueType == EPropertyBagPropertyType::Struct)
 	{
-		// Vector Controllers
-		if (InValueTypeObject == TBaseStructure<FVector>::Get() || InValueTypeObject == TBaseStructure<FVector2D>::Get())
-		{
-			MetaData.Add(FPropertyBagPropertyDescMetaData(FName("Delta"), FString::Printf(TEXT("%f"), VectorSliderDelta)));
-			MetaData.Add(FPropertyBagPropertyDescMetaData(FName("LinearDeltaSensitivity"), FString::Printf(TEXT("%f"), VectorLinearDeltaSensitivity)));
-			MetaData.Add(FPropertyBagPropertyDescMetaData(FName("AllowPreserveRatio"), TEXT("")));
-		}
 		// Color Controllers
-		else if (InValueTypeObject == TBaseStructure<FColor>::Get())
+		if (InValueTypeObject == TBaseStructure<FColor>::Get())
 		{
 			MetaData.Add(FPropertyBagPropertyDescMetaData(FName("OnlyUpdateOnInteractionEnd"), TEXT("true")));
 		}
-		// Rotator Controllers
-		else if (InValueTypeObject == TBaseStructure<FRotator>::Get())
-		{
-			MetaData.Add(FPropertyBagPropertyDescMetaData(FName("Delta"), FString::Printf(TEXT("%f"), RotatorSliderDelta)));
-			MetaData.Add(FPropertyBagPropertyDescMetaData(FName("LinearDeltaSensitivity"), FString::Printf(TEXT("%f"), RotatorLinearDeltaSensitivity)));
-			MetaData.Add(FPropertyBagPropertyDescMetaData(FName("AllowPreserveRatio"), TEXT("")));
-		}
 	}
 
 	return Super::AddProperty(InPropertyName, InPropertyClass, InValueType, InValueTypeObject, MetaData);
diff --git a/Engine/Plugins/VirtualProduction/RemoteControl/Source/RemoteControlLogic/Public/Controller/RCControllerContainer.h b/Engine/Plugins/VirtualProduction/RemoteControl/Source/RemoteControlLogic/Public/Controller/RCControllerContainer.h
--- a/Engine/Plugins/VirtualProduction/RemoteControl/Source/RemoteControlLogic/Public/Controller/RCControllerContainer.h
+++ b/Engine/Plugins/VirtualProduction/RemoteControl/Source/RemoteControlLogic/Public/Controller/RCControllerContainer.h
@@ -25,23 +25,6 @@ public:
 	virtual URCVirtualPropertyInContainer* AddProperty(const FName& InPropertyName, TSubclassOf<URCVirtualPropertyInContainer> InPropertyClass, const EPropertyBagPropertyType InValueType, UObject* InValueTypeObject = nullptr, TArray<FPropertyBagPropertyDescMetaData> MetaData = TArray<FPropertyBagPropertyDescMetaData>()) override;
 	
 private:
-	/** Slider exponent for Vector Controller sliders. Determines how fast the value change travels while scrubbing the X/Y/Z sliders in UI
-	* The default value approximates the behaviour of SceneComponent Location Vectors (see SNumericVectorInputBox in ComponentTransformDetails.cpp) */
-	UPROPERTY()
-	float VectorSliderDelta = 5.f;
-
-	/** Linear Delta Sensitivity for Vector Controller sliders. Multiplied inversely with mouse movement to determine actual slider change */
-	UPROPERTY()
-	float VectorLinearDeltaSensitivity = 10.f;
-
-	/** Slider exponent for Rotator Controller sliders.*/
-	UPROPERTY()
-	float RotatorSliderDelta = 10.f;
-
-	/** Linear Delta Sensitivity for Rotator Controller sliders*/
-	UPROPERTY()
-	float RotatorLinearDeltaSensitivity = 10.f;
-
 #if WITH_EDITOR
 	/** Fetches the Controller underlying a given a Property Changed Event*/
 	URCController* GetControllerFromChangeEvent(const FPropertyChangedEvent& Event);
