Remote Control: Fix property ids no longer working, and in some cases causing crashes.

diff --git a/Engine/Plugins/VirtualProduction/RemoteControl/Source/RemoteControlLogic/Private/Action/RCPropertyIdAction.cpp b/Engine/Plugins/VirtualProduction/RemoteControl/Source/RemoteControlLogic/Private/Action/RCPropertyIdAction.cpp
--- a/Engine/Plugins/VirtualProduction/RemoteControl/Source/RemoteControlLogic/Private/Action/RCPropertyIdAction.cpp
+++ b/Engine/Plugins/VirtualProduction/RemoteControl/Source/RemoteControlLogic/Private/Action/RCPropertyIdAction.cpp
@@ -137,9 +137,10 @@ void URCPropertyIdAction::UpdatePropertyId()
 							continue;
 						}
 
-						FString PrefixPropertyClassName = TargetRCProperty->PropertyId.ToString() + TEXT(".");
-						FName PropertyClassName = FName(PrefixPropertyClassName + PropertyIdHandler->GetPropertySubTypeName(Property).ToString());
-						const FName& NewPropertyIdName = *(TargetRCProperty->PropertyId.ToString() + TEXT(".") + PropertyClassName.ToString());
+						const FName PropertyClassName = *(TargetRCProperty->PropertyId.ToString() + TEXT(".") + PropertyIdHandler->GetPropertySubTypeName(Property).ToString());
+
+						// Set the property name while sanitizing it with valid object characters
+						const FName NewPropertyIdName = *SlugStringForValidName(TargetRCProperty->PropertyId.ToString() + TEXT("_") + PropertyClassName.ToString(), TEXT("_"));
 
 						FPropertyIdContainerKey CurrentKey = FPropertyIdContainerKey(TargetRCProperty->PropertyId, PropertyClassName);
 
