Remote Control: Actor label rename now update the owner column in exposed entities + Fixed issue with grouping not correctly being maintained on entity list refresh

diff --git a/Engine/Plugins/VirtualProduction/RemoteControl/Source/RemoteControlUI/Private/UI/SRCPanelExposedEntitiesList.cpp b/Engine/Plugins/VirtualProduction/RemoteControl/Source/RemoteControlUI/Private/UI/SRCPanelExposedEntitiesList.cpp
--- a/Engine/Plugins/VirtualProduction/RemoteControl/Source/RemoteControlUI/Private/UI/SRCPanelExposedEntitiesList.cpp
+++ b/Engine/Plugins/VirtualProduction/RemoteControl/Source/RemoteControlUI/Private/UI/SRCPanelExposedEntitiesList.cpp
@@ -15,6 +15,7 @@
 #include "Misc/MessageDialog.h"
 #include "PropertyPath.h"
 #include "RCPanelWidgetRegistry.h"
+#include "RCUIHelpers.h"
 #include "RemoteControlPanelStyle.h"
 #include "RemoteControlPreset.h"
 #include "RemoteControlSettings.h"
@@ -246,6 +247,8 @@ void SRCPanelExposedEntitiesList::Construct(const FArguments& InArgs, URemoteCon
 	OnEntityListUpdatedDelegate = InArgs._OnEntityListUpdated;
 	WidgetRegistry = MoveTemp(InWidgetRegistry);
 
+	OnActorLabelChangedHandle = FCoreDelegates::OnActorLabelChanged.AddSP(this, &SRCPanelExposedEntitiesList::OnActorLabelChanged);
+
 	WidgetRegistry.Pin()->OnObjectRefreshed().AddSP(this, &SRCPanelExposedEntitiesList::OnWidgetRegistryRefreshed);
 
 	ColumnSizeData.LeftColumnWidth = TAttribute<float>(this, &SRCPanelExposedEntitiesList::OnGetLeftColumnWidth);
@@ -570,6 +573,9 @@ SRCPanelExposedEntitiesList::~SRCPanelExposedEntitiesList()
 	UnregisterEvents();
 
 	StoreListSettings();
+
+	FCoreDelegates::OnActorLabelChanged.Remove(OnActorLabelChangedHandle);
+	OnActorLabelChangedHandle.Reset();
 }
 
 void SRCPanelExposedEntitiesList::Tick(const FGeometry& AllottedGeometry, const double InCurrentTime, const float InDeltaTime)
@@ -889,8 +895,8 @@ void SRCPanelExposedEntitiesList::StoreListSettings()
 	}
 	else
 	{
-		CurrentGroupType = DefaultSettings.FieldGroupType;
-		CurrentGroupSortType = DefaultSettings.FieldGroupOrder;
+		DefaultSettings.FieldGroupType = CurrentGroupType;
+		DefaultSettings.FieldGroupOrder = CurrentGroupSortType;
 	}
 }
 
@@ -981,6 +987,14 @@ void SRCPanelExposedEntitiesList::OnObjectPropertyChange(UObject* InObject, FPro
 	FieldsListView->RequestListRefresh();
 }
 
+void SRCPanelExposedEntitiesList::OnActorLabelChanged(AActor* InActor)
+{
+	if (CachedEntityOwners.Contains(InActor))
+	{
+		Refresh();
+	}
+}
+
 void SRCPanelExposedEntitiesList::Refresh()
 {
 	bRefreshRequested = true;
@@ -1042,11 +1056,20 @@ void SRCPanelExposedEntitiesList::GenerateListWidgets()
 
 	TArray<FGuid> OrderMap = Preset->Layout.GetDefaultGroupOrder();
 	FieldWidgetMap.Reset();
+	CachedEntityOwners.Reset();
 
-	for (TWeakPtr<FRemoteControlEntity> WeakEntity : Preset->GetExposedEntities())
+	TArray<TWeakPtr<FRemoteControlEntity>> ExposedEntities = Preset->GetExposedEntities();
+	CachedEntityOwners.Reserve(ExposedEntities.Num());
+
+	for (TWeakPtr<FRemoteControlEntity> WeakEntity : ExposedEntities)
 	{
 		if (const TSharedPtr<FRemoteControlEntity> Entity = WeakEntity.Pin())
 		{
+			if (const UObject* EntityOwner = UE::RCUIHelpers::GetEntityOwner(Entity))
+			{
+				CachedEntityOwners.Add(EntityOwner);
+			}
+
 			FGenerateWidgetArgs Args;
 			Args.Entity = Entity;
 			Args.Preset = Preset.Get();
diff --git a/Engine/Plugins/VirtualProduction/RemoteControl/Source/RemoteControlUI/Private/UI/SRCPanelExposedEntitiesList.h b/Engine/Plugins/VirtualProduction/RemoteControl/Source/RemoteControlUI/Private/UI/SRCPanelExposedEntitiesList.h
--- a/Engine/Plugins/VirtualProduction/RemoteControl/Source/RemoteControlUI/Private/UI/SRCPanelExposedEntitiesList.h
+++ b/Engine/Plugins/VirtualProduction/RemoteControl/Source/RemoteControlUI/Private/UI/SRCPanelExposedEntitiesList.h
@@ -2,6 +2,7 @@
 
 #pragma once
 
+#include "Containers/Set.h"
 #include "IRemoteControlUIModule.h"
 #include "Misc/TextFilter.h"
 #include "RemoteControlPreset.h"
@@ -11,11 +12,13 @@
 #include "UI/RCFieldGroupOrder.h"
 #include "UI/RCFieldGroupType.h"
 #include "UI/RCPanelExposedEntitiesListSettingsData.h"
+#include "UObject/ObjectKey.h"
 #include "UObject/StrongObjectPtr.h"
 #include "Widgets/DeclarativeSyntaxSupport.h"
 #include "Widgets/SCompoundWidget.h"
 #include "Widgets/Views/STreeView.h"
 
+class AActor;
 class FDragDropOperation;
 class FRCPanelWidgetRegistry;
 class FReply;
@@ -136,6 +139,8 @@ private:
 	FText HandleEntityListHeaderLabel() const;
 	/** Handles object property changes, used to update arrays correctly.  */
 	void OnObjectPropertyChange(UObject* InObject, FPropertyChangedEvent& InChangeEvent);
+	/** Called when the label of an actor has changed */
+	void OnActorLabelChanged(AActor* InActor);
 	/** Create exposed entity widgets. */
 	void GenerateListWidgets();
 	/** Create exposed entity widgets. */
@@ -279,6 +284,8 @@ private:
 	TArray<TSharedPtr<SRCPanelGroup>> FieldGroups;
 	/** Holds all the field entities. */
 	TArray<TSharedPtr<SRCPanelTreeNode>> FieldEntities;
+	/** Cached list of entity owners used to track whether owner changes should refresh the list*/
+	TSet<FObjectKey> CachedEntityOwners;
 	/** Cached field entities used to store the original Entities when switching groups. */
 	TArray<TSharedPtr<SRCPanelTreeNode>> CachedFieldEntities;
 	/** Holds all the exposed entities groups. */
@@ -301,6 +308,8 @@ private:
 	FDelegateHandle OnPropertyChangedHandle;
 	/** Handle to the delegate called when a binding is added or removed. */
 	FDelegateHandle OnProtocolBindingAddedOrRemovedHandle;
+	/** Handle to the delegate called when an actor's label is renamed */
+	FDelegateHandle OnActorLabelChangedHandle;
 	/** Delegate called on selected group change. */
 	FOnSelectionChange OnSelectionChangeDelegate;
 	/** The column data shared between all tree nodes in order to share a splitter amongst all rows. */
