Remote Control: fixed issues with bindings not working after level reload

diff --git a/Engine/Plugins/VirtualProduction/RemoteControl/Source/RemoteControlUI/Private/UI/SRemoteControlPanel.cpp b/Engine/Plugins/VirtualProduction/RemoteControl/Source/RemoteControlUI/Private/UI/SRemoteControlPanel.cpp
--- a/Engine/Plugins/VirtualProduction/RemoteControl/Source/RemoteControlUI/Private/UI/SRemoteControlPanel.cpp
+++ b/Engine/Plugins/VirtualProduction/RemoteControl/Source/RemoteControlUI/Private/UI/SRemoteControlPanel.cpp
@@ -1071,6 +1071,15 @@ void SRemoteControlPanel::OnMapChange(uint32)
 
 	// Clear the widget cache on map change to make sure we don't keep widgets around pointing to potentially stale objects.
 	WidgetRegistry->Clear();
+
+	// Update the world to the newer level editor map change
+	// Only do this if the current selected world is now invalid (stale due to map change).
+	// GWorld is used here instead of GEditor->EditorWorld as it's the one set to the newer world prior to the delegate being called (GEditor->EditorWorld set at a later time)
+	if (GWorld && !Preset->SelectedWorld.IsValid())
+	{
+		UpdatePanelForWorld(GWorld);
+	}
+
 	Refresh();
 }
 
