Remote Control: String controllers to support multi-line by default

diff --git a/Engine/Plugins/VirtualProduction/RemoteControl/Source/RemoteControlLogic/Private/Controller/RCControllerContainer.cpp b/Engine/Plugins/VirtualProduction/RemoteControl/Source/RemoteControlLogic/Private/Controller/RCControllerContainer.cpp
--- a/Engine/Plugins/VirtualProduction/RemoteControl/Source/RemoteControlLogic/Private/Controller/RCControllerContainer.cpp
+++ b/Engine/Plugins/VirtualProduction/RemoteControl/Source/RemoteControlLogic/Private/Controller/RCControllerContainer.cpp
@@ -26,7 +26,13 @@ void URCControllerContainer::UpdateEntityIds(const TMap<FGuid, FGuid>& InEntityI
 
 URCVirtualPropertyInContainer* URCControllerContainer::AddProperty(const FName& InPropertyName, TSubclassOf<URCVirtualPropertyInContainer> InPropertyClass, const EPropertyBagPropertyType InValueType, UObject* InValueTypeObject /*= nullptr*/, TArray<FPropertyBagPropertyDescMetaData> MetaData /*= TArray<FPropertyBagPropertyDescMetaData>()*/)
 {
-	if (InValueType == EPropertyBagPropertyType::Struct)
+	if (InValueType == EPropertyBagPropertyType::String)
+	{
+		// Add multiline capabilities to string controllers.
+		// TODO: a meta-data editing interface to customize meta-data like this per-controller.
+		MetaData.Add(FPropertyBagPropertyDescMetaData(FName("MultiLine"), TEXT("true")));
+	}
+	else if (InValueType == EPropertyBagPropertyType::Struct)
 	{
 		// Color Controllers
 		if (InValueTypeObject == TBaseStructure<FColor>::Get())
