Remote Control: fixed issue where properties from different actors with the same label would get put in the same owner group

diff --git a/Engine/Plugins/VirtualProduction/RemoteControl/Source/RemoteControlUI/Private/UI/SRCPanelExposedEntitiesGroup.cpp b/Engine/Plugins/VirtualProduction/RemoteControl/Source/RemoteControlUI/Private/UI/SRCPanelExposedEntitiesGroup.cpp
--- a/Engine/Plugins/VirtualProduction/RemoteControl/Source/RemoteControlUI/Private/UI/SRCPanelExposedEntitiesGroup.cpp
+++ b/Engine/Plugins/VirtualProduction/RemoteControl/Source/RemoteControlUI/Private/UI/SRCPanelExposedEntitiesGroup.cpp
@@ -62,7 +62,7 @@ SRCPanelTreeNode::FMakeNodeWidgetArgs SRCPanelExposedEntitiesGroup::CreateNodeWi
 		.VAlign(VAlign_Center)
 		[
 			SNew(STextBlock)
-			.Text_Lambda([this]{ return FText::FromName(OwnerName); })
+			.Text_Lambda([this]{ return OwnerDisplayName; })
 		];
 
 	FText GroupText;
@@ -128,36 +128,40 @@ void SRCPanelExposedEntitiesGroup::AssignChildren(const TArray<TSharedPtr<SRCPan
 {
 	ChildWidgets.Empty();
 
-	OwnerName = NAME_None;
+	OwnerDisplayName = FText::GetEmpty();
 	PropertyIdName = NAME_None;
 
-	auto CompareUpdateName = [](FName& InOutName, const FName& InOtherName)
-		{
-			if (InOutName == NAME_None)
-			{
-				InOutName = InOtherName;
-			}
-			else if (InOutName != InOtherName)
-			{
-				InOutName = TEXT("Multiple Values");
-			}
-		};
-
 	for (const TSharedPtr<SRCPanelTreeNode>& Entity : InFieldEntities)
 	{
 		if (const TSharedPtr<SRCPanelExposedField> ExposedField = StaticCastSharedPtr<SRCPanelExposedField>(Entity))
 		{
-			const FName FieldOwnerName  = ExposedField->GetOwnerName();
 			const FName FieldPropertyId = ExposedField->GetPropertyId();
 
-			const bool bOwnerMatch = GroupType == ERCFieldGroupType::Owner && FieldOwnerName == FieldKey;
+			const bool bOwnerMatch = GroupType == ERCFieldGroupType::Owner && ExposedField->GetOwnerPathName() == FieldKey;
 			const bool bPropertyIdMatch = GroupType == ERCFieldGroupType::PropertyId && FieldPropertyId == FieldKey;
 
 			if (bOwnerMatch || bPropertyIdMatch)
 			{
 				ChildWidgets.Add(Entity);
-				CompareUpdateName(OwnerName, FieldOwnerName);
-				CompareUpdateName(PropertyIdName, FieldPropertyId);
+
+				const FText FieldOwnerDisplayName = ExposedField->GetOwnerDisplayName();
+				if (OwnerDisplayName.IsEmpty())
+				{
+					OwnerDisplayName = FieldOwnerDisplayName;
+				}
+				else if (!OwnerDisplayName.EqualTo(FieldOwnerDisplayName))
+				{
+					OwnerDisplayName = LOCTEXT("MultipleValues", "Multiple Values");
+				}
+
+				if (PropertyIdName == NAME_None)
+				{
+					PropertyIdName = FieldPropertyId;
+				}
+				else if (PropertyIdName != FieldPropertyId)
+				{
+					PropertyIdName = TEXT("Multiple Values");
+				}
 			}
 		}
 	}
diff --git a/Engine/Plugins/VirtualProduction/RemoteControl/Source/RemoteControlUI/Private/UI/SRCPanelExposedEntitiesGroup.h b/Engine/Plugins/VirtualProduction/RemoteControl/Source/RemoteControlUI/Private/UI/SRCPanelExposedEntitiesGroup.h
--- a/Engine/Plugins/VirtualProduction/RemoteControl/Source/RemoteControlUI/Private/UI/SRCPanelExposedEntitiesGroup.h
+++ b/Engine/Plugins/VirtualProduction/RemoteControl/Source/RemoteControlUI/Private/UI/SRCPanelExposedEntitiesGroup.h
@@ -76,7 +76,7 @@ private:
 	FName FieldKey;
 
 	/** Shared owner of the entity group */
-	FName OwnerName;
+	FText OwnerDisplayName;
 
 	/** Shared property id of the entity group */
 	FName PropertyIdName;
diff --git a/Engine/Plugins/VirtualProduction/RemoteControl/Source/RemoteControlUI/Private/UI/SRCPanelExposedEntitiesList.cpp b/Engine/Plugins/VirtualProduction/RemoteControl/Source/RemoteControlUI/Private/UI/SRCPanelExposedEntitiesList.cpp
--- a/Engine/Plugins/VirtualProduction/RemoteControl/Source/RemoteControlUI/Private/UI/SRCPanelExposedEntitiesList.cpp
+++ b/Engine/Plugins/VirtualProduction/RemoteControl/Source/RemoteControlUI/Private/UI/SRCPanelExposedEntitiesList.cpp
@@ -1611,11 +1611,14 @@ void SRCPanelExposedEntitiesList::CreateFieldGroup()
 		{
 			if (TSharedPtr<SRCPanelExposedField> ExposedField = StaticCastSharedPtr<SRCPanelExposedField>(Entity))
 			{
-				if (!GroupKeys.Contains(ExposedField->GetOwnerName()))
+				const FName OwnerPathName = ExposedField->GetOwnerPathName();
+
+				if (!GroupKeys.Contains(OwnerPathName))
 				{
-					GroupKeys.Add(ExposedField->GetOwnerName());
+					GroupKeys.Add(OwnerPathName);
+
 					TSharedRef<SRCPanelExposedEntitiesGroup> GroupWidget = SNew(SRCPanelExposedEntitiesGroup, CurrentGroupType, Preset.Get())
-						.FieldKey(ExposedField->GetOwnerName());
+						.FieldKey(OwnerPathName);
 
 					GroupWidget->AssignChildren(CachedFieldEntities);
 					FieldEntitiesGroups.Add(GroupWidget);
diff --git a/Engine/Plugins/VirtualProduction/RemoteControl/Source/RemoteControlUI/Private/UI/SRCPanelExposedEntity.cpp b/Engine/Plugins/VirtualProduction/RemoteControl/Source/RemoteControlUI/Private/UI/SRCPanelExposedEntity.cpp
--- a/Engine/Plugins/VirtualProduction/RemoteControl/Source/RemoteControlUI/Private/UI/SRCPanelExposedEntity.cpp
+++ b/Engine/Plugins/VirtualProduction/RemoteControl/Source/RemoteControlUI/Private/UI/SRCPanelExposedEntity.cpp
@@ -184,13 +184,14 @@ void SRCPanelExposedEntity::Initialize(const FInitParams& InParams)
 				CachedFieldPath = RCProperty->FieldPathInfo.ToString();
 			}
 
-			FName OwnerFName;
+			CachedOwnerPathName = NAME_None;
+
 			// If the binding is valid, display the actor label if possible
 			if (UObject* Owner = UE::RCUIHelpers::GetEntityOwner(RCEntity))
 			{
 				bValidBinding = true;
-				OwnerFName = Owner->GetFName();
-				CachedOwnerName = *UKismetSystemLibrary::GetDisplayName(Owner);
+				CachedOwnerPathName = *Owner->GetPathName();
+				CachedOwnerDisplayName = FText::FromString(UKismetSystemLibrary::GetDisplayName(Owner)); 
 			}
 			else
 			{
@@ -201,20 +202,20 @@ void SRCPanelExposedEntity::Initialize(const FInitParams& InParams)
 				int32 PersistentLevelIndex = BindingPath.Find(PersistentLevelString);
 				if (PersistentLevelIndex != INDEX_NONE)
 				{
-					OwnerFName = *BindingPath.RightChop(PersistentLevelIndex + PersistentLevelString.Len());
-					CachedOwnerName = OwnerFName;
+					CachedOwnerPathName = *BindingPath.RightChop(PersistentLevelIndex + PersistentLevelString.Len());
+					CachedOwnerDisplayName = FText::FromName(CachedOwnerPathName);
 				}
 			}
-			
-			const int32 OwnerNameIndex = BindingPath.Find(OwnerFName.ToString() + TEXT("."));
+
+			const int32 OwnerNameIndex = BindingPath.Find(CachedOwnerPathName.ToString() + TEXT("."));
 			if (OwnerNameIndex != INDEX_NONE)
 			{
-				CachedSubobjectPath = *BindingPath.RightChop(OwnerNameIndex + OwnerFName.GetStringLength() + 1);
+				CachedSubobjectPath = *BindingPath.RightChop(OwnerNameIndex + CachedOwnerPathName.GetStringLength() + 1);
 			}
 
-			if (CachedOwnerName.IsNone())
+			if (CachedOwnerPathName.IsNone())
 			{
-				CachedOwnerName = *LOCTEXT("InvalidOwner", "Invalid Owner").ToString();
+				CachedOwnerPathName = *LOCTEXT("InvalidOwner", "Invalid Owner").ToString();
 			}
 		}
 	}
@@ -502,7 +503,7 @@ SRCPanelTreeNode::FMakeNodeWidgetArgs SRCPanelExposedEntity::CreateEntityWidgetI
 			[
 				SNew(STextBlock)
 				.ColorAndOpacity_Lambda([this]() { return bValidBinding ? FSlateColor::UseForeground() : FSlateColor::UseSubduedForeground(); })
-				.Text(FText::FromName(CachedOwnerName))
+				.Text(CachedOwnerDisplayName)
 				.ToolTipText(FText::FromString(CachedBindingPath.ToString() + SelectInOutliner.ToString()))
 			]
 		];
diff --git a/Engine/Plugins/VirtualProduction/RemoteControl/Source/RemoteControlUI/Private/UI/SRCPanelExposedEntity.h b/Engine/Plugins/VirtualProduction/RemoteControl/Source/RemoteControlUI/Private/UI/SRCPanelExposedEntity.h
--- a/Engine/Plugins/VirtualProduction/RemoteControl/Source/RemoteControlUI/Private/UI/SRCPanelExposedEntity.h
+++ b/Engine/Plugins/VirtualProduction/RemoteControl/Source/RemoteControlUI/Private/UI/SRCPanelExposedEntity.h
@@ -69,8 +69,10 @@ protected:
 	TAttribute<bool> bLiveMode;
 	/** Command list to use for the context menu */
 	TSharedPtr<FUICommandList> CommandList;
+	/** Display name of the entity's owner (What actor it's bound to). */
+	FText CachedOwnerDisplayName;
 	/** Name of the entity's owner (What actor it's bound to). */
-	FName CachedOwnerName;
+	FName CachedOwnerPathName;
 	/** Path from the entity owner to the object holding the exposed property. */
 	FName CachedSubobjectPath;
 	/** Path of the entity's binding. */
diff --git a/Engine/Plugins/VirtualProduction/RemoteControl/Source/RemoteControlUI/Private/UI/SRCPanelExposedField.cpp b/Engine/Plugins/VirtualProduction/RemoteControl/Source/RemoteControlUI/Private/UI/SRCPanelExposedField.cpp
--- a/Engine/Plugins/VirtualProduction/RemoteControl/Source/RemoteControlUI/Private/UI/SRCPanelExposedField.cpp
+++ b/Engine/Plugins/VirtualProduction/RemoteControl/Source/RemoteControlUI/Private/UI/SRCPanelExposedField.cpp
@@ -200,9 +200,14 @@ FName SRCPanelExposedField::GetFieldLabel() const
 	return CachedLabel;
 }
 
-FName SRCPanelExposedField::GetOwnerName() const
+FText SRCPanelExposedField::GetOwnerDisplayName() const
 {
-	return CachedOwnerName;
+	return CachedOwnerDisplayName;
+}
+
+FName SRCPanelExposedField::GetOwnerPathName() const
+{
+	return CachedOwnerPathName;
 }
 
 SRCPanelTreeNode::FMakeNodeWidgetArgs SRCPanelExposedField::CreateEntityWidgetInternal(TSharedPtr<SWidget> ValueWidget, TSharedPtr<SWidget> ResetWidget, const FText& OptionalWarningMessage, TSharedRef<SWidget> EditConditionWidget)
diff --git a/Engine/Plugins/VirtualProduction/RemoteControl/Source/RemoteControlUI/Private/UI/SRCPanelExposedField.h b/Engine/Plugins/VirtualProduction/RemoteControl/Source/RemoteControlUI/Private/UI/SRCPanelExposedField.h
--- a/Engine/Plugins/VirtualProduction/RemoteControl/Source/RemoteControlUI/Private/UI/SRCPanelExposedField.h
+++ b/Engine/Plugins/VirtualProduction/RemoteControl/Source/RemoteControlUI/Private/UI/SRCPanelExposedField.h
@@ -75,8 +75,11 @@ struct SRCPanelExposedField : public SRCPanelExposedEntity
 	/** Returns this widget's underlying objects. */
 	void GetBoundObjects(TSet<UObject*>& OutBoundObjects) const;
 
-	/** Get the owner name of this field */
-	FName GetOwnerName() const;
+	/** Get the owner display name of this field */
+	FText GetOwnerDisplayName() const;
+
+	/** Get the owner path name of this field */
+	FName GetOwnerPathName() const;
 
 protected:
 	/** Returns populated args to display this widget. */
