Remote Control: fix issue with dragging linear color exposed property into controller area not creating any controllers

diff --git a/Engine/Plugins/VirtualProduction/RemoteControl/Source/RemoteControlLogic/Private/Behaviour/Builtin/Bind/RCBehaviourBind.cpp b/Engine/Plugins/VirtualProduction/RemoteControl/Source/RemoteControlLogic/Private/Behaviour/Builtin/Bind/RCBehaviourBind.cpp
--- a/Engine/Plugins/VirtualProduction/RemoteControl/Source/RemoteControlLogic/Private/Behaviour/Builtin/Bind/RCBehaviourBind.cpp
+++ b/Engine/Plugins/VirtualProduction/RemoteControl/Source/RemoteControlLogic/Private/Behaviour/Builtin/Bind/RCBehaviourBind.cpp
@@ -246,6 +246,7 @@ bool URCBehaviourBind::CanHaveActionForField(URCController* Controller, TSharedR
 
 	// Indirect Binding (for Structs)
 	//
+	// todo: unify with supported controller structs in RC Config
 	const static TMap<UScriptStruct*, TArray<UScriptStruct*>> SupportedStructConversions =
 	{
 		/* Struct Type */                                 /* Supported Struct Types */
@@ -399,6 +400,12 @@ bool URCBehaviourBind::GetPropertyBagTypeFromFieldProperty(const FProperty* InPr
 	{
 		OutPropertyBagType = EPropertyBagPropertyType::String;
 	}
+	// When creating from linear color properties, create FColor controller instead
+	// todo: unify with supported controller structs in RC Config
+	else if (OutStructObject == TBaseStructure<FLinearColor>::Get())
+	{
+		OutStructObject = TBaseStructure<FColor>::Get();
+	}
 
 	return OutPropertyBagType != EPropertyBagPropertyType::None;
 }
diff --git a/Engine/Plugins/VirtualProduction/RemoteControl/Source/RemoteControlLogic/Private/Controller/RCControllerUtilities.cpp b/Engine/Plugins/VirtualProduction/RemoteControl/Source/RemoteControlLogic/Private/Controller/RCControllerUtilities.cpp
--- a/Engine/Plugins/VirtualProduction/RemoteControl/Source/RemoteControlLogic/Private/Controller/RCControllerUtilities.cpp
+++ b/Engine/Plugins/VirtualProduction/RemoteControl/Source/RemoteControlLogic/Private/Controller/RCControllerUtilities.cpp
@@ -21,7 +21,8 @@ namespace UE::RCControllers
 			return InStruct == TBaseStructure<FVector>::Get()
 				|| InStruct == TBaseStructure<FVector2D>::Get()
 				|| InStruct == TBaseStructure<FRotator>::Get()
-				|| InStruct == TBaseStructure<FColor>::Get();
+				|| InStruct == TBaseStructure<FColor>::Get()
+				|| InStruct == TBaseStructure<FLinearColor>::Get(); // todo: unify with supported controller structs in RC Config
 		}
 
 		/**
