Remote Control: Fixed function serialization for functions that become invalid (e.g. BP functions that get deleted).
The previous code would not serialize anything on saving if the function was invalid, but on load it would always serialize the number of bytes for the arguments. This mismatch caused issues with serial mismatch causing fatal errors.
This change does not fix existing damaged assets because on load it's impossible to know whether a function that is now invalid/nonexistent on load was valid or invalid at the time of save,  which for the old saving code meant different serialization results.
This change, however, fixes new assets such that the serial size remain consistent regardless of whether a function becomes invalid at any time.

diff --git a/Engine/Plugins/VirtualProduction/RemoteControl/Source/RemoteControl/Private/RemoteControlField.cpp b/Engine/Plugins/VirtualProduction/RemoteControl/Source/RemoteControl/Private/RemoteControlField.cpp
--- a/Engine/Plugins/VirtualProduction/RemoteControl/Source/RemoteControl/Private/RemoteControlField.cpp
+++ b/Engine/Plugins/VirtualProduction/RemoteControl/Source/RemoteControl/Private/RemoteControlField.cpp
@@ -743,16 +743,14 @@ FArchive& operator<<(FArchive& Ar, FRemoteControlFunction& RCFunction)
 		else
 		{
 			UE_LOG(LogRemoteControl, Warning, TEXT("%s could not be loaded while deserialzing a Remote Control Function."), *RCFunction.FunctionPath.ToString());
-			Ar.SetError();
 			if (NumSerializedBytesFromArchive > 0)
 			{
 				// Skip over chunk of data if we were unable to resolve the class.
 				Ar.Seek(Ar.Tell() + NumSerializedBytesFromArchive);
 			}
 		}
-		
 	}
-	else if (RCFunction.CachedFunction.IsValid())
+	else
 	{
 		// The following code serializes the size of the arguments to serialize so that when loading,
 		// we can skip over it if the function could not be loaded.
@@ -764,7 +762,10 @@ FArchive& operator<<(FArchive& Ar, FRemoteControlFunction& RCFunction)
 
 		// Then serialize the arguments in order to get its size.
 		const int64 ArgsBegin = Ar.Tell();
-		RCFunction.CachedFunction->SerializeTaggedProperties(Ar, RCFunction.FunctionArguments->GetStructMemory(), RCFunction.CachedFunction.Get(), nullptr);
+		if (UFunction* Function = RCFunction.CachedFunction.Get())
+		{
+			Function->SerializeTaggedProperties(Ar, RCFunction.FunctionArguments->GetStructMemory(), Function, nullptr);
+		}
 
 		// Only go back and serialize the number of argument bytes if there is actually an underlying buffer to seek to.
 		if (ArgumentsSizePos != INDEX_NONE)
