Remote Control: fixed issue with level reload causing action add button to stop working.
The current logic attempts to use a cached menu widget. Because menu generation is really only needed when clicking on the add button, this caching is not bringing any performance benefit and instead causing issues where the 'fix' for these issues would be to request refresh on this widget cache.
Current logic listens to events like property id changed, field added/removed, but not events like level reload, and potentially other events that would cause the menu to give a wrong result (or none at all).
The best course of action is to remove this widget cache entirely as it's not necessary, so that everytime the user presses that add action button, they can see the latest set of available actions in the menu.

diff --git a/Engine/Plugins/VirtualProduction/RemoteControl/Source/RemoteControlUI/Private/UI/Action/SRCActionPanel.cpp b/Engine/Plugins/VirtualProduction/RemoteControl/Source/RemoteControlUI/Private/UI/Action/SRCActionPanel.cpp
--- a/Engine/Plugins/VirtualProduction/RemoteControl/Source/RemoteControlUI/Private/UI/Action/SRCActionPanel.cpp
+++ b/Engine/Plugins/VirtualProduction/RemoteControl/Source/RemoteControlUI/Private/UI/Action/SRCActionPanel.cpp
@@ -74,17 +74,6 @@ void SRCActionPanel::Construct(const FArguments& InArgs, const TSharedRef<SRemot
 
 	// Register delegates
 	InPanel->OnBehaviourSelectionChanged.AddSP(this, &SRCActionPanel::OnBehaviourSelectionChanged);
-
-	if (URemoteControlPreset* Preset = GetPreset())
-	{
-		Preset->Layout.OnFieldAdded().AddSP(this, &SRCActionPanel::OnRemoteControlFieldAdded);
-		Preset->Layout.OnFieldDeleted().AddSP(this, &SRCActionPanel::OnRemoteControlFieldDeleted);
-		if (const TObjectPtr<URemoteControlPropertyIdRegistry> Registry = Preset->GetPropertyIdRegistry())
-		{
-			Registry->OnPropertyIdUpdated().AddSPLambda(this, [this](){ bAddActionMenuNeedsRefresh = true; });
-			Registry->OnPropertyIdActionNeedsRefresh().AddSPLambda(this, [this](){ bAddActionMenuNeedsRefresh = true; });
-		}
-	}
 }
 
 END_SLATE_FUNCTION_BUILD_OPTIMIZATION
@@ -150,15 +139,6 @@ void SRCActionPanel::RefreshIsBehaviourEnabled(const bool bIsEnabled)
 
 TSharedRef<SWidget> SRCActionPanel::GetActionMenuContentWidget()
 {
-	if (AddNewActionMenuWidget && !bAddActionMenuNeedsRefresh)
-	{
-		return AddNewActionMenuWidget.ToSharedRef();
-	}
-
-	// Either we are creating the menu the first time, or the list needs to be refreshed
-	// The latter can occur either when a remote control property is added or removed, or when a behaiour's CanHaveActionForField logic has changed
-	// For example, if the user sets the "Allow numeric input as Strings" flag for Bind Behaviour then we need to recalculate the list of eligible actions.
-
 	constexpr bool bShouldCloseWindowAfterMenuSelection = true;
 	FMenuBuilder MenuBuilder(bShouldCloseWindowAfterMenuSelection, nullptr);
 
@@ -241,11 +221,7 @@ TSharedRef<SWidget> SRCActionPanel::GetActionMenuContentWidget()
 		}
 	}
 
-	bAddActionMenuNeedsRefresh = false; // reset
-
-	AddNewActionMenuWidget = MenuBuilder.MakeWidget();
-
-	return AddNewActionMenuWidget.ToSharedRef();
+	return MenuBuilder.MakeWidget();
 }
 
 void SRCActionPanel::CreateActionPanelWidgets()
@@ -634,16 +610,6 @@ FReply SRCActionPanel::OnAddAllSelectedFields()
 	return FReply::Handled();
 }
 
-void SRCActionPanel::OnRemoteControlFieldAdded(const FGuid& GroupId, const FGuid& FieldId, int32 FieldPosition)
-{
-	bAddActionMenuNeedsRefresh = true;
-}
-
-void SRCActionPanel::OnRemoteControlFieldDeleted(const FGuid& GroupId, const FGuid& FieldId, int32 FieldPosition)
-{
-	bAddActionMenuNeedsRefresh = true;
-}
-
 bool SRCActionPanel::IsListFocused() const
 {
 	return ActionPanelList.IsValid() && ActionPanelList->IsListFocused();
diff --git a/Engine/Plugins/VirtualProduction/RemoteControl/Source/RemoteControlUI/Private/UI/Action/SRCActionPanel.h b/Engine/Plugins/VirtualProduction/RemoteControl/Source/RemoteControlUI/Private/UI/Action/SRCActionPanel.h
--- a/Engine/Plugins/VirtualProduction/RemoteControl/Source/RemoteControlUI/Private/UI/Action/SRCActionPanel.h
+++ b/Engine/Plugins/VirtualProduction/RemoteControl/Source/RemoteControlUI/Private/UI/Action/SRCActionPanel.h
@@ -82,11 +82,6 @@ public:
 	/** Whether this Actions panel can create an action for the given remote control field */
 	bool CanHaveActionForField(const FGuid& InRemoteControlFieldId);
 
-	void RequestRefreshForAddActionsMenu()
-	{
-		bAddActionMenuNeedsRefresh = true;
-	}
-
 	/** Set the Enabled state of our parent Behaviour */
 	void SetIsBehaviourEnabled(const bool bIsEnabled);
 
@@ -152,12 +147,6 @@ private:
 	/** Handles click event for Add All Selected button; Adds all selected actions for the active Remote Control Preset*/
 	FReply OnAddAllSelectedFields();
 
-	/** Event invoked when a new remote control field has been added to the Remote Control Preset associated with this Action panel */
-	void OnRemoteControlFieldAdded(const FGuid& GroupId, const FGuid& FieldId, int32 FieldPosition);
-
-	/** Event invoked when a remote control field is removed from the Remote Control Preset associated with this Action panel */
-	void OnRemoteControlFieldDeleted(const FGuid& GroupId, const FGuid& FieldId, int32 FieldPosition);
-
 private:
 	/** Helper widget for behavior details. */
 	static TSharedRef<SBox> GetNoneSelectedWidget();
@@ -181,9 +170,6 @@ private:
 	/** Panel Style reference. */
 	const FRCPanelStyle* RCPanelStyle = nullptr;
 
-	/** Cached menu widget for Add New Action */
-	TSharedPtr<SWidget> AddNewActionMenuWidget;
-
 	/** Action Minor Panel */
 	TSharedPtr<SRCMinorPanel> ActionDockPanel;
 
@@ -201,7 +187,4 @@ private:
 
 	/** Behavior Details Minor Panel */
 	TSharedPtr<SRCMinorPanel> BehaviorDetailsPanel;
-
-	/** Whether the Add Actions menu list is outdated and needs to be refreshed */
-	bool bAddActionMenuNeedsRefresh = false;
 };
diff --git a/Engine/Plugins/VirtualProduction/RemoteControl/Source/RemoteControlUI/Private/UI/Behaviour/Builtin/Bind/SRCBehaviourBind.cpp b/Engine/Plugins/VirtualProduction/RemoteControl/Source/RemoteControlUI/Private/UI/Behaviour/Builtin/Bind/SRCBehaviourBind.cpp
--- a/Engine/Plugins/VirtualProduction/RemoteControl/Source/RemoteControlUI/Private/UI/Behaviour/Builtin/Bind/SRCBehaviourBind.cpp
+++ b/Engine/Plugins/VirtualProduction/RemoteControl/Source/RemoteControlUI/Private/UI/Behaviour/Builtin/Bind/SRCBehaviourBind.cpp
@@ -35,14 +35,14 @@ void SRCBehaviourBind::Construct(const FArguments& InArgs, TSharedRef<FRCBehavio
 		.Padding(FMargin(10.f))
 		[
 			SAssignNew(CheckboxAllowNumericInput, SCheckBox)
-			.IsChecked(false)
+			.IsChecked(ECheckBoxState::Unchecked)
 			.OnCheckStateChanged(this, &SRCBehaviourBind::OnAllowNumericCheckboxChanged)
 		]
 	];
 
 	if (const URCBehaviourBind* BindBehaviour = InBehaviourItem->GetBindBehaviour())
 	{
-		CheckboxAllowNumericInput->SetIsChecked(BindBehaviour->AreNumericInputsAllowedAsStrings());
+		CheckboxAllowNumericInput->SetIsChecked(BindBehaviour->AreNumericInputsAllowedAsStrings() ? ECheckBoxState::Checked : ECheckBoxState::Unchecked);
 	}
 }
 
@@ -54,15 +54,6 @@ void SRCBehaviourBind::OnAllowNumericCheckboxChanged(ECheckBoxState NewState)
 		{
 			// Update flag
 			BindBehaviour->SetAllowNumericInputAsStrings((NewState == ECheckBoxState::Checked));
-
-			// Inform the Action panel that it needs to refresh the Add Action menu list
-			if (TSharedPtr<SRemoteControlPanel> RemoteControlPanel = BehaviourItem->GetRemoteControlPanel())
-			{
-				if (TSharedPtr<SRCActionPanel> ActionPanel = RemoteControlPanel->GetLogicActionPanel())
-				{
-					ActionPanel->RequestRefreshForAddActionsMenu();
-				}
-			}
 		}
 	}
 	
