Remote Control: Added support for enums and additional struct types

diff --git a/Engine/Plugins/VirtualProduction/RemoteControl/Config/BaseRemoteControl.ini b/Engine/Plugins/VirtualProduction/RemoteControl/Config/BaseRemoteControl.ini
--- a/Engine/Plugins/VirtualProduction/RemoteControl/Config/BaseRemoteControl.ini
+++ b/Engine/Plugins/VirtualProduction/RemoteControl/Config/BaseRemoteControl.ini
@@ -11,6 +11,7 @@
 +SupportedControllerTypes=EPropertyBagPropertyType::Bool
 +SupportedControllerTypes=EPropertyBagPropertyType::Struct
 +SupportedControllerTypes=EPropertyBagPropertyType::Object
++SupportedControllerTypes=EPropertyBagPropertyType::Enum
 +SupportedControllerStructTypes=Vector
 +SupportedControllerStructTypes=Vector2D
 +SupportedControllerStructTypes=Color
diff --git a/Engine/Plugins/VirtualProduction/RemoteControl/Source/RemoteControl/Private/RCVirtualProperty.cpp b/Engine/Plugins/VirtualProduction/RemoteControl/Source/RemoteControl/Private/RCVirtualProperty.cpp
--- a/Engine/Plugins/VirtualProduction/RemoteControl/Source/RemoteControl/Private/RCVirtualProperty.cpp
+++ b/Engine/Plugins/VirtualProduction/RemoteControl/Source/RemoteControl/Private/RCVirtualProperty.cpp
@@ -10,6 +10,14 @@
 #include "UObject/TextProperty.h"
 #include "UObject/StructOnScope.h"
 
+template<typename T>
+concept CHasMethodGetPropertyValueInContainer = requires(T && InObject, uint8 * InPtr)
+{
+	{
+		InObject.GetPropertyValue_InContainer(InPtr)
+	};
+};
+
 const FProperty* URCVirtualPropertyBase::GetProperty() const
 {
 	if (const FPropertyBagPropertyDesc* PropertyBagPropertyDesc = GetBagPropertyDesc())
@@ -348,7 +356,19 @@ struct FRCVirtualPropertyCastHelpers
 		
 		if (const TPropertyType* CastProperty = CastField<TPropertyType>(Property))
 		{
-			OutValue = CastProperty->GetPropertyValue_InContainer(ContainerPtr);
+			if constexpr (CHasMethodGetPropertyValueInContainer<TPropertyType>)
+			{
+				OutValue = CastProperty->GetPropertyValue_InContainer(ContainerPtr);
+			}
+			else
+			{
+				FMemory::Memcpy(
+					&OutValue,
+					CastProperty->template ContainerPtrToValuePtr<void>(ContainerPtr),
+					sizeof(OutValue)
+				);
+			}
+
 			return true;
 		}
 
@@ -462,6 +482,11 @@ bool URCVirtualPropertyBase::GetValueByte(uint8& OutByte) const
 	return FRCVirtualPropertyCastHelpers::GetPrimitiveValue<FByteProperty, uint8&>(this, OutByte);
 }
 
+bool URCVirtualPropertyBase::GetValueEnum(uint8& OutByte) const
+{
+	return FRCVirtualPropertyCastHelpers::GetPrimitiveValue<FEnumProperty, uint8&>(this, OutByte);
+}
+
 bool URCVirtualPropertyBase::GetValueInt16(int16& OutInt16) const
 {
 	return FRCVirtualPropertyCastHelpers::GetPrimitiveValue<FInt16Property, int16&>(this, OutInt16);
@@ -637,6 +662,16 @@ FString URCVirtualPropertyBase::GetDisplayValueAsString() const
 			return Object->GetPathName();
 		}
 	}
+	else if (const FEnumProperty* EnumProperty = CastField<FEnumProperty>(Property))
+	{
+		if (UEnum* Enum = EnumProperty->GetEnum())
+		{
+			uint8 Value;
+			GetValueEnum(Value);
+
+			return Enum->GetDisplayNameTextByValue(static_cast<int64>(Value)).ToString();
+		}
+	}
 	else if (IsVectorType())
 	{
 		FVector Vector;
@@ -665,6 +700,19 @@ FString URCVirtualPropertyBase::GetDisplayValueAsString() const
 
 		return Color.ToString();
 	}
+	else if (const FStructProperty* StructProperty = CastField<FStructProperty>(Property))
+	{
+		FString Value;
+		StructProperty->ExportTextItem_Direct(
+			Value,
+			GetValuePtr(),
+			/* Default value ptr */ nullptr,
+			/* Parent */ nullptr,
+			/* Port flags */ 0
+		);
+
+		return Value;
+	}
 
 	// Unimplemented type!
 	ensureAlwaysMsgf(false, TEXT("Unimplemented type %s passed to GetDisplayValueAsString"), *Property->GetName());
diff --git a/Engine/Plugins/VirtualProduction/RemoteControl/Source/RemoteControl/Public/RCVirtualProperty.h b/Engine/Plugins/VirtualProduction/RemoteControl/Source/RemoteControl/Public/RCVirtualProperty.h
--- a/Engine/Plugins/VirtualProduction/RemoteControl/Source/RemoteControl/Public/RCVirtualProperty.h
+++ b/Engine/Plugins/VirtualProduction/RemoteControl/Source/RemoteControl/Public/RCVirtualProperty.h
@@ -177,6 +177,10 @@ public:
 	UFUNCTION(BlueprintCallable, Category = "Remote Control Behaviour")
 	bool GetValueByte(uint8& OutByte) const;
 
+	/** Get raw Enum value from Virtual Property */
+	UFUNCTION(BlueprintCallable, Category = "Remote Control Behaviour")
+	bool GetValueEnum(uint8& OutByte) const;
+
 	/** Get Int16 value from Virtual Property */
 	UFUNCTION()
 	bool GetValueInt16(int16& OutInt16) const;
diff --git a/Engine/Plugins/VirtualProduction/RemoteControl/Source/RemoteControlLogic/Private/Behaviour/Builtin/RCBehaviourOnValueChangedNode.cpp b/Engine/Plugins/VirtualProduction/RemoteControl/Source/RemoteControlLogic/Private/Behaviour/Builtin/RCBehaviourOnValueChangedNode.cpp
--- a/Engine/Plugins/VirtualProduction/RemoteControl/Source/RemoteControlLogic/Private/Behaviour/Builtin/RCBehaviourOnValueChangedNode.cpp
+++ b/Engine/Plugins/VirtualProduction/RemoteControl/Source/RemoteControlLogic/Private/Behaviour/Builtin/RCBehaviourOnValueChangedNode.cpp
@@ -2,6 +2,8 @@
 
 #include "Behaviour/Builtin/RCBehaviourOnValueChangedNode.h"
 
+#include "Behaviour/RCBehaviour.h"
+#include "Controller/RCController.h"
 #include "StructUtils/PropertyBag.h"
 
 URCBehaviourOnValueChangedNode::URCBehaviourOnValueChangedNode()
@@ -18,21 +20,24 @@ bool URCBehaviourOnValueChangedNode::Execute_Implementation(URCBehaviour* InBeha
 
 bool URCBehaviourOnValueChangedNode::IsSupported_Implementation(URCBehaviour* InBehaviour) const
 {
-	static TArray<TPair<EPropertyBagPropertyType, UObject*>> SupportedPropertyBagTypes =
+	static TArray<EPropertyBagPropertyType> SupportedPropertyBagTypes =
 	{
-		{ EPropertyBagPropertyType::Bool, nullptr },
-		{ EPropertyBagPropertyType::Byte, nullptr },
-		{ EPropertyBagPropertyType::Int32, nullptr },
-		{ EPropertyBagPropertyType::Float, nullptr },
-		{ EPropertyBagPropertyType::Double, nullptr },
-		{ EPropertyBagPropertyType::Name, nullptr },
-		{ EPropertyBagPropertyType::String, nullptr },
-		{ EPropertyBagPropertyType::Text, nullptr },
-		{ EPropertyBagPropertyType::Struct, TBaseStructure<FVector>::Get() },
-		{ EPropertyBagPropertyType::Struct, TBaseStructure<FVector2D>::Get() },
-		{ EPropertyBagPropertyType::Struct, TBaseStructure<FColor>::Get() },
-		{ EPropertyBagPropertyType::Struct, TBaseStructure<FRotator>::Get() }
+		EPropertyBagPropertyType::Bool,
+		EPropertyBagPropertyType::Byte,
+		EPropertyBagPropertyType::Int32,
+		EPropertyBagPropertyType::Float,
+		EPropertyBagPropertyType::Double,
+		EPropertyBagPropertyType::Name,
+		EPropertyBagPropertyType::String,
+		EPropertyBagPropertyType::Text,
+		EPropertyBagPropertyType::Struct,
+		EPropertyBagPropertyType::Enum
 	};
 	
-	return SupportedPropertyBagTypes.ContainsByPredicate(GetIsSupportedCallback(InBehaviour));
+	if (const URCController* RCController = InBehaviour->ControllerWeakPtr.Get())
+	{
+		return SupportedPropertyBagTypes.Contains(RCController->GetValueType());
+	}
+	
+	return false;
 }
diff --git a/Engine/Plugins/VirtualProduction/RemoteControl/Source/RemoteControlLogic/Private/Behaviour/RCBehaviourNode.cpp b/Engine/Plugins/VirtualProduction/RemoteControl/Source/RemoteControlLogic/Private/Behaviour/RCBehaviourNode.cpp
--- a/Engine/Plugins/VirtualProduction/RemoteControl/Source/RemoteControlLogic/Private/Behaviour/RCBehaviourNode.cpp
+++ b/Engine/Plugins/VirtualProduction/RemoteControl/Source/RemoteControlLogic/Private/Behaviour/RCBehaviourNode.cpp
@@ -34,7 +34,14 @@ TFunction<bool(const TPair<EPropertyBagPropertyType, UObject*>)> URCBehaviourNod
 	{
 		if (const URCController* RCController = InBehaviour->ControllerWeakPtr.Get())
 		{
-			return RCController->GetValueType() == InPair.Key && RCController->GetValueTypeObjectWeakPtr() == InPair.Value;
+			EPropertyBagPropertyType ValueType = RCController->GetValueType();
+
+			if (ValueType == EPropertyBagPropertyType::Enum && InPair.Key == EPropertyBagPropertyType::Enum)
+			{
+				return true;
+			}
+
+			return ValueType == InPair.Key && RCController->GetValueTypeObjectWeakPtr() == InPair.Value;
 		}
 		
 		return false;
diff --git a/Engine/Plugins/VirtualProduction/RemoteControl/Source/RemoteControlUI/Private/UI/Controller/SRCControllerPanel.cpp b/Engine/Plugins/VirtualProduction/RemoteControl/Source/RemoteControlUI/Private/UI/Controller/SRCControllerPanel.cpp
--- a/Engine/Plugins/VirtualProduction/RemoteControl/Source/RemoteControlUI/Private/UI/Controller/SRCControllerPanel.cpp
+++ b/Engine/Plugins/VirtualProduction/RemoteControl/Source/RemoteControlUI/Private/UI/Controller/SRCControllerPanel.cpp
@@ -6,20 +6,28 @@
 #include "Controller/RCCustomControllerUtilities.h"
 #include "Engine/Texture.h"
 #include "Framework/MultiBox/MultiBoxBuilder.h"
+#include "IDetailTreeNode.h"
+#include "IPropertyRowGenerator.h"
 #include "Materials/MaterialInterface.h"
 #include "Misc/MessageDialog.h"
+#include "Modules/ModuleManager.h"
+#include "PropertyEditorModule.h"
 #include "RCControllerModel.h"
 #include "RCVirtualPropertyContainer.h"
 #include "RemoteControlLogicConfig.h"
 #include "RemoteControlPreset.h"
-#include "SRCControllerPanelList.h"
 #include "SlateOptMacros.h"
+#include "SRCControllerPanelList.h"
 #include "Styling/RemoteControlStyles.h"
+#include "TypeSelector/SRCControllerTypeSelector.h"
 #include "UI/BaseLogicUI/RCLogicModeBase.h"
 #include "UI/Panels/SRCDockPanel.h"
 #include "UI/RCUIHelpers.h"
 #include "UI/RemoteControlPanelStyle.h"
 #include "UI/SRemoteControlPanel.h"
+#include "UObject/UnrealType.h"
+#include "UObject/UnrealTypePrivate.h"
+#include "UObject/UObjectIterator.h"
 #include "Widgets/Input/SComboButton.h"
 #include "Widgets/SBoxPanel.h"
 #include "Widgets/Views/SListView.h"
@@ -70,10 +78,7 @@ void SRCControllerPanel::Construct(const FArguments& InArgs, const TSharedRef<SR
 				.Image(FAppStyle::GetBrush("Icons.PlusCircle"))
 			]
 		]
-		.MenuContent()
-		[
-			GetControllerMenuContentWidget()
-		];
+		.OnGetMenuContent(this, &SRCControllerPanel::GetControllerMenuContentWidget);
 
 	ControllerDockPanel->AddHeaderToolbarItem(EToolbar::Left, AddNewControllerButton);
 	ControllerDockPanel->AddHeaderToolbarItem(EToolbar::Right, GetMultiControllerSwitchWidget());
@@ -95,8 +100,10 @@ void SRCControllerPanel::Refresh()
 
 END_SLATE_FUNCTION_BUILD_OPTIMIZATION
 
-SRCControllerPanel::FRCControllerPropertyInfo SRCControllerPanel::GetPropertyInfoForCustomType(const FName& InType)
+UE::RemoteControl::UI::Private::FRCControllerPropertyInfo SRCControllerPanel::GetPropertyInfoForCustomType(const FName& InType)
 {
+	using namespace UE::RemoteControl::UI::Private;
+
 	FRCControllerPropertyInfo PropertyInfo;
 	const EPropertyBagPropertyType CustomControllerType = UE::RCCustomControllers::GetCustomControllerType(InType);
 
@@ -109,7 +116,6 @@ SRCControllerPanel::FRCControllerPropertyInfo SRCControllerPanel::GetPropertyInf
 	return PropertyInfo;
 }
 
-
 static UObject* GetBaseStructForType(FName StructType)
 {
 	if (StructType == NAME_Vector)
@@ -344,6 +350,8 @@ EVisibility SRCControllerPanel::GetMultiControllerSwitchVisibility() const
 
 TSharedRef<SWidget> SRCControllerPanel::GetControllerMenuContentWidget() const
 {
+	using namespace UE::RemoteControl::UI::Private;
+
 	constexpr  bool bShouldCloseWindowAfterMenuSelection = true;
 	FMenuBuilder MenuBuilder(bShouldCloseWindowAfterMenuSelection, nullptr);
 
@@ -379,6 +387,11 @@ TSharedRef<SWidget> SRCControllerPanel::GetControllerMenuContentWidget() const
 				}	
 			}
 		}
+		else if (ControllerType == EPropertyBagPropertyType::Enum)
+		{
+			// Handled later
+			continue;
+		}
 		else
 		{
 			VirtualPropertyFieldClassNames.Add({ ControllerType, nullptr });
@@ -443,15 +456,189 @@ TSharedRef<SWidget> SRCControllerPanel::GetControllerMenuContentWidget() const
 				.Text(FText::FromName(DefaultName))
 			];
 
-		// Menu Item
 		FUIAction Action(FExecuteAction::CreateSP(this, &SRCControllerPanel::OnAddControllerClicked, Type, PropertyInfo));
-
 		MenuBuilder.AddMenuEntry(Action, MenuItemWidget);
 	}
 
+	// All structs
+	const FName DefaultStructName = TEXT("Structs");
+	FInstancedPropertyBag StructBag;
+	StructBag.AddProperty(DefaultStructName, EPropertyBagPropertyType::Struct, TBaseStructure<FVector2D>().Get());
+
+	if (const FPropertyBagPropertyDesc* StructBagPropertyDesc = StructBag.FindPropertyDescByName(DefaultStructName))
+	{
+		TSharedRef<SWidget> MenuItemWidget =
+			SNew(SHorizontalBox)
+			+ SHorizontalBox::Slot()
+			.AutoWidth()
+			.HAlign(HAlign_Left)
+			.VAlign(VAlign_Fill)
+			[
+				UE::RCUIHelpers::GetTypeColorWidget(StructBagPropertyDesc->CachedProperty)
+			]
+			+ SHorizontalBox::Slot()
+			.AutoWidth()
+			.Padding(FMargin(6.f, 0.f))
+			[
+				SNew(STextBlock)
+				.Text(FText::FromName(DefaultStructName))
+			];
+
+		FNewMenuDelegate Delegate = FNewMenuDelegate::CreateSP(this, &SRCControllerPanel::GetControllerMenuStructContentWidget);
+		MenuBuilder.AddSubMenu(MenuItemWidget, Delegate);
+	}
+
+	// All enums
+	const FName DefaultEnumName = TEXT("Enums");
+	FInstancedPropertyBag EnumBag;
+	EnumBag.AddProperty(DefaultEnumName, EPropertyBagPropertyType::Enum, StaticEnum<EStyleColor>());
+
+	if (const FPropertyBagPropertyDesc* EnumBagPropertyDesc = EnumBag.FindPropertyDescByName(DefaultEnumName))
+	{
+		TSharedRef<SWidget> MenuItemWidget =
+			SNew(SHorizontalBox)
+			+ SHorizontalBox::Slot()
+			.AutoWidth()
+			.HAlign(HAlign_Left)
+			.VAlign(VAlign_Fill)
+			[
+				UE::RCUIHelpers::GetTypeColorWidget(EnumBagPropertyDesc->CachedProperty)
+			]
+			+ SHorizontalBox::Slot()
+			.AutoWidth()
+			.Padding(FMargin(6.f, 0.f))
+			[
+				SNew(STextBlock)
+				.Text(FText::FromName(DefaultEnumName))
+			];
+
+		FNewMenuDelegate Delegate = FNewMenuDelegate::CreateSP(this, &SRCControllerPanel::GetControllerMenuEnumContentWidget);
+		MenuBuilder.AddSubMenu(MenuItemWidget, Delegate);
+	}
+
 	return MenuBuilder.MakeWidget();
 }
 
+void SRCControllerPanel::GetControllerMenuStructContentWidget(FMenuBuilder& InMenuBuilder) const
+{
+	// 99% of this is producing the menu items, not this list.
+	TArray<UScriptStruct*> Structs;
+	Structs.Reserve(1000); // Just start out with a big number
+
+	const URemoteControlLogicConfig* RCLogicConfig = GetDefault<URemoteControlLogicConfig>();
+
+	FPropertyEditorModule& PropertyEditorModule = FModuleManager::GetModuleChecked<FPropertyEditorModule>("PropertyEditor");
+
+	for (UScriptStruct* Struct : TObjectRange<UScriptStruct>())
+	{
+		if (Struct->HasAnyFlags(RF_ClassDefaultObject | RF_MirroredGarbage | RF_NewerVersionExists))
+		{
+			continue;
+		}
+
+		// Don't add the ones explicitly added to the base list.
+		if (RCLogicConfig->SupportedControllerStructTypes.Contains(Struct->GetFName()))
+		{
+			continue;
+		}
+
+		bool bHasEditableField = false;
+
+		for (FProperty* Property : TFieldRange<FProperty>(Struct, EFieldIterationFlags::IncludeSuper))
+		{
+			if (Property->HasAllPropertyFlags(CPF_Edit | CPF_BlueprintVisible))
+			{
+				bHasEditableField = true;
+				break;
+			}
+		}
+
+		if (!bHasEditableField)
+		{
+			continue;
+		}
+
+		// Uncustomised structs do not show correctly.
+		if (!PropertyEditorModule.IsCustomizedStruct(Struct, FCustomPropertyTypeLayoutMap()))
+		{
+			continue;
+		}
+
+		Structs.Add(Struct);
+	}
+
+	Structs.Sort(
+		[](const UScriptStruct& A, const UScriptStruct& B)
+		{
+			return A.GetName() < B.GetName();
+		}
+	);
+
+	using namespace UE::RemoteControl::UI::Private;
+
+	TArray<TSharedPtr<FRCControllerPropertyInfo>> Types;
+	Types.Reserve(Structs.Num());
+
+	for (UScriptStruct* Struct : Structs)
+	{
+		TSharedPtr<FRCControllerPropertyInfo>& Type = Types.Add_GetRef(MakeShared<FRCControllerPropertyInfo>(Struct));
+		Type->Type = EPropertyBagPropertyType::Struct;
+	}
+
+	InMenuBuilder.AddWidget(
+		SNew(SRCControllerTypeSelector, Types)
+			.OnTypeSelected(FRCControllerTypeBase::FOnTypeSelected::CreateSP(this, &SRCControllerPanel::OnAddControllerSelected)),
+		FText::GetEmpty(),
+		/* No Indent */ false
+	);
+}
+
+void SRCControllerPanel::GetControllerMenuEnumContentWidget(FMenuBuilder& InMenuBuilder) const
+{
+	// 99% of this is producing the menu items, not this list.
+	TArray<UEnum*> Enums;
+	Enums.Reserve(5000); // Just start out with a big number
+
+	const URemoteControlLogicConfig* RCLogicConfig = GetDefault<URemoteControlLogicConfig>();
+
+	FPropertyEditorModule& PropertyEditorModule = FModuleManager::GetModuleChecked<FPropertyEditorModule>("PropertyEditor");
+
+	for (UEnum* Enum : TObjectRange<UEnum>())
+	{
+		if (Enum->HasAnyFlags(RF_ClassDefaultObject | RF_MirroredGarbage | RF_NewerVersionExists))
+		{
+			continue;
+		}
+
+		Enums.Add(Enum);
+	}
+
+	Enums.Sort(
+		[](const UEnum& A, const UEnum& B)
+		{
+			return A.GetName() < B.GetName();
+		}
+	);
+
+	using namespace UE::RemoteControl::UI::Private;
+
+	TArray<TSharedPtr<FRCControllerPropertyInfo>> Types;
+	Types.Reserve(Enums.Num());
+
+	for (UEnum* Enum : Enums)
+	{
+		TSharedPtr<FRCControllerPropertyInfo>& Type = Types.Add_GetRef(MakeShared<FRCControllerPropertyInfo>(Enum));
+		Type->Type = EPropertyBagPropertyType::Enum;
+	}
+
+	InMenuBuilder.AddWidget(
+		SNew(SRCControllerTypeSelector, Types)
+		.OnTypeSelected(FRCControllerTypeBase::FOnTypeSelected::CreateSP(this, &SRCControllerPanel::OnAddControllerSelected)),
+		FText::GetEmpty(),
+		/* No Indent */ false
+	);
+}
+
 TSharedRef<SWidget> SRCControllerPanel::GetMultiControllerSwitchWidget()
 {
 	const FText& SwitchTooltipText = LOCTEXT("MultiControllerModeSwitchTooltip", "Enable/Disable MultiController Mode.\nWhen enabled, the Controllers list will only show one Controller per Field Id.");
@@ -474,7 +661,7 @@ TSharedRef<SWidget> SRCControllerPanel::GetMultiControllerSwitchWidget()
 		];
 }
 
-void SRCControllerPanel::OnAddControllerClicked(const EPropertyBagPropertyType InValueType, const FRCControllerPropertyInfo InControllerPropertyInfo) const
+void SRCControllerPanel::OnAddControllerClicked(const EPropertyBagPropertyType InValueType, const UE::RemoteControl::UI::Private::FRCControllerPropertyInfo InControllerPropertyInfo) const
 {
 	// Add to the asset
 	if (URemoteControlPreset* Preset = GetPreset())
@@ -509,6 +696,65 @@ void SRCControllerPanel::OnAddControllerClicked(const EPropertyBagPropertyType I
 	}
 }
 
+void SRCControllerPanel::OnAddControllerSelected(TSharedPtr<UE::RemoteControl::UI::Private::FRCControllerPropertyInfo> InControllerPropertyInfo) const
+{
+	if (UStruct* Struct = Cast<UStruct>(InControllerPropertyInfo->ValueTypeObject))
+	{
+		FInstancedPropertyBag Bag;
+		Bag.AddProperty("Default", EPropertyBagPropertyType::Struct, Struct);
+
+		TSharedRef<FStructOnScope> StructOnScope = MakeShared<FStructOnScope>(Bag.GetPropertyBagStruct(), Bag.GetMutableValue().GetMemory());
+
+		FPropertyRowGeneratorArgs Args;
+		Args.bShouldShowHiddenProperties = true;
+		TSharedRef<IPropertyRowGenerator> PropertyRowGenerator = FModuleManager::GetModuleChecked<FPropertyEditorModule>("PropertyEditor").CreatePropertyRowGenerator(Args);
+		PropertyRowGenerator->SetStructure(StructOnScope);
+
+		TArray<TSharedRef<IDetailTreeNode>> RootNodes = PropertyRowGenerator->GetRootTreeNodes();
+
+		bool bInvalidNodes = false;
+
+		for (const TSharedRef<IDetailTreeNode>& RootNode : RootNodes)
+		{
+			TArray<TSharedRef<IDetailTreeNode>> Children;
+			RootNode->GetChildren(Children);
+
+			for (const TSharedRef<IDetailTreeNode>& Child : Children)
+			{
+				TSharedPtr<IPropertyHandle> ChildHandle = Child->CreatePropertyHandle();
+
+				if (!ChildHandle.IsValid())
+				{
+					bInvalidNodes = true;
+					break;
+				}
+
+				if (ChildHandle->GetPropertyPath().Contains(TEXT("->")))
+				{
+					bInvalidNodes = true;
+					break;
+				}
+			}
+
+			if (bInvalidNodes)
+			{
+				break;
+			}
+		}
+
+		if (bInvalidNodes)
+		{
+			FMessageDialog::Open(
+				EAppMsgType::Ok,
+				LOCTEXT("InvalidNodesError", "This struct is not supported due to its customization.")
+			);
+			return;
+		}
+	}
+
+	OnAddControllerClicked(InControllerPropertyInfo->Type, *InControllerPropertyInfo);
+}
+
 void SRCControllerPanel::EnterRenameMode()
 {
 	if (ControllerPanelList)
diff --git a/Engine/Plugins/VirtualProduction/RemoteControl/Source/RemoteControlUI/Private/UI/Controller/SRCControllerPanel.h b/Engine/Plugins/VirtualProduction/RemoteControl/Source/RemoteControlUI/Private/UI/Controller/SRCControllerPanel.h
--- a/Engine/Plugins/VirtualProduction/RemoteControl/Source/RemoteControlUI/Private/UI/Controller/SRCControllerPanel.h
+++ b/Engine/Plugins/VirtualProduction/RemoteControl/Source/RemoteControlUI/Private/UI/Controller/SRCControllerPanel.h
@@ -5,12 +5,41 @@
 #include "StructUtils/PropertyBag.h"
 #include "UI/BaseLogicUI/SRCLogicPanelBase.h"
 
+class FMenuBuilder;
 class FRCControllerModel;
 class SRCControllerPanel;
 class URCController;
-struct FRCPanelStyle;
 enum class ECheckBoxState : uint8;
 enum class EPropertyBagPropertyType : uint8;
+struct FRCPanelStyle;
+
+namespace UE::RemoteControl::UI::Private
+{
+
+/**
+ * Helper struct used when generating Controllers.
+ * Wraps the TypeObject and the Bag Type.
+ *
+ * Additionally, some facultative MetaData can be added to create a Custom Controller extending the initial type
+ * e.g. External Texture custom controller is a specialized String controller
+ * See RCCustomControllerUtilities.h for helper functions related to custom controllers.
+ */
+struct FRCControllerPropertyInfo
+{
+	FRCControllerPropertyInfo() {}
+	FRCControllerPropertyInfo(UObject* InValueTypeObject) : ValueTypeObject(InValueTypeObject) {}
+
+	/** The type of the controller which will be created starting from this FRCControllerPropertyInfo */
+	EPropertyBagPropertyType Type = EPropertyBagPropertyType::Name;
+
+	/** Object defining an Enum, Struct, or Class */
+	UObject* ValueTypeObject = nullptr;
+
+	/** MetaData Map which can be used to pass additional settings or information to be used when customizing a Custom Controller */
+	TMap<FName, FString> CustomMetaData;
+};
+
+}
 
 /*
 * ~ SRCControllerPanel ~
@@ -20,28 +49,6 @@ enum class EPropertyBagPropertyType : uint8;
 */
 class REMOTECONTROLUI_API SRCControllerPanel : public SRCLogicPanelBase
 {
-	/** Helper struct used when generating Controllers.
-	 * Wraps the TypeObject and the Bag Type.
-	 * 
-	 * Additionally, some facultative MetaData can be added to create a Custom Controller extending the initial type
-	 * e.g. External Texture custom controller is a specialized String controller
-	 * See RCCustomControllerUtilities.h for helper functions related to custom controllers.
-	 */
-	struct FRCControllerPropertyInfo
-	{
-		FRCControllerPropertyInfo(){}
-		FRCControllerPropertyInfo(UObject* InValueTypeObject) : ValueTypeObject(InValueTypeObject){}
-
-		/** The type of the controller which will be created starting from this FRCControllerPropertyInfo */
-		EPropertyBagPropertyType Type = EPropertyBagPropertyType::Name;
-		
-		/** Object defining an Enum, Struct, or Class */
-		UObject* ValueTypeObject = nullptr;
-
-		/** MetaData Map which can be used to pass additional settings or information to be used when customizing a Custom Controller */
-		TMap<FName, FString> CustomMetaData;
-	};
-	
 public:
 	SLATE_BEGIN_ARGS(SRCControllerPanel)
 	{}
@@ -50,7 +57,7 @@ public:
 
 	SLATE_END_ARGS()
 
-	static FRCControllerPropertyInfo GetPropertyInfoForCustomType(const FName& InType);
+	static UE::RemoteControl::UI::Private::FRCControllerPropertyInfo GetPropertyInfoForCustomType(const FName& InType);
 
 	/** Constructs this widget with InArgs */
 	void Construct(const FArguments& InArgs, const TSharedRef<SRemoteControlPanel>& InPanel);
@@ -98,10 +105,19 @@ private:
 	/** Builds a menu containing the list of all possible Controllers*/
 	TSharedRef<SWidget> GetControllerMenuContentWidget() const;
 
+	/** Builds a menu containing the list of all possible structs on which to base a Controller. */
+	void GetControllerMenuStructContentWidget(FMenuBuilder& InMenuBuilder) const;
+
+	/** Builds a menu containing the list of all possible enums on which to base a Controller. */
+	void GetControllerMenuEnumContentWidget(FMenuBuilder& InMenuBuilder) const;
+
 	TSharedRef<SWidget> GetMultiControllerSwitchWidget();
 
-	/** Handles click event for Add Controller button*/
-	void OnAddControllerClicked(const EPropertyBagPropertyType InValueType, const FRCControllerPropertyInfo InControllerPropertyInfo) const;
+	/** Handles click event for Add Controller button */
+	void OnAddControllerClicked(const EPropertyBagPropertyType InValueType, const UE::RemoteControl::UI::Private::FRCControllerPropertyInfo InControllerPropertyInfo) const;
+
+	/** Handles selection of controller type in submenus. */
+	void OnAddControllerSelected(TSharedPtr<UE::RemoteControl::UI::Private::FRCControllerPropertyInfo> InControllerPropertyInfo) const;
 
 	/** Handles checkbox changes event for the MultiController mode selection widget */
 	void OnToggleMultiControllersMode(ECheckBoxState CheckBoxState);
diff --git a/Engine/Plugins/VirtualProduction/RemoteControl/Source/RemoteControlUI/Private/UI/Controller/TypeSelector/RCControllerTypeBase.h b/Engine/Plugins/VirtualProduction/RemoteControl/Source/RemoteControlUI/Private/UI/Controller/TypeSelector/RCControllerTypeBase.h
new file mode 100644
--- /dev/null
+++ b/Engine/Plugins/VirtualProduction/RemoteControl/Source/RemoteControlUI/Private/UI/Controller/TypeSelector/RCControllerTypeBase.h
@@ -0,0 +1,22 @@
+// Copyright Epic Games, Inc. All Rights Reserved.
+
+#pragma once
+
+#include "Delegates/Delegate.h"
+#include "Delegates/DelegateCombinations.h"
+#include "Templates/SharedPointer.h"
+
+namespace UE::RemoteControl::UI::Private
+{
+
+struct FRCControllerPropertyInfo;
+
+class FRCControllerTypeBase
+{
+public:
+	using ItemType = TSharedPtr<FRCControllerPropertyInfo>;
+
+	DECLARE_DELEGATE_OneParam(FOnTypeSelected, ItemType)
+};
+
+} // UE::RemoteControl::UI::Private
diff --git a/Engine/Plugins/VirtualProduction/RemoteControl/Source/RemoteControlUI/Private/UI/Controller/TypeSelector/SRCControllerTypeList.cpp b/Engine/Plugins/VirtualProduction/RemoteControl/Source/RemoteControlUI/Private/UI/Controller/TypeSelector/SRCControllerTypeList.cpp
new file mode 100644
--- /dev/null
+++ b/Engine/Plugins/VirtualProduction/RemoteControl/Source/RemoteControlUI/Private/UI/Controller/TypeSelector/SRCControllerTypeList.cpp
@@ -0,0 +1,110 @@
+// Copyright Epic Games, Inc. All Rights Reserved.
+
+#include "UI/Controller/TypeSelector/SRCControllerTypeList.h"
+
+#include "Framework/Application/SlateApplication.h"
+#include "UI/Controller/SRCControllerPanel.h"
+#include "UI/Controller/TypeSelector/SRCControllerTypeListEntry.h"
+#include "Widgets/Text/STextBlock.h"
+
+namespace UE::RemoteControl::UI::Private
+{
+
+void SRCControllerTypeList::Construct(const FArguments& InArgs, const TArray<ItemType>& InTypes)
+{
+	Types = InTypes;
+	Filter = FString();
+	OnTypeSelectedDelegate = InArgs._OnTypeSelected;
+
+	Super::Construct(
+		Super::FArguments()
+		.OnGenerateRow(this, &SRCControllerTypeList::GenerateRow)
+		.ListItemsSource(&Types)
+		.OnSelectionChanged(this, &SRCControllerTypeList::OnTypeSelected)
+		.SelectionMode(ESelectionMode::Single)
+	);
+}
+
+void SRCControllerTypeList::OnFilterTextChanged(const FText& InFilterText)
+{
+	const FString NewFilter = InFilterText.ToString();
+
+	if (!Filter.IsEmpty() && NewFilter.StartsWith(Filter))
+	{
+		UpdateItemsBasedOnFilter(NewFilter, /* Re-Filter */ true);
+	}
+	else
+	{
+		UpdateItemsBasedOnFilter(NewFilter, /* Re-Filter */ false);
+	}
+
+	Filter = NewFilter;
+}
+
+TSharedRef<ITableRow> SRCControllerTypeList::GenerateRow(ItemType InItem, const TSharedRef<STableViewBase>& InOwnerTable)
+{
+	return SNew(SRCControllerTypeListEntry, InOwnerTable, InItem);
+}
+
+void SRCControllerTypeList::UpdateItemsBasedOnFilter(const FString& InFilter, bool bInReFilter)
+{
+	if (InFilter.IsEmpty())
+	{
+		SetItemsSource(&Types);
+		FilteredTypes.Empty();
+		return;
+	}
+
+	TArray<ItemType>& UnfilteredTypes = bInReFilter
+		? FilteredTypes
+		: Types;
+
+	TArray<ItemType> NewTypes;
+	NewTypes.Reserve(UnfilteredTypes.Num());
+
+	for (const ItemType& Type : UnfilteredTypes)
+	{
+		switch (Type->Type)
+		{
+			case EPropertyBagPropertyType::Enum:
+			{
+				if (UEnum* Enum = Cast<UEnum>(Type->ValueTypeObject))
+				{
+					if (Enum->GetDisplayNameText().ToString().Contains(InFilter, ESearchCase::IgnoreCase))
+					{
+						NewTypes.Add(Type);
+					}
+				}
+				break;
+			}
+
+			case EPropertyBagPropertyType::Struct:
+			{
+				if (UScriptStruct* Struct = Cast<UScriptStruct>(Type->ValueTypeObject))
+				{
+					if (Struct->GetDisplayNameText().ToString().Contains(InFilter, ESearchCase::IgnoreCase))
+					{
+						NewTypes.Add(Type);
+					}
+				}
+				break;
+			}
+		}
+	}
+
+	FilteredTypes = NewTypes;
+	ClearItemsSource();
+	SetItemsSource(&FilteredTypes);
+}
+
+void SRCControllerTypeList::OnTypeSelected(ItemType InItem, ESelectInfo::Type InSelectInfo)
+{
+	FSlateApplication::Get().DismissAllMenus();
+
+	if (OnTypeSelectedDelegate.IsBound())
+	{
+		OnTypeSelectedDelegate.Execute(InItem);
+	}
+}
+
+} // UE::RemoteControl::UI::Private
diff --git a/Engine/Plugins/VirtualProduction/RemoteControl/Source/RemoteControlUI/Private/UI/Controller/TypeSelector/SRCControllerTypeList.h b/Engine/Plugins/VirtualProduction/RemoteControl/Source/RemoteControlUI/Private/UI/Controller/TypeSelector/SRCControllerTypeList.h
new file mode 100644
--- /dev/null
+++ b/Engine/Plugins/VirtualProduction/RemoteControl/Source/RemoteControlUI/Private/UI/Controller/TypeSelector/SRCControllerTypeList.h
@@ -0,0 +1,46 @@
+// Copyright Epic Games, Inc. All Rights Reserved.
+
+#pragma once
+
+#include "UI/Controller/TypeSelector/RCControllerTypeBase.h"
+#include "Widgets/Views/SListView.h"
+
+#include "Containers/Array.h"
+#include "Containers/UnrealString.h"
+#include "Framework/SlateDelegates.h"
+#include "Internationalization/Text.h"
+
+namespace UE::RemoteControl::UI::Private
+{
+
+struct FRCControllerPropertyInfo;
+
+class SRCControllerTypeList : public SListView<TSharedPtr<FRCControllerPropertyInfo>>, public FRCControllerTypeBase
+{	
+	using Super = SListView<ItemType>;
+
+public:
+	SLATE_BEGIN_ARGS(SRCControllerTypeList)
+		{}
+		SLATE_ARGUMENT(FOnTypeSelected, OnTypeSelected)
+	SLATE_END_ARGS()
+
+	void Construct(const FArguments& InArgs, const TArray<ItemType>& InTypes);
+
+	void OnFilterTextChanged(const FText& InFilterText);
+
+private:
+	TArray<ItemType> Types;
+	TArray<ItemType> FilteredTypes;
+	FString Filter;
+	FOnTypeSelected OnTypeSelectedDelegate;
+
+	TSharedRef<ITableRow> GenerateRow(ItemType InItem, const TSharedRef<STableViewBase>& InOwnerTable);
+
+	/** Setting bReFilter to true will filter the already filtered list. */
+	void UpdateItemsBasedOnFilter(const FString& InFilter, bool bInReFilter);
+
+	void OnTypeSelected(ItemType InItem, ESelectInfo::Type InSelectInfo);
+};
+
+} // UE::RemoteControl::UI::Private
diff --git a/Engine/Plugins/VirtualProduction/RemoteControl/Source/RemoteControlUI/Private/UI/Controller/TypeSelector/SRCControllerTypeListEntry.cpp b/Engine/Plugins/VirtualProduction/RemoteControl/Source/RemoteControlUI/Private/UI/Controller/TypeSelector/SRCControllerTypeListEntry.cpp
new file mode 100644
--- /dev/null
+++ b/Engine/Plugins/VirtualProduction/RemoteControl/Source/RemoteControlUI/Private/UI/Controller/TypeSelector/SRCControllerTypeListEntry.cpp
@@ -0,0 +1,45 @@
+// Copyright Epic Games, Inc. All Rights Reserved.
+
+#include "UI/Controller/TypeSelector/SRCControllerTypeListEntry.h"
+
+#include "UI/Controller/SRCControllerPanel.h"
+#include "Widgets/Text/STextBlock.h"
+
+namespace UE::RemoteControl::UI::Private
+{
+
+void SRCControllerTypeListEntry::Construct(const FArguments& InArgs, const TSharedRef<STableViewBase>& InOwnerTable, ItemType InItem)
+{
+	Super::Construct(Super::FArguments(), InOwnerTable);
+
+	switch (InItem->Type)
+	{
+		case EPropertyBagPropertyType::Enum:
+		{
+			if (UEnum* Enum = Cast<UEnum>(InItem->ValueTypeObject))
+			{
+				ChildSlot
+				[
+					SNew(STextBlock)
+					.Text(Enum->GetDisplayNameText())
+				];
+			}
+			break;
+		}
+
+		case EPropertyBagPropertyType::Struct:
+		{
+			if (UScriptStruct* Struct = Cast<UScriptStruct>(InItem->ValueTypeObject))
+			{
+				ChildSlot
+				[
+					SNew(STextBlock)
+					.Text(Struct->GetDisplayNameText())
+				];
+			}
+			break;
+		}
+	}
+}
+
+}
diff --git a/Engine/Plugins/VirtualProduction/RemoteControl/Source/RemoteControlUI/Private/UI/Controller/TypeSelector/SRCControllerTypeListEntry.h b/Engine/Plugins/VirtualProduction/RemoteControl/Source/RemoteControlUI/Private/UI/Controller/TypeSelector/SRCControllerTypeListEntry.h
new file mode 100644
--- /dev/null
+++ b/Engine/Plugins/VirtualProduction/RemoteControl/Source/RemoteControlUI/Private/UI/Controller/TypeSelector/SRCControllerTypeListEntry.h
@@ -0,0 +1,27 @@
+// Copyright Epic Games, Inc. All Rights Reserved.
+
+#pragma once
+
+#include "UI/Controller/TypeSelector/RCControllerTypeBase.h"
+#include "Widgets/Views/SListView.h"
+#include "Widgets/Views/STableRow.h"
+
+namespace UE::RemoteControl::UI::Private
+{
+
+struct FRCControllerPropertyInfo;
+
+class SRCControllerTypeListEntry : public STableRow<TSharedPtr<FRCControllerPropertyInfo>>, public FRCControllerTypeBase
+{
+	using Super = STableRow<ItemType>;
+
+public:
+	SLATE_BEGIN_ARGS(SRCControllerTypeListEntry)
+		{}
+	SLATE_END_ARGS()
+
+	/** Constructs this widget with InArgs */
+	void Construct(const FArguments& InArgs, const TSharedRef<STableViewBase>& InOwnerTable, ItemType InItem);
+};
+
+} // UE::RemoteControl::UI::Private
diff --git a/Engine/Plugins/VirtualProduction/RemoteControl/Source/RemoteControlUI/Private/UI/Controller/TypeSelector/SRCControllerTypeSelector.cpp b/Engine/Plugins/VirtualProduction/RemoteControl/Source/RemoteControlUI/Private/UI/Controller/TypeSelector/SRCControllerTypeSelector.cpp
new file mode 100644
--- /dev/null
+++ b/Engine/Plugins/VirtualProduction/RemoteControl/Source/RemoteControlUI/Private/UI/Controller/TypeSelector/SRCControllerTypeSelector.cpp
@@ -0,0 +1,49 @@
+// Copyright Epic Games, Inc. All Rights Reserved.
+
+#include "UI/Controller/TypeSelector/SRCControllerTypeSelector.h"
+
+#include "UI/Controller/TypeSelector/SRCControllerTypeList.h"
+#include "Widgets/Input/SSearchBox.h"
+#include "Widgets/Layout/SBox.h"
+#include "Widgets/SBoxPanel.h"
+
+namespace UE::RemoteControl::UI::Private
+{
+
+void SRCControllerTypeSelector::Construct(const FArguments& InArgs, const TArray<ItemType>& InTypes)
+{
+	ChildSlot
+	[
+		SNew(SVerticalBox)
+		+ SVerticalBox::Slot()
+		.AutoHeight()
+		.Padding(4.f, 4.f, 4.f, 4.f)
+		[
+			SNew(SSearchBox)
+			.OnTextChanged(this, &SRCControllerTypeSelector::OnFilterTextCommitted, ETextCommit::Default)
+			.OnTextCommitted(this, &SRCControllerTypeSelector::OnFilterTextCommitted)
+		]
+		+ SVerticalBox::Slot()
+		.AutoHeight()
+		.Padding(4.f, 0.f, 4.f, 4.f)
+		[
+			SNew(SBox)
+			.WidthOverride(300.f)
+			.HeightOverride(300.f)
+			[
+				SAssignNew(TypeList, SRCControllerTypeList, InTypes)
+				.OnTypeSelected(InArgs._OnTypeSelected)
+			]
+		]
+	];
+}
+
+void SRCControllerTypeSelector::OnFilterTextCommitted(const FText& InNewText, ETextCommit::Type InCommitInfo)
+{
+	if (InCommitInfo != ETextCommit::OnUserMovedFocus)
+	{
+		TypeList->OnFilterTextChanged(InNewText);
+	}
+}
+
+} // UE::RemoteControl::UI::Private
diff --git a/Engine/Plugins/VirtualProduction/RemoteControl/Source/RemoteControlUI/Private/UI/Controller/TypeSelector/SRCControllerTypeSelector.h b/Engine/Plugins/VirtualProduction/RemoteControl/Source/RemoteControlUI/Private/UI/Controller/TypeSelector/SRCControllerTypeSelector.h
new file mode 100644
--- /dev/null
+++ b/Engine/Plugins/VirtualProduction/RemoteControl/Source/RemoteControlUI/Private/UI/Controller/TypeSelector/SRCControllerTypeSelector.h
@@ -0,0 +1,31 @@
+// Copyright Epic Games, Inc. All Rights Reserved.
+
+#pragma once
+
+#include "UI/Controller/TypeSelector/RCControllerTypeBase.h"
+#include "Widgets/SCompoundWidget.h"
+
+class FText;
+
+namespace UE::RemoteControl::UI::Private
+{
+
+class SRCControllerTypeList;
+
+class SRCControllerTypeSelector : public SCompoundWidget, public FRCControllerTypeBase
+{
+public:
+	SLATE_BEGIN_ARGS(SRCControllerTypeSelector)
+		{}
+		SLATE_ARGUMENT(FOnTypeSelected, OnTypeSelected)
+	SLATE_END_ARGS()
+
+	void Construct(const FArguments& InArgs, const TArray<ItemType>& InTypes);
+
+private:
+	TSharedPtr<SRCControllerTypeList> TypeList;
+
+	void OnFilterTextCommitted(const FText& InNewText, ETextCommit::Type InCommitInfo);
+};
+
+} // UE::RemoteControl::UI::Private
