Remote Control: Select bound action sources

diff --git a/Engine/Plugins/VirtualProduction/RemoteControl/Source/RemoteControlUI/Private/Commands/RemoteControlCommands.cpp b/Engine/Plugins/VirtualProduction/RemoteControl/Source/RemoteControlUI/Private/Commands/RemoteControlCommands.cpp
--- a/Engine/Plugins/VirtualProduction/RemoteControl/Source/RemoteControlUI/Private/Commands/RemoteControlCommands.cpp
+++ b/Engine/Plugins/VirtualProduction/RemoteControl/Source/RemoteControlUI/Private/Commands/RemoteControlCommands.cpp
@@ -103,6 +103,13 @@ void FRemoteControlCommands::RegisterCommands()
 		EUserInterfaceActionType::Button,
 		FInputChord(EKeys::U, EModifierKey::Control));
 
+	// Select Source
+	UI_COMMAND(SelectSource,
+		"Select Source",
+		"Selects the source of the bound value",
+		EUserInterfaceActionType::Button,
+		FInputChord());
+
 	// Protocols generate property changed events
 	UI_COMMAND(ProtocolsGeneratePropertyChangeEvents,
 		"Protocols Generate Property Change Events",
diff --git a/Engine/Plugins/VirtualProduction/RemoteControl/Source/RemoteControlUI/Private/Commands/RemoteControlCommands.h b/Engine/Plugins/VirtualProduction/RemoteControl/Source/RemoteControlUI/Private/Commands/RemoteControlCommands.h
--- a/Engine/Plugins/VirtualProduction/RemoteControl/Source/RemoteControlUI/Private/Commands/RemoteControlCommands.h
+++ b/Engine/Plugins/VirtualProduction/RemoteControl/Source/RemoteControlUI/Private/Commands/RemoteControlCommands.h
@@ -86,6 +86,11 @@ public:
 	 */
 	TSharedPtr<FUICommandInfo> UpdateValue;
 
+	/**
+	 * UI Command for selecting the source of a property in an action. Currently used for Logic panel
+	 */
+	TSharedPtr<FUICommandInfo> SelectSource;
+
 	/**
 	 * UI Command to Toggle if protocols generate property changed events
 	 */
diff --git a/Engine/Plugins/VirtualProduction/RemoteControl/Source/RemoteControlUI/Private/RemoteControlUIModule.cpp b/Engine/Plugins/VirtualProduction/RemoteControl/Source/RemoteControlUI/Private/RemoteControlUIModule.cpp
--- a/Engine/Plugins/VirtualProduction/RemoteControl/Source/RemoteControlUI/Private/RemoteControlUIModule.cpp
+++ b/Engine/Plugins/VirtualProduction/RemoteControl/Source/RemoteControlUI/Private/RemoteControlUIModule.cpp
@@ -1221,6 +1221,73 @@ TSharedPtr<SWidget> FRemoteControlUIModule::CreateCustomControllerWidget(URCVirt
 	return CustomControllerWidget;
 }
 
+void FRemoteControlUIModule::SelectSource(TSharedPtr<FRemoteControlEntity> InEntity)
+{
+	const TSharedPtr<FRemoteControlField> RCField = StaticCastSharedPtr<FRemoteControlField>(InEntity);
+
+	if (!RCField.IsValid() || RCField->GetStruct() != FRemoteControlProperty::StaticStruct())
+	{
+		return;
+	}
+
+	const TSharedPtr<FRemoteControlProperty> RCProp = StaticCastSharedPtr<FRemoteControlProperty>(RCField);
+
+	// Resolve it to get the property path.
+	RCProp->FieldPathInfo.Resolve(RCProp->GetBoundObject());
+
+	if (!RCProp->FieldPathInfo.IsResolved())
+	{
+		return;
+	}
+
+	TSharedRef<FPropertyPath> PropertyPath = RCProp->FieldPathInfo.ToPropertyPath();
+	TArray<UObject*> BoundObjects = RCField->GetBoundObjects();
+
+	for (auto It = BoundObjects.CreateIterator(); It; ++It)
+	{
+		if (AActor* OwnerActor = (*It)->GetTypedOuter<AActor>())
+		{
+			const UObject* BindingObject = *It;
+
+			// When we encounter a non-component object, 
+			if (!BindingObject->IsA<UActorComponent>())
+			{
+				BoundObjects.Add(OwnerActor);
+				It.RemoveCurrent();
+			}
+
+			// --- Special NDisplay handling because of their customization ---
+			// Since display cluster config data is not created as a default subobject, therefore we have no way of retrieving the CurrentConfigData property from the config object itself.
+			static FName DisplayClusterConfigDataClassName = "DisplayClusterConfigurationData";
+			if (BindingObject->GetClass()->GetFName() == DisplayClusterConfigDataClassName)
+			{
+				if (FProperty* Property = OwnerActor->GetClass()->FindPropertyByName("CurrentConfigData"))
+				{
+					// Append "CurrentConfigData" to the beginning of the path.
+					TSharedRef<FPropertyPath> NewPropertyPath = FPropertyPath::Create(Property);
+					for (int32 Index = 0; Index < PropertyPath->GetNumProperties(); Index++)
+					{
+						NewPropertyPath->AddProperty(PropertyPath->GetPropertyInfo(Index));
+					}
+
+					PropertyPath = NewPropertyPath;
+				}
+			}
+		}
+	}
+
+	SelectObjects(BoundObjects);
+
+	FTimerHandle Handle;
+	const FTimerDelegate Delegate = FTimerDelegate::CreateLambda(([PropertyPath]()
+		{
+			FRemoteControlUIModule::Get().HighlightPropertyInDetailsPanel(*PropertyPath);
+		}));
+
+	// Needed because modifying the selection set is asynchronous.
+	GEditor->GetTimerManager()->SetTimer(Handle, Delegate, 0.1, false, -1);
+}
+
 void FRemoteControlUIModule::RegisterWidgetFactories()
 {
 	RegisterWidgetFactoryForType(FRemoteControlActor::StaticStruct(), FOnGenerateRCWidget::CreateStatic(&SRCPanelExposedActor::MakeInstance));
diff --git a/Engine/Plugins/VirtualProduction/RemoteControl/Source/RemoteControlUI/Private/RemoteControlUIModule.h b/Engine/Plugins/VirtualProduction/RemoteControl/Source/RemoteControlUI/Private/RemoteControlUIModule.h
--- a/Engine/Plugins/VirtualProduction/RemoteControl/Source/RemoteControlUI/Private/RemoteControlUIModule.h
+++ b/Engine/Plugins/VirtualProduction/RemoteControl/Source/RemoteControlUI/Private/RemoteControlUIModule.h
@@ -75,6 +75,8 @@ public:
 	virtual TSharedPtr<SWidget> CreateCustomControllerWidget(URCVirtualPropertyBase* InController, TSharedPtr<IPropertyHandle> InOriginalPropertyHandle) const override;
 	//~ End IRemoteControlUIModule interface
 
+	void SelectSource(TSharedPtr<FRemoteControlEntity> InEntity);
+
 	/**
 	 * Create a remote control panel to display a given preset.
 	 * @param The preset to display.
diff --git a/Engine/Plugins/VirtualProduction/RemoteControl/Source/RemoteControlUI/Private/UI/Action/SRCActionPanelList.h b/Engine/Plugins/VirtualProduction/RemoteControl/Source/RemoteControlUI/Private/UI/Action/SRCActionPanelList.h
--- a/Engine/Plugins/VirtualProduction/RemoteControl/Source/RemoteControlUI/Private/UI/Action/SRCActionPanelList.h
+++ b/Engine/Plugins/VirtualProduction/RemoteControl/Source/RemoteControlUI/Private/UI/Action/SRCActionPanelList.h
@@ -6,9 +6,13 @@
 #include "Action/RCActionContainer.h"
 #include "Action/RCFunctionAction.h"
 #include "Action/RCPropertyAction.h"
+#include "Commands/RemoteControlCommands.h"
+#include "Framework/Commands/UIAction.h"
 #include "RCActionModel.h"
 #include "RemoteControlPreset.h"
+#include "RemoteControlUIModule.h"
 #include "SDropTarget.h"
+#include "Selection.h"
 #include "SlateOptMacros.h"
 #include "SRCActionPanel.h"
 #include "Styling/RemoteControlStyles.h"
@@ -41,7 +45,8 @@ class FRCActionModel;
 * UI Widget for Actions List
 * Used as part of the RC Logic Actions Panel.
 */
-template <typename ActionType>
+template <typename ActionType
+	UE_REQUIRES(TIsDerivedFrom<ActionType, FRCActionModel>::Value)>
 class REMOTECONTROLUI_API SRCActionPanelList : public SRCLogicPanelListBase
 {
 public:
@@ -59,6 +64,13 @@ public:
 
 		SRCLogicPanelListBase::Construct(SRCLogicPanelListBase::FArguments(), InActionPanel, RemoteControlPanel.ToSharedRef());
 
+		const FRemoteControlCommands& Commands = FRemoteControlCommands::Get();
+
+		CommandList->MapAction(
+			Commands.SelectSource,
+			FUIAction(FExecuteAction::CreateSP(this, &SRCActionPanelList::SelectSource_Execute))
+		);
+
 		ActionPanelWeakPtr = InActionPanel;
 		BehaviourItemWeakPtr = InBehaviourItem;
 
@@ -581,6 +593,30 @@ private:
 		SelectedActionItem = InItem;
 	}
 
+	void SelectSource_Execute()
+	{
+		TArray<TSharedPtr<ActionType>> SelectedItems = ListView->GetSelectedItems();
+		TArray<UObject*> ObjectsToSelect;
+
+		for (const TSharedPtr<ActionType>& SelectedAction : SelectedItems)
+		{
+			if (SelectedAction.IsValid())
+			{
+				if (URCAction* Action = SelectedAction->GetAction())
+				{
+					if (URemoteControlPreset* Preset = Action->PresetWeakPtr.Get())
+					{
+						if (TSharedPtr<FRemoteControlEntity> Entity = Preset->GetExposedEntity(Action->ExposedFieldId).Pin())
+						{
+							FRemoteControlUIModule::Get().SelectSource(Entity);
+							return;
+						}
+					}
+				}
+			}
+		}
+	}
+
 private:
 
 	/** The currently selected Action item */
diff --git a/Engine/Plugins/VirtualProduction/RemoteControl/Source/RemoteControlUI/Private/UI/BaseLogicUI/SRCLogicPanelListBase.cpp b/Engine/Plugins/VirtualProduction/RemoteControl/Source/RemoteControlUI/Private/UI/BaseLogicUI/SRCLogicPanelListBase.cpp
--- a/Engine/Plugins/VirtualProduction/RemoteControl/Source/RemoteControlUI/Private/UI/BaseLogicUI/SRCLogicPanelListBase.cpp
+++ b/Engine/Plugins/VirtualProduction/RemoteControl/Source/RemoteControlUI/Private/UI/BaseLogicUI/SRCLogicPanelListBase.cpp
@@ -29,6 +29,7 @@ TSharedPtr<SWidget> SRCLogicPanelListBase::GetContextMenuWidget()
 
 	const FRemoteControlCommands& Commands = FRemoteControlCommands::Get();
 
+	MenuBuilder.AddMenuEntry(Commands.SelectSource);
 	MenuBuilder.AddMenuEntry(Commands.CopyItem);
 
 	FText PasteItemLabel = LOCTEXT("Paste", "Paste");
diff --git a/Engine/Plugins/VirtualProduction/RemoteControl/Source/RemoteControlUI/Private/UI/SRCPanelExposedEntitiesList.cpp b/Engine/Plugins/VirtualProduction/RemoteControl/Source/RemoteControlUI/Private/UI/SRCPanelExposedEntitiesList.cpp
--- a/Engine/Plugins/VirtualProduction/RemoteControl/Source/RemoteControlUI/Private/UI/SRCPanelExposedEntitiesList.cpp
+++ b/Engine/Plugins/VirtualProduction/RemoteControl/Source/RemoteControlUI/Private/UI/SRCPanelExposedEntitiesList.cpp
@@ -137,7 +137,6 @@ public:
 	}
 
 private:
-
 	//~ SWidget Interface
 	virtual FReply OnMouseButtonDoubleClick(const FGeometry& InMyGeometry, const FPointerEvent& InMouseEvent) override
 	{
@@ -147,80 +146,7 @@ private:
 			{
 				if (const TSharedPtr<FRemoteControlEntity> ExposedEntity = ExposedEntityWidget->GetEntity())
 				{
-					if (const TSharedPtr<FRemoteControlField> RCField = StaticCastSharedPtr<FRemoteControlField>(ExposedEntity))
-					{
-						TArray<UObject*> BoundObjects = RCField->GetBoundObjects();
-						TSet<UObject*> Objects = TSet<UObject*>{ BoundObjects };
-
-						TArray<UObject*> OwnerActors;
-						for (UObject* Object : Objects)
-						{
-							if (Object->IsA<AActor>())
-							{
-								OwnerActors.Add(Object);
-							}
-							else
-							{
-								OwnerActors.Add(Object->GetTypedOuter<AActor>());
-							}
-						}
-
-						if (RCField->GetStruct() == FRemoteControlProperty::StaticStruct())
-						{
-							const TSharedPtr<FRemoteControlProperty> RCProp = StaticCastSharedPtr<FRemoteControlProperty>(RCField);
-
-							// Resolve it to get the property path.
-							RCProp->FieldPathInfo.Resolve(RCProp->GetBoundObject());
-							if (RCProp->FieldPathInfo.IsResolved())
-							{
-								TSharedRef<FPropertyPath> PropertyPath = RCProp->FieldPathInfo.ToPropertyPath();
-
-								for (auto It = BoundObjects.CreateIterator(); It; ++It)
-								{
-									if (AActor* OwnerActor = (*It)->GetTypedOuter<AActor>())
-									{
-										const UObject* BindingObject = *It;
-
-										// When we encounter a non-component object, 
-										if (!BindingObject->IsA<UActorComponent>())
-										{
-											BoundObjects.Add(OwnerActor);
-											It.RemoveCurrent();
-										}
-
-										// --- Special NDisplay handling because of their customization ---
-										// Since display cluster config data is not created as a default subobject, therefore we have no way of retrieving the CurrentConfigData property from the config object itself.
-										static FName DisplayClusterConfigDataClassName = "DisplayClusterConfigurationData";
-										if (BindingObject->GetClass()->GetFName() == DisplayClusterConfigDataClassName)
-										{
-											if (FProperty* Property = OwnerActor->GetClass()->FindPropertyByName("CurrentConfigData"))
-											{
-												// Append "CurrentConfigData" to the beginning of the path.
-												TSharedRef<FPropertyPath> NewPropertyPath = FPropertyPath::Create(Property);
-												for (int32 Index = 0; Index < PropertyPath->GetNumProperties(); Index++)
-												{
-													NewPropertyPath->AddProperty(PropertyPath->GetPropertyInfo(Index));
-												}
-
-												PropertyPath = NewPropertyPath;
-											}
-										}
-									}
-								}
-
-								FRemoteControlUIModule::Get().SelectObjects(BoundObjects);
-								
-								FTimerHandle Handle;
-								const FTimerDelegate Delegate = FTimerDelegate::CreateLambda(([PropertyPath]()
-									{
-										FRemoteControlUIModule::Get().HighlightPropertyInDetailsPanel(*PropertyPath);
-									}));
-								
-								// Needed because modifying the selection set is asynchronous.
-								GEditor->GetTimerManager()->SetTimer(Handle, Delegate, 0.1, false, -1);
-							}
-						}
-					}
+					FRemoteControlUIModule::Get().SelectSource(ExposedEntity);
 				}
 			}
 		}
@@ -239,7 +165,19 @@ private:
 
 void SRCPanelExposedEntitiesList::Construct(const FArguments& InArgs, URemoteControlPreset* InPreset, TWeakPtr<FRCPanelWidgetRegistry> InWidgetRegistry)
 {
-	CommandList = InArgs._CommandList;
+	CommandList = MakeShared<FUICommandList>();
+
+	if (InArgs._CommandList.IsValid())
+	{
+		CommandList->Append(InArgs._CommandList.ToSharedRef());
+	}
+
+	const FRemoteControlCommands& Commands = FRemoteControlCommands::Get();
+
+	CommandList->MapAction(
+		Commands.SelectSource,
+		FUIAction(FExecuteAction::CreateSP(this, &SRCPanelExposedEntitiesList::SelectSource_Execute))
+	);
 
 	bIsInLiveMode = InArgs._LiveMode;
 	bIsInProtocolsMode = InArgs._ProtocolsMode;
@@ -922,6 +860,25 @@ void SRCPanelExposedEntitiesList::RecallListSettings()
 	OrderGroups();
 }
 
+void SRCPanelExposedEntitiesList::SelectSource_Execute()
+{
+	TSharedPtr<SRCPanelTreeNode> SelectedEntity = GetSelectedEntity();
+
+	if (!SelectedEntity.IsValid() || SelectedEntity->GetRCType() == SRCPanelTreeNode::FieldGroup)
+	{
+		return;
+	}
+
+	TSharedPtr<FRemoteControlEntity> Entity = StaticCastSharedPtr<SRCPanelExposedEntity>(SelectedEntity)->GetEntity();
+
+	if (!Entity.IsValid())
+	{
+		return;
+	}
+
+	FRemoteControlUIModule::Get().SelectSource(Entity);
+}
+
 void SRCPanelExposedEntitiesList::OnObjectPropertyChange(UObject* InObject, FPropertyChangedEvent& InChangeEvent)
 {
 	static constexpr EPropertyChangeType::Type TypesNeedingRefresh = EPropertyChangeType::ArrayAdd | EPropertyChangeType::ArrayClear | EPropertyChangeType::ArrayRemove | EPropertyChangeType::ValueSet | EPropertyChangeType::Duplicate;
diff --git a/Engine/Plugins/VirtualProduction/RemoteControl/Source/RemoteControlUI/Private/UI/SRCPanelExposedEntitiesList.h b/Engine/Plugins/VirtualProduction/RemoteControl/Source/RemoteControlUI/Private/UI/SRCPanelExposedEntitiesList.h
--- a/Engine/Plugins/VirtualProduction/RemoteControl/Source/RemoteControlUI/Private/UI/SRCPanelExposedEntitiesList.h
+++ b/Engine/Plugins/VirtualProduction/RemoteControl/Source/RemoteControlUI/Private/UI/SRCPanelExposedEntitiesList.h
@@ -272,6 +272,9 @@ private:
 
 	/** Recalls the list settings */
 	void RecallListSettings();
+	
+	/** Selects the source of the currently selected entity. */
+	void SelectSource_Execute();
 
 private:
 	/** Command list to use if specified */
diff --git a/Engine/Plugins/VirtualProduction/RemoteControl/Source/RemoteControlUI/Private/UI/SRCPanelExposedEntity.cpp b/Engine/Plugins/VirtualProduction/RemoteControl/Source/RemoteControlUI/Private/UI/SRCPanelExposedEntity.cpp
--- a/Engine/Plugins/VirtualProduction/RemoteControl/Source/RemoteControlUI/Private/UI/SRCPanelExposedEntity.cpp
+++ b/Engine/Plugins/VirtualProduction/RemoteControl/Source/RemoteControlUI/Private/UI/SRCPanelExposedEntity.cpp
@@ -90,6 +90,7 @@ TSharedPtr<SWidget> SRCPanelExposedEntity::GetContextMenu()
 	MenuBuilder.AddMenuEntry(FRemoteControlCommands::Get().RenameEntity, NAME_None, TAttribute<FText>(), TAttribute<FText>(), FSlateIcon(FAppStyle::GetAppStyleSetName(), TEXT("GenericCommands.Rename")));
 	MenuBuilder.AddMenuEntry(FRemoteControlCommands::Get().ChangePropId, NAME_None, TAttribute<FText>(), TAttribute<FText>(), FSlateIcon(FAppStyle::GetAppStyleSetName(), TEXT("Icons.Edit")));
 	MenuBuilder.AddMenuEntry(FRemoteControlCommands::Get().DeleteEntity, NAME_None, TAttribute<FText>(), TAttribute<FText>(), FSlateIcon(FAppStyle::GetAppStyleSetName(), TEXT("GenericCommands.Delete")));
+	MenuBuilder.AddMenuEntry(FRemoteControlCommands::Get().SelectSource);
 
 	MenuBuilder.EndSection();
 
