Remote Control: Added indication of whether a property is bound to a controller.

diff --git a/Engine/Plugins/VirtualProduction/RemoteControl/Source/RemoteControl/Private/RCVirtualPropertyContainer.cpp b/Engine/Plugins/VirtualProduction/RemoteControl/Source/RemoteControl/Private/RCVirtualPropertyContainer.cpp
--- a/Engine/Plugins/VirtualProduction/RemoteControl/Source/RemoteControl/Private/RCVirtualPropertyContainer.cpp
+++ b/Engine/Plugins/VirtualProduction/RemoteControl/Source/RemoteControl/Private/RCVirtualPropertyContainer.cpp
@@ -191,6 +191,11 @@ URCVirtualPropertyInContainer* URCVirtualPropertyContainerBase::DuplicateVirtual
 	return nullptr;
 }
 
+bool URCVirtualPropertyContainerBase::HasVirtualProperty(URCVirtualPropertyBase* InVirtualProperty) const
+{
+	return InVirtualProperty && VirtualProperties.Contains(InVirtualProperty);
+}
+
 bool URCVirtualPropertyContainerBase::RemoveProperty(const FName& InPropertyName)
 {
 	Bag.RemovePropertyByName(InPropertyName);
diff --git a/Engine/Plugins/VirtualProduction/RemoteControl/Source/RemoteControl/Private/RemoteControlPreset.cpp b/Engine/Plugins/VirtualProduction/RemoteControl/Source/RemoteControl/Private/RemoteControlPreset.cpp
--- a/Engine/Plugins/VirtualProduction/RemoteControl/Source/RemoteControl/Private/RemoteControlPreset.cpp
+++ b/Engine/Plugins/VirtualProduction/RemoteControl/Source/RemoteControl/Private/RemoteControlPreset.cpp
@@ -916,6 +916,16 @@ URCVirtualPropertyInContainer* URemoteControlPreset::AddController(TSubclassOf<U
 	return NewController;
 }
 
+bool URemoteControlPreset::HasController(URCVirtualPropertyBase* InController) const
+{
+	if (!ensure(ControllerContainer))
+	{
+		return false;
+	}
+
+	return ControllerContainer->HasVirtualProperty(InController);
+}
+
 bool URemoteControlPreset::RemoveController(const FName& InPropertyName)
 {
 	if (ensure(ControllerContainer))
diff --git a/Engine/Plugins/VirtualProduction/RemoteControl/Source/RemoteControl/Public/RCVirtualPropertyContainer.h b/Engine/Plugins/VirtualProduction/RemoteControl/Source/RemoteControl/Public/RCVirtualPropertyContainer.h
--- a/Engine/Plugins/VirtualProduction/RemoteControl/Source/RemoteControl/Public/RCVirtualPropertyContainer.h
+++ b/Engine/Plugins/VirtualProduction/RemoteControl/Source/RemoteControl/Public/RCVirtualPropertyContainer.h
@@ -73,6 +73,10 @@ public:
 	 */
 	virtual URCVirtualPropertyInContainer* DuplicateVirtualProperty(URCVirtualPropertyInContainer* InVirtualProperty);
 
+	/**
+	 * Returns whether this container contains the given property.
+	 */
+	bool HasVirtualProperty(URCVirtualPropertyBase* InVirtualProperty) const;
 
 	/**
 	 * Removes a property from the container by name if it exists.
diff --git a/Engine/Plugins/VirtualProduction/RemoteControl/Source/RemoteControl/Public/RemoteControlPreset.h b/Engine/Plugins/VirtualProduction/RemoteControl/Source/RemoteControl/Public/RemoteControlPreset.h
--- a/Engine/Plugins/VirtualProduction/RemoteControl/Source/RemoteControl/Public/RemoteControlPreset.h
+++ b/Engine/Plugins/VirtualProduction/RemoteControl/Source/RemoteControl/Public/RemoteControlPreset.h
@@ -774,6 +774,9 @@ public:
 	/** Adds a Virtual Property (Controller) to the Remote Control Preset */
 	URCVirtualPropertyInContainer* AddController(TSubclassOf<URCVirtualPropertyInContainer> InPropertyClass, const EPropertyBagPropertyType InValueType, UObject* InValueTypeObject = nullptr, const FName InPropertyName = NAME_None);
 
+	/** Returns true if the preset contains this controller. */
+	bool HasController(URCVirtualPropertyBase* InController) const;
+
 	/** Removes a given Virtual Property (by Name) from the Remote Control preset */
 	bool RemoveController(const FName& InPropertyName);
 
diff --git a/Engine/Plugins/VirtualProduction/RemoteControl/Source/RemoteControlLogic/Private/Behaviour/RCBehaviour.cpp b/Engine/Plugins/VirtualProduction/RemoteControl/Source/RemoteControlLogic/Private/Behaviour/RCBehaviour.cpp
--- a/Engine/Plugins/VirtualProduction/RemoteControl/Source/RemoteControlLogic/Private/Behaviour/RCBehaviour.cpp
+++ b/Engine/Plugins/VirtualProduction/RemoteControl/Source/RemoteControlLogic/Private/Behaviour/RCBehaviour.cpp
@@ -75,6 +75,11 @@ int32 URCBehaviour::GetNumActions() const
 	return ActionContainer->GetActions().Num();
 }
 
+bool URCBehaviour::HasAction(URCAction* InAction) const
+{
+	return InAction && ActionContainer->GetActions().Contains(InAction);
+}
+
 bool URCBehaviour::CanHaveActionForField(const TSharedPtr<FRemoteControlField> InRemoteControlField) const
 {
 	if (ActionContainer->FindActionByFieldId(InRemoteControlField->GetId()))
diff --git a/Engine/Plugins/VirtualProduction/RemoteControl/Source/RemoteControlLogic/Private/Controller/RCController.cpp b/Engine/Plugins/VirtualProduction/RemoteControl/Source/RemoteControlLogic/Private/Controller/RCController.cpp
--- a/Engine/Plugins/VirtualProduction/RemoteControl/Source/RemoteControlLogic/Private/Controller/RCController.cpp
+++ b/Engine/Plugins/VirtualProduction/RemoteControl/Source/RemoteControlLogic/Private/Controller/RCController.cpp
@@ -100,6 +100,11 @@ int32 URCController::RemoveBehaviour(const FGuid InBehaviourId)
 	return RemovedCount;
 }
 
+bool URCController::HasBehaviour(URCBehaviour* InBehaviour) const
+{
+	return InBehaviour && Behaviours.Contains(InBehaviour);
+}
+
 void URCController::EmptyBehaviours()
 {
 	Behaviours.Empty();
diff --git a/Engine/Plugins/VirtualProduction/RemoteControl/Source/RemoteControlLogic/Public/Behaviour/RCBehaviour.h b/Engine/Plugins/VirtualProduction/RemoteControl/Source/RemoteControlLogic/Public/Behaviour/RCBehaviour.h
--- a/Engine/Plugins/VirtualProduction/RemoteControl/Source/RemoteControlLogic/Public/Behaviour/RCBehaviour.h
+++ b/Engine/Plugins/VirtualProduction/RemoteControl/Source/RemoteControlLogic/Public/Behaviour/RCBehaviour.h
@@ -57,6 +57,8 @@ public:
 	/** Get number of action associated with behaviour */
 	int32 GetNumActions() const;
 
+	bool HasAction(URCAction* InAction) const;
+
 	/** Whether we can create an action pertaining to a given remote control field for the current behaviour */
 	virtual bool CanHaveActionForField(const TSharedPtr<FRemoteControlField> InRemoteControlField) const;
 
diff --git a/Engine/Plugins/VirtualProduction/RemoteControl/Source/RemoteControlLogic/Public/Controller/RCController.h b/Engine/Plugins/VirtualProduction/RemoteControl/Source/RemoteControlLogic/Public/Controller/RCController.h
--- a/Engine/Plugins/VirtualProduction/RemoteControl/Source/RemoteControlLogic/Public/Controller/RCController.h
+++ b/Engine/Plugins/VirtualProduction/RemoteControl/Source/RemoteControlLogic/Public/Controller/RCController.h
@@ -44,6 +44,9 @@ public:
 	/** Remove the behaviour by behaviour id */
 	virtual int32 RemoveBehaviour(const FGuid InBehaviourId);
 
+	/** Returns true if the behaviour exist in this controller. */
+	bool HasBehaviour(URCBehaviour* InBehaviour) const;
+
 	/** Removes all behaviours. */
 	virtual void EmptyBehaviours();
 
diff --git a/Engine/Plugins/VirtualProduction/RemoteControl/Source/RemoteControlUI/Private/UI/Action/SRCActionPanelList.h b/Engine/Plugins/VirtualProduction/RemoteControl/Source/RemoteControlUI/Private/UI/Action/SRCActionPanelList.h
--- a/Engine/Plugins/VirtualProduction/RemoteControl/Source/RemoteControlUI/Private/UI/Action/SRCActionPanelList.h
+++ b/Engine/Plugins/VirtualProduction/RemoteControl/Source/RemoteControlUI/Private/UI/Action/SRCActionPanelList.h
@@ -24,6 +24,7 @@
 #include "UI/RCUIHelpers.h"
 #include "UI/RemoteControlPanelStyle.h"
 #include "UI/SRCPanelExposedEntitiesGroup.h"
+#include "UI/SRCPanelExposedEntitiesList.h"
 #include "UI/SRCPanelExposedField.h"
 #include "UI/SRCPanelFieldGroup.h"
 #include "UI/SRemoteControlPanel.h"
@@ -257,6 +258,7 @@ public:
 					if (ensure(ActionItem.IsValid()))
 					{
 						ActionItems.Add(ActionItem);
+						UpdateEntityUsage();
 					}
 				}
 			}
@@ -324,6 +326,8 @@ private:
 		}
 
 		ListView->RebuildList();
+
+		UpdateEntityUsage();
 	}
 	
 	/** Refreshes the list */
@@ -617,6 +621,20 @@ private:
 		}
 	}
 
+	void UpdateEntityUsage()
+	{
+		if (TSharedPtr<SRCActionPanel> ActionPanel = ActionPanelWeakPtr.Pin())
+		{
+			if (TSharedPtr<SRemoteControlPanel> RCPanel = ActionPanel->GetRemoteControlPanel())
+			{
+				if (TSharedPtr<SRCPanelExposedEntitiesList> EntitiesList = RCPanel->GetEntityList())
+				{
+					EntitiesList->UpdateEntityUsage();
+				}
+			}
+		}
+	}
+
 private:
 
 	/** The currently selected Action item */
diff --git a/Engine/Plugins/VirtualProduction/RemoteControl/Source/RemoteControlUI/Private/UI/Behaviour/SRCBehaviourPanelList.cpp b/Engine/Plugins/VirtualProduction/RemoteControl/Source/RemoteControlUI/Private/UI/Behaviour/SRCBehaviourPanelList.cpp
--- a/Engine/Plugins/VirtualProduction/RemoteControl/Source/RemoteControlUI/Private/UI/Behaviour/SRCBehaviourPanelList.cpp
+++ b/Engine/Plugins/VirtualProduction/RemoteControl/Source/RemoteControlUI/Private/UI/Behaviour/SRCBehaviourPanelList.cpp
@@ -8,16 +8,17 @@
 #include "Controller/RCController.h"
 #include "RCBehaviourModel.h"
 #include "RemoteControlPreset.h"
-#include "SRCBehaviourPanel.h"
 #include "SlateOptMacros.h"
+#include "SRCBehaviourPanel.h"
 #include "Styling/RemoteControlStyles.h"
 #include "UI/Action/SRCActionPanel.h"
 #include "UI/Behaviour/Builtin/Bind/RCBehaviourBindModel.h"
 #include "UI/Behaviour/Builtin/Conditional/RCBehaviourConditionalModel.h"
 #include "UI/Behaviour/Builtin/Path/RCBehaviourSetAssetByPathModel.h"
-#include "UI/Controller/RCControllerModel.h"
 #include "UI/Behaviour/Builtin/RangeMap/RCBehaviourRangeMapModel.h"
+#include "UI/Controller/RCControllerModel.h"
 #include "UI/RemoteControlPanelStyle.h"
+#include "UI/SRCPanelExposedEntitiesList.h"
 #include "UI/SRemoteControlPanel.h"
 #include "Widgets/Views/SHeaderRow.h"
 #include "Widgets/Views/SListView.h"
@@ -306,6 +307,8 @@ void SRCBehaviourPanelList::Reset()
 
 	ListView->RebuildList();
 	SetSelection(SelectedBehaviors);
+
+	UpdateEntityUsage();
 }
 
 TSharedRef<ITableRow> SRCBehaviourPanelList::OnGenerateWidgetForList(TSharedPtr<FRCBehaviourModel> InItem,
@@ -413,6 +416,8 @@ int32 SRCBehaviourPanelList::RemoveModel(const TSharedPtr<FRCLogicModeBase> InMo
 				// Remove Model from Data Container
 				const int32 RemoveCount = Controller->RemoveBehaviour(SelectedBehaviour->GetBehaviour());
 
+				UpdateEntityUsage();
+
 				return RemoveCount;
 			}
 		}
@@ -469,4 +474,18 @@ void SRCBehaviourPanelList::RequestRefresh()
 	ListView->RequestListRefresh();
 }
 
+void SRCBehaviourPanelList::UpdateEntityUsage()
+{
+	if (TSharedPtr<SRCBehaviourPanel> BehaviourPanel = BehaviourPanelWeakPtr.Pin())
+	{
+		if (TSharedPtr<SRemoteControlPanel> RCPanel = BehaviourPanel->GetRemoteControlPanel())
+		{
+			if (TSharedPtr<SRCPanelExposedEntitiesList> EntitiesList = RCPanel->GetEntityList())
+			{
+				EntitiesList->UpdateEntityUsage();
+			}
+		}
+	}
+}
+
 #undef LOCTEXT_NAMESPACE
\ No newline at end of file
diff --git a/Engine/Plugins/VirtualProduction/RemoteControl/Source/RemoteControlUI/Private/UI/Behaviour/SRCBehaviourPanelList.h b/Engine/Plugins/VirtualProduction/RemoteControl/Source/RemoteControlUI/Private/UI/Behaviour/SRCBehaviourPanelList.h
--- a/Engine/Plugins/VirtualProduction/RemoteControl/Source/RemoteControlUI/Private/UI/Behaviour/SRCBehaviourPanelList.h
+++ b/Engine/Plugins/VirtualProduction/RemoteControl/Source/RemoteControlUI/Private/UI/Behaviour/SRCBehaviourPanelList.h
@@ -126,6 +126,8 @@ private:
 
 	void OnBehaviourListModified();
 
+	void UpdateEntityUsage();
+
 private:
 
 	/** The parent Behaviour Panel widget */
diff --git a/Engine/Plugins/VirtualProduction/RemoteControl/Source/RemoteControlUI/Private/UI/Controller/SRCControllerPanelList.cpp b/Engine/Plugins/VirtualProduction/RemoteControl/Source/RemoteControlUI/Private/UI/Controller/SRCControllerPanelList.cpp
--- a/Engine/Plugins/VirtualProduction/RemoteControl/Source/RemoteControlUI/Private/UI/Controller/SRCControllerPanelList.cpp
+++ b/Engine/Plugins/VirtualProduction/RemoteControl/Source/RemoteControlUI/Private/UI/Controller/SRCControllerPanelList.cpp
@@ -21,17 +21,18 @@
 #include "RCVirtualPropertyContainer.h"
 #include "RemoteControlField.h"
 #include "RemoteControlPreset.h"
+#include "ScopedTransaction.h"
 #include "SDropTarget.h"
+#include "SlateOptMacros.h"
 #include "SRCControllerItemRow.h"
 #include "SRCControllerPanel.h"
-#include "ScopedTransaction.h"
-#include "SlateOptMacros.h"
 #include "Styling/RemoteControlStyles.h"
 #include "UI/Action/SRCActionPanelList.h"
 #include "UI/BaseLogicUI/RCLogicModeBase.h"
 #include "UI/RCUIHelpers.h"
 #include "UI/RemoteControlPanelStyle.h"
 #include "UI/SRCPanelDragHandle.h"
+#include "UI/SRCPanelExposedEntitiesList.h"
 #include "UI/SRCPanelExposedEntity.h"
 #include "UI/SRemoteControlPanel.h"
 #include "Widgets/Views/SHeaderRow.h"
@@ -292,6 +293,8 @@ void SRCControllerPanelList::Reset()
 			}
 		}
 	}
+
+	UpdateEntityUsage();
 }
 
 TSharedRef<ITableRow> SRCControllerPanelList::OnGenerateWidgetForList(TSharedPtr<FRCControllerModel> InItem, const TSharedRef<STableViewBase>& OwnerTable)
@@ -480,7 +483,12 @@ bool SRCControllerPanelList::IsListFocused() const
 void SRCControllerPanelList::DeleteSelectedPanelItems()
 {
 	FScopedTransaction Transaction(LOCTEXT("DeleteSelectedItems", "Delete Selected Items"));
-	if (!DeleteItemsFromLogicPanel<FRCControllerModel>(ControllerItems, ListView->GetSelectedItems()))
+
+	if (DeleteItemsFromLogicPanel<FRCControllerModel>(ControllerItems, ListView->GetSelectedItems()))
+	{
+		UpdateEntityUsage();
+	}
+	else
 	{
 		Transaction.Cancel();
 	}
@@ -873,4 +881,18 @@ void SRCControllerPanelList::SetMultiControllerMode(bool bIsUniqueModeOn)
 	Reset();
 }
 
+void SRCControllerPanelList::UpdateEntityUsage()
+{
+	if (TSharedPtr<SRCControllerPanel> ControllerPanel = ControllerPanelWeakPtr.Pin())
+	{
+		if (TSharedPtr<SRemoteControlPanel> RCPanel = ControllerPanel->GetRemoteControlPanel())
+		{
+			if (TSharedPtr<SRCPanelExposedEntitiesList> EntitiesList = RCPanel->GetEntityList())
+			{
+				EntitiesList->UpdateEntityUsage();
+			}
+		}
+	}
+}
+
 #undef LOCTEXT_NAMESPACE
diff --git a/Engine/Plugins/VirtualProduction/RemoteControl/Source/RemoteControlUI/Private/UI/Controller/SRCControllerPanelList.h b/Engine/Plugins/VirtualProduction/RemoteControl/Source/RemoteControlUI/Private/UI/Controller/SRCControllerPanelList.h
--- a/Engine/Plugins/VirtualProduction/RemoteControl/Source/RemoteControlUI/Private/UI/Controller/SRCControllerPanelList.h
+++ b/Engine/Plugins/VirtualProduction/RemoteControl/Source/RemoteControlUI/Private/UI/Controller/SRCControllerPanelList.h
@@ -241,4 +241,7 @@ public:
 
 	/** List of custom columns names */
 	TArray<FName> CustomColumns;
+
+private:
+	void UpdateEntityUsage();
 };
diff --git a/Engine/Plugins/VirtualProduction/RemoteControl/Source/RemoteControlUI/Private/UI/SRCPanelExposedEntitiesList.cpp b/Engine/Plugins/VirtualProduction/RemoteControl/Source/RemoteControlUI/Private/UI/SRCPanelExposedEntitiesList.cpp
--- a/Engine/Plugins/VirtualProduction/RemoteControl/Source/RemoteControlUI/Private/UI/SRCPanelExposedEntitiesList.cpp
+++ b/Engine/Plugins/VirtualProduction/RemoteControl/Source/RemoteControlUI/Private/UI/SRCPanelExposedEntitiesList.cpp
@@ -2,7 +2,10 @@
 
 #include "SRCPanelExposedEntitiesList.h"
 
+#include "Action/RCAction.h"
+#include "Behaviour/RCBehaviour.h"
 #include "Commands/RemoteControlCommands.h"
+#include "Controller/RCController.h"
 #include "Editor.h"
 #include "Engine/Selection.h"
 #include "Filters/SRCPanelFilter.h"
@@ -195,6 +198,7 @@ void SRCPanelExposedEntitiesList::Construct(const FArguments& InArgs, URemoteCon
 
 	bFilterApplicationRequested = false;
 	bSearchRequested = false;
+	bUpdateEntityUsage = true; // Perform this on creation
 	
 	RecallListSettings();
 
@@ -545,6 +549,12 @@ void SRCPanelExposedEntitiesList::Tick(const FGeometry& AllottedGeometry, const
 	}
 
 	ExposedEntitiesNodesRefresh();
+
+	if (bUpdateEntityUsage)
+	{
+		RunUpdateEntityUsage();
+		bUpdateEntityUsage = false;
+	}
 }
 
 TSharedPtr<SRCPanelTreeNode> SRCPanelExposedEntitiesList::GetSelectedGroup() const
@@ -879,6 +889,96 @@ void SRCPanelExposedEntitiesList::SelectSource_Execute()
 	FRemoteControlUIModule::Get().SelectSource(Entity);
 }
 
+void SRCPanelExposedEntitiesList::RunUpdateEntityUsage()
+{
+	if (!Preset.IsValid())
+	{
+		return;
+	}
+
+	TArray<UObject*> Subobjects;
+	GetObjectsWithOuter(Preset.Get(), Subobjects);
+
+	BoundEntityIds.Empty();
+	TSet<UObject*> ValidContainers;
+	TSet<UObject*> InvalidContainers;
+
+	for (UObject* Object : Subobjects)
+	{
+		URCAction* Action = Cast<URCAction>(Object);
+
+		if (!Action || !Action->ExposedFieldId.IsValid())
+		{
+			continue;
+		}
+
+		// Check to see if the action is still part of the behaviour.
+		URCBehaviour* Behaviour = Action->GetParentBehaviour();
+			
+		if (!Behaviour || InvalidContainers.Contains(Behaviour) || !Behaviour->HasAction(Action))
+		{
+			continue;
+		}
+				
+		if (!ValidContainers.Contains(Behaviour))
+		{
+			// Make sure the behaviour is still part of the controller.
+			URCController* Controller = Behaviour->ControllerWeakPtr.Get();
+
+			if (!Controller || InvalidContainers.Contains(Controller) || !Controller->HasBehaviour(Behaviour))
+			{
+				InvalidContainers.Add(Behaviour);
+				continue;
+			}
+
+			if (!ValidContainers.Contains(Controller))
+			{
+				// Make sure the controller is still part of the preset
+				URemoteControlPreset* ControllerPreset = Controller->PresetWeakPtr.Get();
+
+				if (!ControllerPreset || ControllerPreset != Preset.Get() || !ControllerPreset->HasController(Controller))
+				{
+					InvalidContainers.Add(Behaviour);
+					InvalidContainers.Add(Controller);
+					continue;
+				}
+
+				ValidContainers.Add(Behaviour);
+				ValidContainers.Add(Controller);
+			}
+		}
+
+		BoundEntityIds.Add(Action->ExposedFieldId);
+	}
+
+	for (const TPair <FGuid, TSharedPtr<SRCPanelTreeNode>>& Node : FieldWidgetMap)
+	{
+		const TSharedPtr<SRCPanelTreeNode>& NodeWidget = Node.Value;
+
+		if (NodeWidget.IsValid())
+		{
+			if (NodeWidget->GetRCType() == SRCPanelTreeNode::ENodeType::Field && NodeWidget->GetRCId().IsValid())
+			{
+				NodeWidget->SetFieldInUse(BoundEntityIds.Contains(NodeWidget->GetRCId()));
+			}
+			else
+			{
+				NodeWidget->SetFieldInUse(false);
+			}
+		}
+	}
+}
+
+void SRCPanelExposedEntitiesList::UpdateEntityUsage()
+{
+	bUpdateEntityUsage = true;
+}
+
+bool SRCPanelExposedEntitiesList::IsEntityBound(const FGuid& InEntityGuid) const
+{
+	return InEntityGuid.IsValid() && BoundEntityIds.Contains(InEntityGuid);
+}
+
 void SRCPanelExposedEntitiesList::OnObjectPropertyChange(UObject* InObject, FPropertyChangedEvent& InChangeEvent)
 {
 	static constexpr EPropertyChangeType::Type TypesNeedingRefresh = EPropertyChangeType::ArrayAdd | EPropertyChangeType::ArrayClear | EPropertyChangeType::ArrayRemove | EPropertyChangeType::ValueSet | EPropertyChangeType::Duplicate;
@@ -1153,6 +1253,15 @@ TSharedRef<ITableRow> SRCPanelExposedEntitiesList::OnGenerateRow(TSharedPtr<SRCP
 	}
 	else
 	{
+		if (Node->GetRCType() == SRCPanelTreeNode::ENodeType::Field && Node->GetRCId().IsValid())
+		{
+			Node->SetFieldInUse(BoundEntityIds.Contains(Node->GetRCId()));
+		}
+		else
+		{
+			Node->SetFieldInUse(false);
+		}
+
 		constexpr float LeftPadding = 3.f;
 		const FMargin Margin = Node->GetRCType() == SRCPanelTreeNode::FieldChild ? FMargin(LeftPadding + 10.f, 1.f, 1.f, 1.f) : FMargin(LeftPadding, 1.f, 1.f, 1.f);
 		Node->OnPropertyIdRenamed().BindSP(this, &SRCPanelExposedEntitiesList::OnPropertyIdRenamed, Node);
diff --git a/Engine/Plugins/VirtualProduction/RemoteControl/Source/RemoteControlUI/Private/UI/SRCPanelExposedEntitiesList.h b/Engine/Plugins/VirtualProduction/RemoteControl/Source/RemoteControlUI/Private/UI/SRCPanelExposedEntitiesList.h
--- a/Engine/Plugins/VirtualProduction/RemoteControl/Source/RemoteControlUI/Private/UI/SRCPanelExposedEntitiesList.h
+++ b/Engine/Plugins/VirtualProduction/RemoteControl/Source/RemoteControlUI/Private/UI/SRCPanelExposedEntitiesList.h
@@ -134,6 +134,12 @@ public:
 	/** Update the Search */
 	void UpdateSearch();
 
+	/** Update the In Use flag on entities. */
+	void UpdateEntityUsage();
+
+	/** Returns true if the entity was bound in the last binding check. */
+	bool IsEntityBound(const FGuid& InEntityGuid) const;
+
 private:
 	/** Handles label to be shown in the entity list header.  */
 	FText HandleEntityListHeaderLabel() const;
@@ -276,6 +282,9 @@ private:
 	/** Selects the source of the currently selected entity. */
 	void SelectSource_Execute();
 
+	/** Updates the usage of entities in this list. */
+	void RunUpdateEntityUsage();
+
 private:
 	/** Command list to use if specified */
 	TSharedPtr<FUICommandList> CommandList;
@@ -329,6 +338,10 @@ private:
 	bool bFilterApplicationRequested;
 	/** If true, the entity items will be refreshed next frame. */
 	bool bSearchRequested;
+	/** If true, the entity list will have its in use flag updated next frame. */
+	bool bUpdateEntityUsage;
+	/** Temporary list of the currently bound entity ids. */
+	TSet<FGuid> BoundEntityIds;
 	/** Panel Style reference. */
 	const FRCPanelStyle* RCPanelStyle;
 	/** Holds the header row of entities list. */
diff --git a/Engine/Plugins/VirtualProduction/RemoteControl/Source/RemoteControlUI/Private/UI/SRCPanelExposedField.cpp b/Engine/Plugins/VirtualProduction/RemoteControl/Source/RemoteControlUI/Private/UI/SRCPanelExposedField.cpp
--- a/Engine/Plugins/VirtualProduction/RemoteControl/Source/RemoteControlUI/Private/UI/SRCPanelExposedField.cpp
+++ b/Engine/Plugins/VirtualProduction/RemoteControl/Source/RemoteControlUI/Private/UI/SRCPanelExposedField.cpp
@@ -226,6 +226,7 @@ SRCPanelTreeNode::FMakeNodeWidgetArgs SRCPanelExposedField::CreateEntityWidgetIn
 				.ClearKeyboardFocusOnCommit(true)
 				.Text_Lambda([this] () { return GetPropertyIdText(); })
 				.OnTextCommitted(this, &SRCPanelExposedField::OnPropertyIdTextCommitted)
+				.ForegroundColor(this, &SRCPanelExposedField::GetIdWidgetTextColor)
 			];
 	}
 
@@ -913,6 +914,16 @@ void SRCPanelExposedField::OnPropertyIdTextCommitted(const FText& InText, ETextC
 	}
 }
 
+FSlateColor SRCPanelExposedField::GetIdWidgetTextColor() const
+{
+	if (bIsFieldInUse)
+	{
+		return FSlateColor::UseForeground();
+	}
+
+	return FSlateColor::UseSubduedForeground();
+}
+
 void SRCPanelFieldChildNode::Construct(const FArguments& InArgs, const TSharedRef<IDetailTreeNode>& InNode, FRCColumnSizeData InColumnSizeData)
 {
 	TArray<TSharedRef<IDetailTreeNode>> ChildNodes;
diff --git a/Engine/Plugins/VirtualProduction/RemoteControl/Source/RemoteControlUI/Private/UI/SRCPanelExposedField.h b/Engine/Plugins/VirtualProduction/RemoteControl/Source/RemoteControlUI/Private/UI/SRCPanelExposedField.h
--- a/Engine/Plugins/VirtualProduction/RemoteControl/Source/RemoteControlUI/Private/UI/SRCPanelExposedField.h
+++ b/Engine/Plugins/VirtualProduction/RemoteControl/Source/RemoteControlUI/Private/UI/SRCPanelExposedField.h
@@ -110,6 +110,9 @@ private:
 	/** Handles committing a property id label. */
 	void OnPropertyIdTextCommitted(const FText& InText, ETextCommit::Type InCommitInfo);
 
+	/** Returns subdued foreground if unused, foreground if used in an action. */
+	FSlateColor GetIdWidgetTextColor() const;
+
 private:
 	/** Weak pointer to the underlying RC Field. */
 	TWeakPtr<FRemoteControlField> WeakField;
diff --git a/Engine/Plugins/VirtualProduction/RemoteControl/Source/RemoteControlUI/Private/UI/SRCPanelTreeNode.h b/Engine/Plugins/VirtualProduction/RemoteControl/Source/RemoteControlUI/Private/UI/SRCPanelTreeNode.h
--- a/Engine/Plugins/VirtualProduction/RemoteControl/Source/RemoteControlUI/Private/UI/SRCPanelTreeNode.h
+++ b/Engine/Plugins/VirtualProduction/RemoteControl/Source/RemoteControlUI/Private/UI/SRCPanelTreeNode.h
@@ -117,6 +117,8 @@ struct SRCPanelTreeNode : public SCompoundWidget, public IHasProtocolExtensibili
 	/** Executed when the PropertyId of the node is changed. */
 	FOnPropertyIdRenamed& OnPropertyIdRenamed() { return OnPropertyIdRenamedDelegate; }
 
+	void SetFieldInUse(bool bInFieldInUse) { bIsFieldInUse = bInFieldInUse; }
+
 protected:
 	struct FMakeNodeWidgetArgs
 	{
@@ -155,6 +157,9 @@ protected:
 	/** Holds the row's columns' width. */
 	FRCColumnSizeData ColumnSizeData;
 
+	/** If true, an action is bound to this node's entity. */
+	bool bIsFieldInUse = false;
+
 private:
 
 	/** Cached widget of PropertyId identifier. */
