Remote Control: Changed "in use" visualisation for properties to a background color similar to selection.

diff --git a/Engine/Plugins/VirtualProduction/RemoteControl/Source/RemoteControlUI/Private/UI/RemoteControlPanelStyle.cpp b/Engine/Plugins/VirtualProduction/RemoteControl/Source/RemoteControlUI/Private/UI/RemoteControlPanelStyle.cpp
--- a/Engine/Plugins/VirtualProduction/RemoteControl/Source/RemoteControlUI/Private/UI/RemoteControlPanelStyle.cpp
+++ b/Engine/Plugins/VirtualProduction/RemoteControl/Source/RemoteControlUI/Private/UI/RemoteControlPanelStyle.cpp
@@ -90,6 +90,8 @@ void FRemoteControlPanelStyle::Initialize()
 
 	StyleSet->Set("RemoteControlPanel.GroupRowSelected", new BOX_BRUSH("Common/GroupBorderLight", FMargin(4.0f / 16.0f), NewSelectionColor));
 	StyleSet->Set("RemoteControlPanel.TransparentBorder", new FSlateColorBrush(FColor::Transparent));
+
+	StyleSet->Set("RemoteControlPanel.EntityInUseBorder", new FSlateColorBrush(FStyleColors::SelectInactive.GetSpecifiedColor()));
 	
 	// Toggle Button Style
 	StyleSet->Set("Switch.ToggleOff", new IMAGE_PLUGIN_BRUSH("Icons/Switch_OFF", Icon28x14));
diff --git a/Engine/Plugins/VirtualProduction/RemoteControl/Source/RemoteControlUI/Private/UI/SRCPanelExposedEntitiesList.cpp b/Engine/Plugins/VirtualProduction/RemoteControl/Source/RemoteControlUI/Private/UI/SRCPanelExposedEntitiesList.cpp
--- a/Engine/Plugins/VirtualProduction/RemoteControl/Source/RemoteControlUI/Private/UI/SRCPanelExposedEntitiesList.cpp
+++ b/Engine/Plugins/VirtualProduction/RemoteControl/Source/RemoteControlUI/Private/UI/SRCPanelExposedEntitiesList.cpp
@@ -97,6 +97,16 @@ public:
 		FSuperRowType::Construct(SuperArgs, OwnerTableView);
 	}
 
+	virtual const FSlateBrush* GetBorder() const override
+	{
+		if (Entity.IsValid() && Entity->IsFieldInUse() && !IsHovered())
+		{
+			return FRemoteControlPanelStyle::Get()->GetBrush("RemoteControlPanel.EntityInUseBorder");
+		}
+
+		return SMultiColumnTableRow<TSharedPtr<SRCPanelTreeNode>>::GetBorder();
+	}
+
 	TSharedRef<SWidget> GenerateWidgetForColumn(const FName& InColumnName) override
 	{
 		if (Entity.IsValid())
diff --git a/Engine/Plugins/VirtualProduction/RemoteControl/Source/RemoteControlUI/Private/UI/SRCPanelExposedField.cpp b/Engine/Plugins/VirtualProduction/RemoteControl/Source/RemoteControlUI/Private/UI/SRCPanelExposedField.cpp
--- a/Engine/Plugins/VirtualProduction/RemoteControl/Source/RemoteControlUI/Private/UI/SRCPanelExposedField.cpp
+++ b/Engine/Plugins/VirtualProduction/RemoteControl/Source/RemoteControlUI/Private/UI/SRCPanelExposedField.cpp
@@ -226,7 +226,6 @@ SRCPanelTreeNode::FMakeNodeWidgetArgs SRCPanelExposedField::CreateEntityWidgetIn
 				.ClearKeyboardFocusOnCommit(true)
 				.Text_Lambda([this] () { return GetPropertyIdText(); })
 				.OnTextCommitted(this, &SRCPanelExposedField::OnPropertyIdTextCommitted)
-				.ForegroundColor(this, &SRCPanelExposedField::GetIdWidgetTextColor)
 			];
 	}
 
@@ -914,16 +913,6 @@ void SRCPanelExposedField::OnPropertyIdTextCommitted(const FText& InText, ETextC
 	}
 }
 
-FSlateColor SRCPanelExposedField::GetIdWidgetTextColor() const
-{
-	if (bIsFieldInUse)
-	{
-		return FSlateColor::UseForeground();
-	}
-
-	return FSlateColor::UseSubduedForeground();
-}
-
 void SRCPanelFieldChildNode::Construct(const FArguments& InArgs, const TSharedRef<IDetailTreeNode>& InNode, FRCColumnSizeData InColumnSizeData)
 {
 	TArray<TSharedRef<IDetailTreeNode>> ChildNodes;
diff --git a/Engine/Plugins/VirtualProduction/RemoteControl/Source/RemoteControlUI/Private/UI/SRCPanelExposedField.h b/Engine/Plugins/VirtualProduction/RemoteControl/Source/RemoteControlUI/Private/UI/SRCPanelExposedField.h
--- a/Engine/Plugins/VirtualProduction/RemoteControl/Source/RemoteControlUI/Private/UI/SRCPanelExposedField.h
+++ b/Engine/Plugins/VirtualProduction/RemoteControl/Source/RemoteControlUI/Private/UI/SRCPanelExposedField.h
@@ -110,9 +110,6 @@ private:
 	/** Handles committing a property id label. */
 	void OnPropertyIdTextCommitted(const FText& InText, ETextCommit::Type InCommitInfo);
 
-	/** Returns subdued foreground if unused, foreground if used in an action. */
-	FSlateColor GetIdWidgetTextColor() const;
-
 private:
 	/** Weak pointer to the underlying RC Field. */
 	TWeakPtr<FRemoteControlField> WeakField;
diff --git a/Engine/Plugins/VirtualProduction/RemoteControl/Source/RemoteControlUI/Private/UI/SRCPanelTreeNode.h b/Engine/Plugins/VirtualProduction/RemoteControl/Source/RemoteControlUI/Private/UI/SRCPanelTreeNode.h
--- a/Engine/Plugins/VirtualProduction/RemoteControl/Source/RemoteControlUI/Private/UI/SRCPanelTreeNode.h
+++ b/Engine/Plugins/VirtualProduction/RemoteControl/Source/RemoteControlUI/Private/UI/SRCPanelTreeNode.h
@@ -117,7 +117,15 @@ struct SRCPanelTreeNode : public SCompoundWidget, public IHasProtocolExtensibili
 	/** Executed when the PropertyId of the node is changed. */
 	FOnPropertyIdRenamed& OnPropertyIdRenamed() { return OnPropertyIdRenamedDelegate; }
 
-	void SetFieldInUse(bool bInFieldInUse) { bIsFieldInUse = bInFieldInUse; }
+	bool IsFieldInUse() const 
+	{ 
+		return bIsFieldInUse; 
+	}
+
+	void SetFieldInUse(bool bInFieldInUse) 
+	{ 
+		bIsFieldInUse = bInFieldInUse; 
+	}
 
 protected:
 	struct FMakeNodeWidgetArgs
