Fix use of an expiring temporary in ProtocolPanelStyle.cpp

diff --git a/Engine/Plugins/VirtualProduction/RemoteControl/Source/RemoteControlProtocolWidgets/Private/Styling/ProtocolPanelStyle.cpp b/Engine/Plugins/VirtualProduction/RemoteControl/Source/RemoteControlProtocolWidgets/Private/Styling/ProtocolPanelStyle.cpp
--- a/Engine/Plugins/VirtualProduction/RemoteControl/Source/RemoteControlProtocolWidgets/Private/Styling/ProtocolPanelStyle.cpp
+++ b/Engine/Plugins/VirtualProduction/RemoteControl/Source/RemoteControlProtocolWidgets/Private/Styling/ProtocolPanelStyle.cpp
@@ -78,8 +78,7 @@ void FProtocolPanelStyle::SetupWidgetStyles(TSharedRef<FSlateStyleSet> InStyle)
 
 	const FTextBlockStyle& PlainTextStyle = AppStyle.GetWidgetStyle<FTextBlockStyle>("NormalText");
 
-	const FTextBlockStyle& BoldTextStyle = FTextBlockStyle(PlainTextStyle)
-		.SetFont(StyleFonts.NormalBold);
+	const FTextBlockStyle BoldTextStyle = FTextBlockStyle(PlainTextStyle).SetFont(StyleFonts.NormalBold);
 
 	FProtocolWidgetStyle ProtocolMaskWidgetStyle = FProtocolWidgetStyle()
 		.SetContentAreaBrush(FSlateColorBrush(FStyleColors::Panel))
