Remote Control: String List Controller

diff --git a/Engine/Plugins/VirtualProduction/RemoteControl/Source/RemoteControl/Private/RCVirtualPropertyContainer.cpp b/Engine/Plugins/VirtualProduction/RemoteControl/Source/RemoteControl/Private/RCVirtualPropertyContainer.cpp
--- a/Engine/Plugins/VirtualProduction/RemoteControl/Source/RemoteControl/Private/RCVirtualPropertyContainer.cpp
+++ b/Engine/Plugins/VirtualProduction/RemoteControl/Source/RemoteControl/Private/RCVirtualPropertyContainer.cpp
@@ -93,6 +93,16 @@ URCVirtualPropertyInContainer* URCVirtualPropertyContainerBase::AddProperty(cons
 	PropertyBagDesc.MetaData = MetaData;
 #endif
 
+	// Create Property in Container
+	URCVirtualPropertyInContainer* VirtualPropertyInContainer = NewObject<URCVirtualPropertyInContainer>(this, InPropertyClass.Get(), NAME_None, RF_Transactional);
+	VirtualPropertyInContainer->PropertyName = PropertyName;
+	VirtualPropertyInContainer->DisplayName = PropertyName;
+	VirtualPropertyInContainer->PresetWeakPtr = PresetWeakPtr;
+	VirtualPropertyInContainer->ContainerWeakPtr = this;
+	VirtualPropertyInContainer->Id = FGuid::NewGuid();
+
+	VirtualPropertyInContainer->InitPropertyBagPropertyDescription(PropertyBagDesc);
+
 	Bag.AddProperties({ PropertyBagDesc });
 
 	// Ensure that the property has been successfully added to the Bag
@@ -102,14 +112,6 @@ URCVirtualPropertyInContainer* URCVirtualPropertyContainerBase::AddProperty(cons
 		return nullptr;
 	}
 	
-	// Create Property in Container
-	URCVirtualPropertyInContainer* VirtualPropertyInContainer = NewObject<URCVirtualPropertyInContainer>(this, InPropertyClass.Get(), NAME_None, RF_Transactional);
-	VirtualPropertyInContainer->PropertyName = PropertyName;
-	VirtualPropertyInContainer->DisplayName = PropertyName;
-	VirtualPropertyInContainer->PresetWeakPtr = PresetWeakPtr;
-	VirtualPropertyInContainer->ContainerWeakPtr = this;
-	VirtualPropertyInContainer->Id = FGuid::NewGuid();
-
 	ControllerLabelToIdCache.Add(PropertyName, VirtualPropertyInContainer->Id);
 	AddVirtualProperty(VirtualPropertyInContainer);
 
diff --git a/Engine/Plugins/VirtualProduction/RemoteControl/Source/RemoteControl/Public/RCVirtualProperty.h b/Engine/Plugins/VirtualProduction/RemoteControl/Source/RemoteControl/Public/RCVirtualProperty.h
--- a/Engine/Plugins/VirtualProduction/RemoteControl/Source/RemoteControl/Public/RCVirtualProperty.h
+++ b/Engine/Plugins/VirtualProduction/RemoteControl/Source/RemoteControl/Public/RCVirtualProperty.h
@@ -54,6 +54,8 @@ public:
 	/** Initialization routine. Called after the parent container has setup data for this property */
 	virtual void Init() {}
 
+	virtual void InitPropertyBagPropertyDescription(FPropertyBagPropertyDesc& InOutPropertyDesc) {}
+
 	/**
 	 * @brief Called internally when entity Ids are renewed.
 	 * @param InEntityIdMap Map of old Id to new Id.
diff --git a/Engine/Plugins/VirtualProduction/RemoteControl/Source/RemoteControlLogic/Public/Controller/RCCustomControllerUtilities.h b/Engine/Plugins/VirtualProduction/RemoteControl/Source/RemoteControlLogic/Public/Controller/RCCustomControllerUtilities.h
--- a/Engine/Plugins/VirtualProduction/RemoteControl/Source/RemoteControlLogic/Public/Controller/RCCustomControllerUtilities.h
+++ b/Engine/Plugins/VirtualProduction/RemoteControl/Source/RemoteControlLogic/Public/Controller/RCCustomControllerUtilities.h
@@ -1,4 +1,4 @@
-﻿// Copyright Epic Games, Inc. All Rights Reserved.
+// Copyright Epic Games, Inc. All Rights Reserved.
 
 #pragma once
 #include "StructUtils/PropertyBag.h"
@@ -41,4 +41,7 @@ namespace UE::RCCustomControllers
 	inline static FName CustomControllerExecuteOnLoadNameKey = "CustomControllerExecuteOnLoad";
 
 	inline static FString CustomTextureControllerName = TEXT("Texture");
+	inline static FString CustomStringListControllerName = TEXT("StringList");
+
+	inline static FName StringListControllerEntriesName = "StringOptions";
 };
diff --git a/Engine/Plugins/VirtualProduction/RemoteControl/Source/RemoteControlUI/Private/RemoteControlUIModule.cpp b/Engine/Plugins/VirtualProduction/RemoteControl/Source/RemoteControlUI/Private/RemoteControlUIModule.cpp
--- a/Engine/Plugins/VirtualProduction/RemoteControl/Source/RemoteControlUI/Private/RemoteControlUIModule.cpp
+++ b/Engine/Plugins/VirtualProduction/RemoteControl/Source/RemoteControlUI/Private/RemoteControlUIModule.cpp
@@ -39,6 +39,7 @@
 #include "UI/Behaviour/Builtin/Conditional/SRCBehaviourConditional.h"
 #include "UI/Behaviour/SRCBehaviourPanel.h"
 #include "UI/Controller/CustomControllers/SCustomTextureControllerWidget.h"
+#include "UI/Controller/CustomControllers/SCustomStringListControllerWidget.h"
 #include "UI/Customizations/FPassphraseCustomization.h"
 #include "UI/Customizations/NetworkAddressCustomization.h"
 #include "UI/Customizations/RCAssetPathElementCustomization.h"
@@ -1215,7 +1216,11 @@ TSharedPtr<SWidget> FRemoteControlUIModule::CreateCustomControllerWidget(URCVirt
 	const FString& CustomName = UE::RCCustomControllers::GetCustomControllerTypeName(InController);
 	if (CustomName == UE::RCCustomControllers::CustomTextureControllerName)
 	{
-		CustomControllerWidget = SNew(SCustomTextureControllerWidget, InOriginalPropertyHandle);
+		CustomControllerWidget = SNew(SCustomTextureControllerWidget, InController, InOriginalPropertyHandle);
+	}
+	else if (CustomName == UE::RCCustomControllers::CustomStringListControllerName)
+	{
+		CustomControllerWidget = SNew(SCustomStringListControllerWidget, InController, InOriginalPropertyHandle);
 	}
 
 	return CustomControllerWidget;
diff --git a/Engine/Plugins/VirtualProduction/RemoteControl/Source/RemoteControlUI/Private/UI/Behaviour/Builtin/Conditional/RCBehaviourConditionalModel.cpp b/Engine/Plugins/VirtualProduction/RemoteControl/Source/RemoteControlUI/Private/UI/Behaviour/Builtin/Conditional/RCBehaviourConditionalModel.cpp
--- a/Engine/Plugins/VirtualProduction/RemoteControl/Source/RemoteControlUI/Private/UI/Behaviour/Builtin/Conditional/RCBehaviourConditionalModel.cpp
+++ b/Engine/Plugins/VirtualProduction/RemoteControl/Source/RemoteControlUI/Private/UI/Behaviour/Builtin/Conditional/RCBehaviourConditionalModel.cpp
@@ -105,6 +105,20 @@ TSharedRef<SWidget> FRCBehaviourConditionalModel::GetComparandFieldWidget() cons
 {
 	if (const TSharedPtr<IDetailTreeNode> DetailTreeNode = DetailTreeNodeWeakPtr.Pin())
 	{
+		if (TSharedPtr<IPropertyHandle> PropertyHandle = DetailTreeNode->CreatePropertyHandle())
+		{
+			if (URCBehaviour* Behaviour = GetBehaviour())
+			{
+				if (URCVirtualPropertyBase* Controller = Behaviour->ControllerWeakPtr.Get())
+				{
+					if (const TSharedPtr<SWidget>& CustomControllerWidget = IRemoteControlUIModule::Get().CreateCustomControllerWidget(Controller, PropertyHandle.ToSharedRef()))
+					{
+						return CustomControllerWidget.ToSharedRef();
+					}
+				}
+			}
+		}
+
 		return UE::RCUIHelpers::GetGenericFieldWidget(DetailTreeNode);
 	}
 
diff --git a/Engine/Plugins/VirtualProduction/RemoteControl/Source/RemoteControlUI/Private/UI/Controller/CustomControllers/SCustomStringListControllerListEditorWidget.cpp b/Engine/Plugins/VirtualProduction/RemoteControl/Source/RemoteControlUI/Private/UI/Controller/CustomControllers/SCustomStringListControllerListEditorWidget.cpp
new file mode 100644
--- /dev/null
+++ b/Engine/Plugins/VirtualProduction/RemoteControl/Source/RemoteControlUI/Private/UI/Controller/CustomControllers/SCustomStringListControllerListEditorWidget.cpp
@@ -0,0 +1,542 @@
+// Copyright Epic Games, Inc. All Rights Reserved.
+
+#include "UI/Controller/CustomControllers/SCustomStringListControllerListEditorWidget.h"
+
+#include "Containers/UnrealString.h"
+#include "DragAndDrop/DecoratedDragDropOp.h"
+#include "Framework/Application/SlateApplication.h"
+#include "Internationalization/Text.h"
+#include "Misc/CString.h"
+#include "Styles/SlateBrushTemplates.h"
+#include "Styling/AppStyle.h"
+#include "UI/Controller/CustomControllers/SCustomStringListControllerWidget.h"
+#include "Widgets/Images/SImage.h"
+#include "Widgets/Input/SButton.h"
+#include "Widgets/Input/SEditableTextBox.h"
+#include "Widgets/SBoxPanel.h"
+#include "Widgets/SNullWidget.h"
+#include "Widgets/SWindow.h"
+#include "Widgets/Text/STextBlock.h"
+
+#define LOCTEXT_NAMESPACE "SCustomStringListControllerListEditorWidget"
+
+class SCustomStringListControllerEnumEditorRowEditBox : public SCompoundWidget
+{
+	SLATE_BEGIN_ARGS(SCustomStringListControllerEnumEditorRowEditBox)
+		{}
+	SLATE_END_ARGS()
+
+	void Construct(const FArguments& InArgs, TSharedRef<SCustomStringListControllerListEditorWidget> InParentWidget, FName InEntry, bool bInIsInTable)
+	{
+		ParentWidgetWeak = InParentWidget;
+		Entry = InEntry;
+		bIsInTable = bInIsInTable;
+
+		TSharedRef<SHorizontalBox> EditBox = SNew(SHorizontalBox)
+			+ SHorizontalBox::Slot()
+			.MinWidth(200.f)
+			.MaxWidth(200.f)
+			.Padding(5.f)
+			[
+				SNew(SEditableTextBox)
+				.Padding(5.f)
+				.Text(this, &SCustomStringListControllerEnumEditorRowEditBox::GetEntry)
+				.HintText(bInIsInTable 
+					? LOCTEXT("TableHint", "Edit Entry")
+					: LOCTEXT("AddHint", "Add Entry")
+				)
+				.Justification(ETextJustify::Left)
+				.RevertTextOnEscape(true)
+				.SelectAllTextWhenFocused(true)
+				.OnVerifyTextChanged(this, &SCustomStringListControllerEnumEditorRowEditBox::OnEntryVerify)
+				.OnTextCommitted(this, &SCustomStringListControllerEnumEditorRowEditBox::OnEntryCommitted)
+			];
+
+		if (bIsInTable)
+		{
+			EditBox->AddSlot()
+				.AutoWidth()
+				.Padding(FMargin(0.f, 0.f, 5.f, 0.f))
+				.VAlign(VAlign_Center)
+				[
+					SNew(SButton)
+					.ButtonStyle(FAppStyle::Get(), "HoverHintOnly")
+					.ContentPadding(5.f)
+					.OnClicked(this, &SCustomStringListControllerEnumEditorRowEditBox::OnRemoveClicked)
+					.IsEnabled(this, &SCustomStringListControllerEnumEditorRowEditBox::CanRemove)
+					[
+						SNew(SImage)
+						.Image(FAppStyle::Get().GetBrush("Icons.Delete"))
+						.DesiredSizeOverride(FVector2D(16.f, 16.f))
+					]
+				];
+
+			EditBox->AddSlot()
+				.AutoWidth()
+				.Padding(FMargin(0.f, 0.f, 10.f, 0.f))
+				.VAlign(VAlign_Center)
+				[
+					SNew(SImage)
+					.Image(FSlateBrushTemplates::DragHandle())
+					.DesiredSizeOverride(FVector2D(5.f, 15.f))
+				];
+		}
+		else
+		{
+			EditBox->AddSlot()
+				.AutoWidth()
+				.Padding(FMargin(5.f, 0.f, 0.f, 0.f))
+				.VAlign(VAlign_Center)
+				[
+					SNew(SImage)
+					.Image(FAppStyle::Get().GetBrush("Icons.PlusCircle"))
+					.DesiredSizeOverride(FVector2D(16.f, 16.f))
+				];
+		}
+
+		ChildSlot
+		[
+			EditBox
+		];
+	}
+
+private:
+	TWeakPtr<SCustomStringListControllerListEditorWidget> ParentWidgetWeak;
+	FName Entry;
+	bool bIsInTable;
+
+	FText GetEntry() const
+	{
+		if (Entry.IsNone())
+		{
+			return FText::GetEmpty();
+		}
+
+		return FText::FromName(Entry);
+	}
+
+	bool OnEntryVerify(const FText& InNewValue, FText& OutError)
+	{
+		TSharedPtr<SCustomStringListControllerListEditorWidget> ParentWidget = ParentWidgetWeak.Pin();
+
+		if (!ParentWidget.IsValid())
+		{
+			OutError = LOCTEXT("InternalError", "Internal Error");
+			return false;
+		}
+
+		if (InNewValue.IsEmpty())
+		{
+			if (bIsInTable)
+			{
+				OutError = LOCTEXT("Missing", "Missing Entry");
+				return false;
+			}
+
+			return true;
+		}
+
+		const FString NewValueString = InNewValue.ToString();
+
+		if (NewValueString.Contains(TEXT(";")) || NewValueString.Contains(TEXT("\"")))
+		{
+			OutError = LOCTEXT("Delimiter", "Entries cannot contain ';' or '\"'.");
+			return false;
+		}
+		
+		const FName New = *NewValueString;
+
+		if (Entry == New)
+		{
+			return true;
+		}
+
+		if (ParentWidget->HasEntry(New))
+		{
+			OutError = LOCTEXT("Duplicate", "Duplicate Entry");
+			return false;
+		}
+
+		return true;
+	}
+
+	void OnEntryCommitted(const FText& InNew, ETextCommit::Type InCommitType)
+	{
+		TSharedPtr<SCustomStringListControllerListEditorWidget> ParentWidget = ParentWidgetWeak.Pin();
+
+		if (!ParentWidget.IsValid())
+		{
+			return;
+		}
+
+		if (InNew.IsEmpty())
+		{
+			return;
+		}
+
+		const FName NewEntry = *InNew.ToString().TrimStartAndEnd();
+
+		if (Entry == NewEntry)
+		{
+			return;
+		}
+
+		if (ParentWidget->HasEntry(NewEntry))
+		{
+			return;
+		}
+
+		if (bIsInTable)
+		{
+			ParentWidget->ReplaceEntry(Entry, NewEntry);
+			Entry = NewEntry;
+		}
+		else
+		{
+			ParentWidget->AddEntry(NewEntry);
+		}
+	}
+
+	FReply OnRemoveClicked()
+	{
+		if (TSharedPtr<SCustomStringListControllerListEditorWidget> ParentWidget = ParentWidgetWeak.Pin())
+		{
+			ParentWidget->RemoveEntry(Entry);
+		}
+
+		return FReply::Handled();
+	}
+
+	bool CanRemove() const
+	{
+		if (TSharedPtr<SCustomStringListControllerListEditorWidget> ParentWidget = ParentWidgetWeak.Pin())
+		{
+			return ParentWidget->GetEntries().Num() > 1;
+		}
+
+		return false;
+	}
+};
+
+class FCustomStringListControllerEnumEditorRowDragDropOp : public FDecoratedDragDropOp
+{
+public:
+	DRAG_DROP_OPERATOR_TYPE(FCustomStringListControllerEnumEditorRowDragDropOp, FDecoratedDragDropOp)
+
+	FCustomStringListControllerEnumEditorRowDragDropOp(FName InEntry)
+		: Entry(InEntry)
+	{}
+
+	virtual ~FCustomStringListControllerEnumEditorRowDragDropOp() = default;
+
+	virtual TSharedPtr<SWidget> GetDefaultDecorator() const override
+	{
+		return SNew(STextBlock)
+			.Text(FText::FromName(Entry));
+	}
+
+	FName GetEntry() const
+	{
+		return Entry;
+	}
+
+private:
+	FName Entry;
+};
+
+class SCustomStringListControllerEnumEditorRow : public STableRow<FName>
+{
+public:
+	void Construct(const FArguments& InArgs, const TSharedRef<STableViewBase>& InOwnerTableView,
+		TSharedRef<SCustomStringListControllerListEditorWidget> InParentWidget, FName InEntry)
+	{
+		ParentWidgetWeak = InParentWidget;
+		Entry = InEntry;
+		
+		STableRow<FName>::Construct(
+			STableRow<FName>::FArguments()
+			.OnDragDetected(this, &SCustomStringListControllerEnumEditorRow::OnDragDetected)
+			.OnCanAcceptDrop(this, &SCustomStringListControllerEnumEditorRow::OnCanAcceptDrop)
+			.OnAcceptDrop(this, &SCustomStringListControllerEnumEditorRow::OnAcceptDrop)
+			.Content()
+			[
+				SNew(SCustomStringListControllerEnumEditorRowEditBox, InParentWidget, InEntry, /* In Table */ true)
+			],
+			InOwnerTableView);
+	}
+
+private:
+	TWeakPtr<SCustomStringListControllerListEditorWidget> ParentWidgetWeak;
+	FName Entry;
+
+	FReply OnDragDetected(const FGeometry& InGeometry, const FPointerEvent& InPointerEvent)
+	{
+		if (!Entry.IsNone())
+		{
+			TSharedRef<FCustomStringListControllerEnumEditorRowDragDropOp> DragDropOp = MakeShared<FCustomStringListControllerEnumEditorRowDragDropOp>(Entry);
+			return FReply::Handled().BeginDragDrop(DragDropOp);
+		}
+
+		return FReply::Handled();
+	}
+
+	TOptional<EItemDropZone> OnCanAcceptDrop(const FDragDropEvent& InDragDropEvent, EItemDropZone InDropZone, FName InEntry)
+	{
+		const TSharedPtr<FCustomStringListControllerEnumEditorRowDragDropOp> DragDropOp = InDragDropEvent.GetOperationAs<FCustomStringListControllerEnumEditorRowDragDropOp>();
+
+		if (!DragDropOp.IsValid())
+		{
+			return {};
+		}
+
+		if (Entry.IsNone())
+		{
+			return {};
+		}
+
+		return InDropZone;
+	}
+
+	FReply OnAcceptDrop(const FDragDropEvent& InDragDropEvent, EItemDropZone InDropZone, FName InEntry)
+	{
+		const TSharedPtr<FCustomStringListControllerEnumEditorRowDragDropOp> DragDropOp = InDragDropEvent.GetOperationAs<FCustomStringListControllerEnumEditorRowDragDropOp>();
+
+		if (!DragDropOp.IsValid())
+		{
+			return FReply::Handled();
+		}
+
+		TSharedPtr<SCustomStringListControllerListEditorWidget> ParentWidget = ParentWidgetWeak.Pin();
+
+		if (!ParentWidget.IsValid())
+		{
+			return FReply::Handled();
+		}
+
+		ParentWidget->OnValidDrop(InEntry, DragDropOp->GetEntry(), InDropZone);
+
+		return FReply::Handled();
+	}
+};
+
+void SCustomStringListControllerListEditorWidget::OpenModalWindow(const TSharedRef<SCustomStringListControllerWidget>& InParentWidget, const TArray<FName>& InEntries)
+{
+	TSharedRef<SCustomStringListControllerListEditorWidget> EditorWidget = SNew(SCustomStringListControllerListEditorWidget, InParentWidget, InEntries);
+
+	TSharedRef<SWindow> Container = SNew(SWindow)
+		.SizingRule(ESizingRule::Autosized)
+		.SupportsMinimize(false)
+		.SupportsMaximize(false)
+		.Title(LOCTEXT("ModalWindowTitle", "Edit String List"))
+		[
+			EditorWidget
+		];
+		
+	FSlateApplication::Get().AddModalWindow(Container, InParentWidget);
+}
+
+void SCustomStringListControllerListEditorWidget::Construct(const FArguments& InArgs, const TSharedRef<SCustomStringListControllerWidget>& InParentWidget, const TArray<FName>& InEntries)
+{
+	ParentWidgetWeak = InParentWidget;
+	Entries = InEntries;
+
+	ListView = SNew(SListView<FName>)
+		.ListItemsSource(&Entries)
+		.OnGenerateRow(this, &SCustomStringListControllerListEditorWidget::OnGenerateRow)
+		.SelectionMode(ESelectionMode::Single);
+
+	TSharedPtr<SWidget> NewRowWidget;
+
+	ChildSlot
+	[
+		SNew(SVerticalBox)
+		+ SVerticalBox::Slot()
+		.Padding(10.f, 10.f, 10.f, 10.f)
+		.HAlign(HAlign_Fill)
+		.AutoHeight()
+		[
+			ListView.ToSharedRef()
+		]
+		+ SVerticalBox::Slot()
+		.Padding(10.f, 0.f, 10.f, 10.f)
+		.HAlign(HAlign_Fill)
+		.AutoHeight()
+		[
+			SAssignNew(NewRowWidget, SCustomStringListControllerEnumEditorRowEditBox, SharedThis(this), NAME_None, /* In Table */ false)
+		]
+		+ SVerticalBox::Slot()
+		.Padding(10.f, 0.f, 10.f, 10.f)
+		.AutoHeight()
+		.HAlign(HAlign_Center)
+		[
+			SNew(SHorizontalBox)
+			+ SHorizontalBox::Slot()
+			.AutoWidth()
+			.Padding(1.f, 0.f, 10.f, 0.f)
+			[
+				SNew(SButton)
+				.OnClicked(this, &SCustomStringListControllerListEditorWidget::OnOkClicked)
+				.ToolTipText(LOCTEXT("OkayToolTip", "Close this window and use these values for the enum."))
+				[
+					SNew(STextBlock)
+					.Text(LOCTEXT("AcceptValues", "Ok"))
+				]
+			]
+			+ SHorizontalBox::Slot()
+			.AutoWidth()
+			.Padding(10.f, 0.f, 0.f, 0.f)
+			[
+				SNew(SButton)
+				.OnClicked(this, &SCustomStringListControllerListEditorWidget::OnCancelClicked)
+				.ToolTipText(LOCTEXT("CancelToolTip", "Renumber the entries sequentially from 0."))
+				[
+					SNew(STextBlock)
+					.Text(LOCTEXT("RejectValues", "Cancel"))
+				]
+			]
+		]
+	];
+}
+
+bool SCustomStringListControllerListEditorWidget::HasEntry(FName InEntry) const
+{
+	for (const FName& Entry : Entries)
+	{
+		if (Entry == InEntry)
+		{
+			return true;
+		}
+	}
+
+	return false;
+}
+
+void SCustomStringListControllerListEditorWidget::OnValuesChanged()
+{
+	ListView->RequestListRefresh();
+}
+
+TConstArrayView<FName> SCustomStringListControllerListEditorWidget::GetEntries() const
+{
+	return Entries;
+}
+
+TSharedRef<ITableRow> SCustomStringListControllerListEditorWidget::OnGenerateRow(FName InEntry, const TSharedRef<STableViewBase>& InOwnerTableView)
+{
+	return SNew(SCustomStringListControllerEnumEditorRow, InOwnerTableView, SharedThis(this), InEntry);
+}
+
+void SCustomStringListControllerListEditorWidget::UpdateStringList()
+{
+	if (Entries.IsEmpty())
+	{
+		return;
+	}
+	
+	if (TSharedPtr<SCustomStringListControllerWidget> ParentWidget = ParentWidgetWeak.Pin())
+	{
+		ParentWidget->UpdateStringList(Entries);
+	}
+}
+
+void SCustomStringListControllerListEditorWidget::CloseWindow()
+{
+	if (TSharedPtr<SWindow> CurrentWindow = FSlateApplication::Get().FindWidgetWindow(SharedThis(this)))
+	{
+		CurrentWindow->RequestDestroyWindow();
+	}
+}
+
+FReply SCustomStringListControllerListEditorWidget::OnOkClicked()
+{
+	UpdateStringList();
+	CloseWindow();
+
+	return FReply::Handled();
+}
+
+FReply SCustomStringListControllerListEditorWidget::OnCancelClicked()
+{
+	CloseWindow();
+
+	return FReply::Handled();
+}
+
+void SCustomStringListControllerListEditorWidget::OnValidDrop(FName InDropTarget, FName InDropped, EItemDropZone InDropZone)
+{
+	if (InDropTarget == InDropped)
+	{
+		return;
+	}
+
+	int32 DroppedOnIndex = Entries.Find(InDropTarget);
+
+	if (DroppedOnIndex == INDEX_NONE)
+	{
+		return;
+	}
+
+	const int32 DroppedIndex = Entries.Find(InDropped);
+
+	if (DroppedIndex == INDEX_NONE || DroppedIndex == DroppedOnIndex)
+	{
+		return;
+	}
+
+	if (DroppedOnIndex >= DroppedIndex)
+	{
+		--DroppedOnIndex;
+	}
+
+	switch (InDropZone)
+	{
+		case EItemDropZone::BelowItem:
+			Entries.RemoveAt(DroppedIndex);
+			Entries.Insert(InDropped, DroppedOnIndex + 1);
+			break;
+
+		case EItemDropZone::OntoItem:
+		case EItemDropZone::AboveItem:
+			Entries.RemoveAt(DroppedIndex);
+			Entries.Insert(InDropped, DroppedOnIndex);
+			break;
+
+		default:
+			return;
+	}
+
+	OnValuesChanged();
+}
+
+void SCustomStringListControllerListEditorWidget::AddEntry(FName InEntry)
+{
+	Entries.Add(InEntry);
+
+	OnValuesChanged();
+}
+
+void SCustomStringListControllerListEditorWidget::RemoveEntry(FName InEntry)
+{
+	if (Entries.Num() <= 1)
+	{
+		return;
+	}
+
+	Entries.Remove(InEntry);
+
+	OnValuesChanged();
+}
+
+void SCustomStringListControllerListEditorWidget::ReplaceEntry(FName InFrom, FName InTo)
+{
+	for (int32 Index = 0; Index < Entries.Num(); ++Index)
+	{
+		if (Entries[Index] == InFrom)
+		{
+			Entries[Index] = InTo;
+			return;
+		}
+	}
+}
+
+#undef LOCTEXT_NAMESPACE
diff --git a/Engine/Plugins/VirtualProduction/RemoteControl/Source/RemoteControlUI/Private/UI/Controller/CustomControllers/SCustomStringListControllerListEditorWidget.h b/Engine/Plugins/VirtualProduction/RemoteControl/Source/RemoteControlUI/Private/UI/Controller/CustomControllers/SCustomStringListControllerListEditorWidget.h
new file mode 100644
--- /dev/null
+++ b/Engine/Plugins/VirtualProduction/RemoteControl/Source/RemoteControlUI/Private/UI/Controller/CustomControllers/SCustomStringListControllerListEditorWidget.h
@@ -0,0 +1,55 @@
+// Copyright Epic Games, Inc. All Rights Reserved.
+
+#pragma once
+
+#include "Widgets/SCompoundWidget.h"
+
+#include "UObject/NameTypes.h"
+#include "UObject/WeakObjectPtr.h"
+#include "Widgets/Views/SListView.h"
+
+class SCustomStringListControllerWidget;
+class SWindow;
+
+class SCustomStringListControllerListEditorWidget : public SCompoundWidget
+{
+public:
+	SLATE_BEGIN_ARGS(SCustomStringListControllerListEditorWidget)
+		{}
+	SLATE_END_ARGS()
+
+	static void OpenModalWindow(const TSharedRef<SCustomStringListControllerWidget>& InParentWidget, const TArray<FName>& InEntries);
+
+	virtual ~SCustomStringListControllerListEditorWidget() override = default;
+
+	void Construct(const FArguments& InArgs, const TSharedRef<SCustomStringListControllerWidget>& InParentWidget, const TArray<FName>& InEntries);
+
+	bool HasEntry(FName InDescription) const;
+
+	void AddEntry(FName InEntry);
+
+	void RemoveEntry(FName InEntry);
+
+	void ReplaceEntry(FName InFrom, FName InTo);
+
+	void OnValuesChanged();
+
+	TConstArrayView<FName> GetEntries() const;
+
+	void OnValidDrop(FName InDropTarget, FName InDropped, EItemDropZone InDropZone);
+
+private:
+	TWeakPtr<SCustomStringListControllerWidget> ParentWidgetWeak;
+	TSharedPtr<SListView<FName>> ListView;
+	TArray<FName> Entries;
+
+	TSharedRef<ITableRow> OnGenerateRow(FName InEntry, const TSharedRef<STableViewBase>& InOwnerTable);
+
+	void UpdateStringList();
+
+	void CloseWindow();
+
+	FReply OnOkClicked();
+		
+	FReply OnCancelClicked();
+};
diff --git a/Engine/Plugins/VirtualProduction/RemoteControl/Source/RemoteControlUI/Private/UI/Controller/CustomControllers/SCustomStringListControllerWidget.cpp b/Engine/Plugins/VirtualProduction/RemoteControl/Source/RemoteControlUI/Private/UI/Controller/CustomControllers/SCustomStringListControllerWidget.cpp
new file mode 100644
--- /dev/null
+++ b/Engine/Plugins/VirtualProduction/RemoteControl/Source/RemoteControlUI/Private/UI/Controller/CustomControllers/SCustomStringListControllerWidget.cpp
@@ -0,0 +1,205 @@
+// Copyright Epic Games, Inc. All Rights Reserved.
+
+#include "UI/Controller/CustomControllers/SCustomStringListControllerWidget.h"
+
+#include "Controller/RCCustomControllerUtilities.h"
+#include "PropertyHandle.h"
+#include "RCVirtualProperty.h"
+#include "ScopedTransaction.h"
+#include "UI/Controller/CustomControllers/SCustomStringListControllerListEditorWidget.h"
+#include "Widgets/Images/SImage.h"
+#include "Widgets/Input/SButton.h"
+#include "Widgets/SBoxPanel.h"
+#include "Widgets/Text/STextBlock.h"
+
+#define LOCTEXT_NAMESPACE "SCustomStringListControllerWidget"
+
+void SCustomStringListControllerWidget::Construct(const FArguments& InArgs, URCVirtualPropertyBase* InController, const TSharedPtr<IPropertyHandle>& InOriginalPropertyHandle)
+{
+	ControllerWeak = InController;
+	OriginalPropertyHandle = InOriginalPropertyHandle;
+
+	using namespace UE::RCCustomControllers;
+	
+	const FString EntriesString = InController->GetMetadataValue(StringListControllerEntriesName);
+	TArray<FString> EntriesArray;
+	EntriesString.ParseIntoArray(EntriesArray, TEXT(";"));
+
+	if (EntriesArray.IsEmpty())
+	{
+		EntriesArray.Add("New Entry 1");
+	}
+
+	Entries.Reserve(EntriesArray.Num());
+
+	for (const FString& Entry : EntriesArray)
+	{
+		Entries.Add(*Entry);
+	}
+
+	FString CurrentValue;
+	InOriginalPropertyHandle->GetValue(CurrentValue);
+
+	if (CurrentValue.IsEmpty())
+	{
+		CurrentValue = Entries[0].ToString();
+		InOriginalPropertyHandle->SetValue(CurrentValue);
+	}
+
+	StringSelector = SNew(SComboBox<FName>)
+		.OptionsSource(&Entries)
+		.InitiallySelectedItem(*CurrentValue)
+		.OnGenerateWidget(this, &SCustomStringListControllerWidget::OnGenerateWidget)
+		.OnSelectionChanged(this, &SCustomStringListControllerWidget::OnSelectionChanged)
+		.Content()
+		[
+			SNew(STextBlock)
+			.Text(this, &SCustomStringListControllerWidget::GetCurrentValueText)
+		];
+
+	ChildSlot
+	[
+		SNew(SHorizontalBox)
+		+ SHorizontalBox::Slot()
+		.FillWidth(1.f)
+		[
+			StringSelector.ToSharedRef()
+		]
+		+ SHorizontalBox::Slot()
+		.AutoWidth()
+		.Padding(2.f, 0.f, 0.f, 0.f)
+		[
+			SNew(SButton)
+			.OnClicked(this, &SCustomStringListControllerWidget::OnEditStringListClicked)
+			.ButtonStyle(FAppStyle::Get(), "HoverHintOnly")
+			.ContentPadding(2.f)
+			[
+				SNew(SImage)
+				.Image(FAppStyle::Get().GetBrush("Icons.Toolbar.Details"))
+				.DesiredSizeOverride(FVector2D(16.f, 16.f))
+			]
+		]
+	];
+}
+
+void SCustomStringListControllerWidget::UpdateStringList(const TArray<FName>& InNewEntries)
+{
+	if (InNewEntries.IsEmpty())
+	{
+		return;
+	}
+
+	URCVirtualPropertyBase* Controller = ControllerWeak.Get();
+
+	if (!Controller)
+	{
+		return;
+	}
+
+	FString CurrentValue;
+	OriginalPropertyHandle->GetValue(CurrentValue);
+
+	// Separators
+	int32 Length = InNewEntries.Num() - 1;
+
+	for (const FName& Entry : InNewEntries)
+	{
+		Length += Entry.GetStringLength();
+	}
+
+	FString NewMeta;
+	NewMeta.Reserve(Length);
+
+	Entries.Empty(InNewEntries.Num());
+
+	for (const FName& Entry : InNewEntries)
+	{
+		if (Entry.IsNone())
+		{
+			continue;
+		}
+
+		if (!NewMeta.IsEmpty())
+		{
+			NewMeta += TEXT(";");
+		}
+
+		NewMeta += Entry.ToString();
+
+		Entries.Add(Entry);
+	}
+
+	FScopedTransaction Transaction(LOCTEXT("UpdateStringValues", "Update String Values"));
+
+	Controller->Modify();
+
+	TArray<UObject*> Outers;
+	OriginalPropertyHandle->GetOuterObjects(Outers);
+
+	if (!Outers.IsEmpty())
+	{
+		for (UObject* Outer : Outers)
+		{
+			Outer->Modify();
+		}
+	}
+
+	using namespace UE::RCCustomControllers;
+
+	Controller->SetMetadataValue(StringListControllerEntriesName, NewMeta);
+
+	StringSelector->SetItemsSource(&Entries);
+
+	FName CurrentValueName = *CurrentValue;
+
+	if (!InNewEntries.Contains(CurrentValueName))
+	{
+		OriginalPropertyHandle->SetValue(InNewEntries[0].ToString());
+	}
+}
+
+FReply SCustomStringListControllerWidget::OnEditStringListClicked()
+{
+	SCustomStringListControllerListEditorWidget::OpenModalWindow(SharedThis(this), Entries);
+
+	return FReply::Handled();
+}
+
+TSharedRef<SWidget> SCustomStringListControllerWidget::OnGenerateWidget(FName InEntry)
+{
+	return SNew(STextBlock)
+		.Text(FText::FromName(InEntry));
+}
+
+void SCustomStringListControllerWidget::OnSelectionChanged(FName InEntry, ESelectInfo::Type InSelectionType)
+{
+	if (InEntry.IsNone())
+	{
+		return;
+	}
+
+	TArray<UObject*> Outers;
+	OriginalPropertyHandle->GetOuterObjects(Outers);
+
+	if (!Outers.IsEmpty())
+	{
+		FScopedTransaction Transaction(LOCTEXT("SetStringValue", "Set String Value"));
+
+		for (UObject* Outer : Outers)
+		{
+			Outer->Modify();
+		}
+	}
+	
+	OriginalPropertyHandle->SetValue(InEntry.ToString());
+}
+
+FText SCustomStringListControllerWidget::GetCurrentValueText() const
+{
+	FString CurrentValue;
+	OriginalPropertyHandle->GetValue(CurrentValue);
+
+	return FText::FromString(CurrentValue);
+}
+
+#undef LOCTEXT_NAMESPACE
diff --git a/Engine/Plugins/VirtualProduction/RemoteControl/Source/RemoteControlUI/Private/UI/Controller/CustomControllers/SCustomStringListControllerWidget.h b/Engine/Plugins/VirtualProduction/RemoteControl/Source/RemoteControlUI/Private/UI/Controller/CustomControllers/SCustomStringListControllerWidget.h
new file mode 100644
--- /dev/null
+++ b/Engine/Plugins/VirtualProduction/RemoteControl/Source/RemoteControlUI/Private/UI/Controller/CustomControllers/SCustomStringListControllerWidget.h
@@ -0,0 +1,50 @@
+// Copyright Epic Games, Inc. All Rights Reserved.
+
+#pragma once
+
+#include "Widgets/Layout/SBox.h"
+#include "Widgets/SCompoundWidget.h"
+
+#include "Containers/Array.h"
+#include "UObject/WeakObjectPtr.h"
+#include "Widgets/Input/SComboBox.h"
+
+class FText;
+class IPropertyHandle;
+class SBox;
+class URCVirtualPropertyBase;
+struct FPropertyChangedEvent;
+
+class SCustomStringListControllerWidget : public SCompoundWidget
+{
+public:
+	SLATE_BEGIN_ARGS(SCustomStringListControllerWidget)
+		{}
+	SLATE_END_ARGS()
+
+	virtual ~SCustomStringListControllerWidget() override = default;
+
+	/**
+	 * Constructs this widget with InArgs
+	 * @param InOriginalPropertyHandle Original PropertyHandle of this CustomWidget
+	 */
+	void Construct(const FArguments& InArgs, URCVirtualPropertyBase* InController, const TSharedPtr<IPropertyHandle>& InOriginalPropertyHandle);
+
+	void UpdateStringList(const TArray<FName>& InNewEntries);
+
+private:
+	/** Original PropertyHandle, used to update correctly the Widget and the Controller */
+	TWeakObjectPtr<URCVirtualPropertyBase> ControllerWeak;
+	TSharedPtr<IPropertyHandle> OriginalPropertyHandle;
+	TSharedPtr<SBox> PropertyWidgetContainer;
+	TArray<FName> Entries;
+	TSharedPtr<SComboBox<FName>> StringSelector;
+
+	FReply OnEditStringListClicked();
+
+	TSharedRef<SWidget> OnGenerateWidget(FName InEntry);
+
+	void OnSelectionChanged(FName InEntry, ESelectInfo::Type InSelectionType);
+
+	FText GetCurrentValueText() const;
+};
diff --git a/Engine/Plugins/VirtualProduction/RemoteControl/Source/RemoteControlUI/Private/UI/Controller/CustomControllers/SCustomTextureControllerWidget.cpp b/Engine/Plugins/VirtualProduction/RemoteControl/Source/RemoteControlUI/Private/UI/Controller/CustomControllers/SCustomTextureControllerWidget.cpp
--- a/Engine/Plugins/VirtualProduction/RemoteControl/Source/RemoteControlUI/Private/UI/Controller/CustomControllers/SCustomTextureControllerWidget.cpp
+++ b/Engine/Plugins/VirtualProduction/RemoteControl/Source/RemoteControlUI/Private/UI/Controller/CustomControllers/SCustomTextureControllerWidget.cpp
@@ -21,7 +21,7 @@
 
 #define LOCTEXT_NAMESPACE "ExternalTextureControllerWidget"
 
-void SCustomTextureControllerWidget::Construct(const FArguments& InArgs, const TSharedPtr<IPropertyHandle>& InOriginalPropertyHandle)
+void SCustomTextureControllerWidget::Construct(const FArguments& InArgs, URCVirtualPropertyBase* InController, const TSharedPtr<IPropertyHandle>& InOriginalPropertyHandle)
 {
 	if (!InOriginalPropertyHandle.IsValid())
 	{
diff --git a/Engine/Plugins/VirtualProduction/RemoteControl/Source/RemoteControlUI/Private/UI/Controller/CustomControllers/SCustomTextureControllerWidget.h b/Engine/Plugins/VirtualProduction/RemoteControl/Source/RemoteControlUI/Private/UI/Controller/CustomControllers/SCustomTextureControllerWidget.h
--- a/Engine/Plugins/VirtualProduction/RemoteControl/Source/RemoteControlUI/Private/UI/Controller/CustomControllers/SCustomTextureControllerWidget.h
+++ b/Engine/Plugins/VirtualProduction/RemoteControl/Source/RemoteControlUI/Private/UI/Controller/CustomControllers/SCustomTextureControllerWidget.h
@@ -1,4 +1,4 @@
-﻿// Copyright Epic Games, Inc. All Rights Reserved.
+// Copyright Epic Games, Inc. All Rights Reserved.
 
 #pragma once
 
@@ -22,7 +22,7 @@ public:
 	 * Constructs this widget with InArgs
 	 * @param InOriginalPropertyHandle Original PropertyHandle of this CustomWidget
 	 */
-	void Construct(const FArguments& InArgs, const TSharedPtr<IPropertyHandle>& InOriginalPropertyHandle);
+	void Construct(const FArguments& InArgs, URCVirtualPropertyBase* InController, const TSharedPtr<IPropertyHandle>& InOriginalPropertyHandle);
 
 private:
 	/** Handle toggling Texture Controller type change: External vs. Internal */
diff --git a/Engine/Plugins/VirtualProduction/RemoteControl/Source/RemoteControlUI/Private/UI/Controller/SRCControllerPanel.cpp b/Engine/Plugins/VirtualProduction/RemoteControl/Source/RemoteControlUI/Private/UI/Controller/SRCControllerPanel.cpp
--- a/Engine/Plugins/VirtualProduction/RemoteControl/Source/RemoteControlUI/Private/UI/Controller/SRCControllerPanel.cpp
+++ b/Engine/Plugins/VirtualProduction/RemoteControl/Source/RemoteControlUI/Private/UI/Controller/SRCControllerPanel.cpp
@@ -494,6 +494,36 @@ TSharedRef<SWidget> SRCControllerPanel::GetControllerMenuContentWidget() const
 	EnumBag.AddProperty(DefaultEnumName, EPropertyBagPropertyType::Enum, StaticEnum<EStyleColor>());
 
 	if (const FPropertyBagPropertyDesc* EnumBagPropertyDesc = EnumBag.FindPropertyDescByName(DefaultEnumName))
+	{
+		{
+			TSharedRef<SWidget> MenuItemWidget =
+				SNew(SHorizontalBox)
+				+ SHorizontalBox::Slot()
+				.AutoWidth()
+				.HAlign(HAlign_Left)
+				.VAlign(VAlign_Fill)
+				[
+					UE::RCUIHelpers::GetTypeColorWidget(EnumBagPropertyDesc->CachedProperty)
+				]
+				+ SHorizontalBox::Slot()
+				.AutoWidth()
+				.Padding(FMargin(6.f, 0.f))
+				[
+					SNew(STextBlock)
+					.Text(FText::FromName(DefaultEnumName))
+				];
+
+			FNewMenuDelegate Delegate = FNewMenuDelegate::CreateSP(this, &SRCControllerPanel::GetControllerMenuEnumContentWidget);
+			MenuBuilder.AddSubMenu(MenuItemWidget, Delegate);
+		}
+	}
+
+	// String List
+	const FName DefaultStringName = TEXT("Strings");
+	FInstancedPropertyBag StringBag;
+	StringBag.AddProperty(DefaultStringName, EPropertyBagPropertyType::String);
+
+	if (const FPropertyBagPropertyDesc* StringBagPropertyDesc = StringBag.FindPropertyDescByName(DefaultStringName))
 	{
 		TSharedRef<SWidget> MenuItemWidget =
 			SNew(SHorizontalBox)
@@ -502,18 +532,25 @@ TSharedRef<SWidget> SRCControllerPanel::GetControllerMenuContentWidget() const
 			.HAlign(HAlign_Left)
 			.VAlign(VAlign_Fill)
 			[
-				UE::RCUIHelpers::GetTypeColorWidget(EnumBagPropertyDesc->CachedProperty)
+				UE::RCUIHelpers::GetTypeColorWidget(StringBagPropertyDesc->CachedProperty)
 			]
 			+ SHorizontalBox::Slot()
 			.AutoWidth()
 			.Padding(FMargin(6.f, 0.f))
 			[
 				SNew(STextBlock)
-				.Text(FText::FromName(DefaultEnumName))
+				.Text(LOCTEXT("StringList", "String List"))
 			];
 
-		FNewMenuDelegate Delegate = FNewMenuDelegate::CreateSP(this, &SRCControllerPanel::GetControllerMenuEnumContentWidget);
-		MenuBuilder.AddSubMenu(MenuItemWidget, Delegate);
+		using namespace UE::RCCustomControllers;
+
+		FRCControllerPropertyInfo PropertyInfo;
+		PropertyInfo.Type = EPropertyBagPropertyType::String;
+		PropertyInfo.CustomMetaData.Add(CustomControllerNameKey, CustomStringListControllerName);
+		PropertyInfo.CustomMetaData.Add(StringListControllerEntriesName, "New Entry");
+
+		FUIAction Action(FExecuteAction::CreateSP(this, &SRCControllerPanel::OnAddControllerClicked, PropertyInfo.Type, PropertyInfo));
+		MenuBuilder.AddMenuEntry(Action, MenuItemWidget);
 	}
 
 	return MenuBuilder.MakeWidget();
@@ -662,6 +699,11 @@ TSharedRef<SWidget> SRCControllerPanel::GetMultiControllerSwitchWidget()
 }
 
 void SRCControllerPanel::OnAddControllerClicked(const EPropertyBagPropertyType InValueType, const UE::RemoteControl::UI::Private::FRCControllerPropertyInfo InControllerPropertyInfo) const
+{
+	OnAddControllerClicked(InValueType, InControllerPropertyInfo, URCController::StaticClass());
+}
+
+void SRCControllerPanel::OnAddControllerClicked(const EPropertyBagPropertyType InValueType, const UE::RemoteControl::UI::Private::FRCControllerPropertyInfo InControllerPropertyInfo, TSubclassOf<URCController> InControllerClass) const
 {
 	// Add to the asset
 	if (URemoteControlPreset* Preset = GetPreset())
@@ -669,7 +711,7 @@ void SRCControllerPanel::OnAddControllerClicked(const EPropertyBagPropertyType I
 		FScopedTransaction Transaction(LOCTEXT("AddController", "Add Controller"));
 		Preset->Modify();
 
-		URCVirtualPropertyInContainer* NewVirtualProperty = Preset->AddController(URCController::StaticClass(), InValueType, InControllerPropertyInfo.ValueTypeObject);
+		URCVirtualPropertyInContainer* NewVirtualProperty = Preset->AddController(InControllerClass, InValueType, InControllerPropertyInfo.ValueTypeObject);
 
 		// Copy over the custom metadata of this property
 		for (const TPair<FName, FString>& Pair : InControllerPropertyInfo.CustomMetaData)
diff --git a/Engine/Plugins/VirtualProduction/RemoteControl/Source/RemoteControlUI/Private/UI/Controller/SRCControllerPanel.h b/Engine/Plugins/VirtualProduction/RemoteControl/Source/RemoteControlUI/Private/UI/Controller/SRCControllerPanel.h
--- a/Engine/Plugins/VirtualProduction/RemoteControl/Source/RemoteControlUI/Private/UI/Controller/SRCControllerPanel.h
+++ b/Engine/Plugins/VirtualProduction/RemoteControl/Source/RemoteControlUI/Private/UI/Controller/SRCControllerPanel.h
@@ -115,6 +115,7 @@ private:
 
 	/** Handles click event for Add Controller button */
 	void OnAddControllerClicked(const EPropertyBagPropertyType InValueType, const UE::RemoteControl::UI::Private::FRCControllerPropertyInfo InControllerPropertyInfo) const;
+	void OnAddControllerClicked(const EPropertyBagPropertyType InValueType, const UE::RemoteControl::UI::Private::FRCControllerPropertyInfo InControllerPropertyInfo, TSubclassOf<URCController> InControllerClass) const;
 
 	/** Handles selection of controller type in submenus. */
 	void OnAddControllerSelected(TSharedPtr<UE::RemoteControl::UI::Private::FRCControllerPropertyInfo> InControllerPropertyInfo) const;
diff --git a/Engine/Plugins/VirtualProduction/RemoteControl/Source/RemoteControlUI/RemoteControlUI.Build.cs b/Engine/Plugins/VirtualProduction/RemoteControl/Source/RemoteControlUI/RemoteControlUI.Build.cs
--- a/Engine/Plugins/VirtualProduction/RemoteControl/Source/RemoteControlUI/RemoteControlUI.Build.cs
+++ b/Engine/Plugins/VirtualProduction/RemoteControl/Source/RemoteControlUI/RemoteControlUI.Build.cs
@@ -43,7 +43,8 @@ public class RemoteControlUI : ModuleRules
 				"ToolWidgets",
 				"TypedElementFramework",
 				"TypedElementRuntime",
-				"UnrealEd"
+				"UnrealEd",
+				"WidgetRegistration"
 			}
 		);
     }
