Remote Control: Changed visual of property in use indicator
- Indicator is now an icon instead of a row highlight.
- Icon also now indicates if more than 1 behavour is linked to a property.

diff --git a/Engine/Plugins/VirtualProduction/RemoteControl/Source/RemoteControlUI/Private/UI/SRCPanelExposedEntitiesList.cpp b/Engine/Plugins/VirtualProduction/RemoteControl/Source/RemoteControlUI/Private/UI/SRCPanelExposedEntitiesList.cpp
--- a/Engine/Plugins/VirtualProduction/RemoteControl/Source/RemoteControlUI/Private/UI/SRCPanelExposedEntitiesList.cpp
+++ b/Engine/Plugins/VirtualProduction/RemoteControl/Source/RemoteControlUI/Private/UI/SRCPanelExposedEntitiesList.cpp
@@ -97,16 +97,6 @@ public:
 		FSuperRowType::Construct(SuperArgs, OwnerTableView);
 	}
 
-	virtual const FSlateBrush* GetBorder() const override
-	{
-		if (Entity.IsValid() && Entity->IsFieldInUse() && !IsHovered())
-		{
-			return FRemoteControlPanelStyle::Get()->GetBrush("RemoteControlPanel.EntityInUseBorder");
-		}
-
-		return SMultiColumnTableRow<TSharedPtr<SRCPanelTreeNode>>::GetBorder();
-	}
-
 	TSharedRef<SWidget> GenerateWidgetForColumn(const FName& InColumnName) override
 	{
 		if (Entity.IsValid())
@@ -909,7 +899,7 @@ void SRCPanelExposedEntitiesList::RunUpdateEntityUsage()
 	TArray<UObject*> Subobjects;
 	GetObjectsWithOuter(Preset.Get(), Subobjects);
 
-	BoundEntityIds.Empty();
+	TMap<FGuid, TSet<URCBehaviour*>> EntityToBehaviour;
 	TSet<UObject*> ValidContainers;
 	TSet<UObject*> InvalidContainers;
 
@@ -958,7 +948,15 @@ void SRCPanelExposedEntitiesList::RunUpdateEntityUsage()
 			}
 		}
 
-		BoundEntityIds.Add(Action->ExposedFieldId);
+		//BoundEntityIds.FindOrAdd(Action->ExposedFieldId) += 1;
+		EntityToBehaviour.FindOrAdd(Action->ExposedFieldId).Add(Behaviour);
+	}
+
+	BoundEntityIds.Empty(EntityToBehaviour.Num());
+
+	for (const TPair<FGuid, TSet<URCBehaviour*>>& EntityPair : EntityToBehaviour)
+	{
+		BoundEntityIds.Add(EntityPair.Key, EntityPair.Value.Num());
 	}
 
 	for (const TPair <FGuid, TSharedPtr<SRCPanelTreeNode>>& Node : FieldWidgetMap)
@@ -969,11 +967,11 @@ void SRCPanelExposedEntitiesList::RunUpdateEntityUsage()
 		{
 			if (NodeWidget->GetRCType() == SRCPanelTreeNode::ENodeType::Field && NodeWidget->GetRCId().IsValid())
 			{
-				NodeWidget->SetFieldInUse(BoundEntityIds.Contains(NodeWidget->GetRCId()));
+				NodeWidget->SetFieldInUseCount(BoundEntityIds.FindRef(NodeWidget->GetRCId()));
 			}
 			else
 			{
-				NodeWidget->SetFieldInUse(false);
+				NodeWidget->SetFieldInUseCount(0);
 			}
 		}
 	}
@@ -1265,11 +1263,11 @@ TSharedRef<ITableRow> SRCPanelExposedEntitiesList::OnGenerateRow(TSharedPtr<SRCP
 	{
 		if (Node->GetRCType() == SRCPanelTreeNode::ENodeType::Field && Node->GetRCId().IsValid())
 		{
-			Node->SetFieldInUse(BoundEntityIds.Contains(Node->GetRCId()));
+			Node->SetFieldInUseCount(BoundEntityIds.FindRef(Node->GetRCId()));
 		}
 		else
 		{
-			Node->SetFieldInUse(false);
+			Node->SetFieldInUseCount(0);
 		}
 
 		constexpr float LeftPadding = 3.f;
diff --git a/Engine/Plugins/VirtualProduction/RemoteControl/Source/RemoteControlUI/Private/UI/SRCPanelExposedEntitiesList.h b/Engine/Plugins/VirtualProduction/RemoteControl/Source/RemoteControlUI/Private/UI/SRCPanelExposedEntitiesList.h
--- a/Engine/Plugins/VirtualProduction/RemoteControl/Source/RemoteControlUI/Private/UI/SRCPanelExposedEntitiesList.h
+++ b/Engine/Plugins/VirtualProduction/RemoteControl/Source/RemoteControlUI/Private/UI/SRCPanelExposedEntitiesList.h
@@ -340,8 +340,8 @@ private:
 	bool bSearchRequested;
 	/** If true, the entity list will have its in use flag updated next frame. */
 	bool bUpdateEntityUsage;
-	/** Temporary list of the currently bound entity ids. */
-	TSet<FGuid> BoundEntityIds;
+	/** Transient list of the currently bound entity ids and the number of times they are bound. */
+	TMap<FGuid, int32> BoundEntityIds;
 	/** Panel Style reference. */
 	const FRCPanelStyle* RCPanelStyle;
 	/** Holds the header row of entities list. */
diff --git a/Engine/Plugins/VirtualProduction/RemoteControl/Source/RemoteControlUI/Private/UI/SRCPanelExposedField.cpp b/Engine/Plugins/VirtualProduction/RemoteControl/Source/RemoteControlUI/Private/UI/SRCPanelExposedField.cpp
--- a/Engine/Plugins/VirtualProduction/RemoteControl/Source/RemoteControlUI/Private/UI/SRCPanelExposedField.cpp
+++ b/Engine/Plugins/VirtualProduction/RemoteControl/Source/RemoteControlUI/Private/UI/SRCPanelExposedField.cpp
@@ -434,6 +434,24 @@ TSharedRef<SWidget> SRCPanelExposedField::GetWidget(const FName ForColumnName, c
 		}
 	}
 
+	if (ForColumnName == RemoteControlPresetColumns::Reset)
+	{
+		return SNew(SHorizontalBox)
+			+ SHorizontalBox::Slot()
+			.VAlign(VAlign_Center)
+			.HAlign(HAlign_Center)
+			[
+				SRCPanelTreeNode::GetWidget(ForColumnName, InActiveProtocol)
+			]
+			+ SHorizontalBox::Slot()
+			.AutoWidth()
+			.VAlign(VAlign_Center)
+			.HAlign(HAlign_Center)
+			[
+				CreateInUseWidget()
+			];
+	}
+
 	return SRCPanelTreeNode::GetWidget(ForColumnName, InActiveProtocol);
 }
 
@@ -913,6 +931,59 @@ void SRCPanelExposedField::OnPropertyIdTextCommitted(const FText& InText, ETextC
 	}
 }
 
+TSharedRef<SWidget> SRCPanelExposedField::CreateInUseWidget()
+{
+	if (const TSharedPtr<FRemoteControlField> RCField = WeakField.Pin())
+	{
+		return SNew(SImage)
+			.DesiredSizeOverride(FVector2D(12.f))
+			.Image_Lambda(
+				[ThisWeak = AsWeak()]
+				{
+					if (TSharedPtr<SRCPanelExposedField> This = StaticCastSharedPtr<SRCPanelExposedField>(ThisWeak.Pin()))
+					{
+						if (This->GetFieldInUseCount() > 1)
+						{
+							return FAppStyle::Get().GetBrush("Icons.WarningWithColor");
+						}
+					}
+
+					return FAppStyle::Get().GetBrush("Icons.Link");
+				}
+			)
+			.ToolTipText_Lambda(
+				[ThisWeak = AsWeak()]
+				{
+					if (TSharedPtr<SRCPanelExposedField> This = StaticCastSharedPtr<SRCPanelExposedField>(ThisWeak.Pin()))
+					{
+						if (This->GetFieldInUseCount() > 1)
+						{
+							return LOCTEXT("PropertyBoundMultiple", "WARNING: Property is bound to multiple behaviours!");
+						}
+					}
+
+					return LOCTEXT("PropertyBound", "Property is bound to a controller.");
+				}
+			)
+			.Visibility_Lambda(
+				[ThisWeak = AsWeak()]
+				{
+					if (TSharedPtr<SRCPanelExposedField> This = StaticCastSharedPtr<SRCPanelExposedField>(ThisWeak.Pin()))
+					{
+						if (This->GetFieldInUseCount() > 0)
+						{
+							return EVisibility::Visible;
+						}
+					}
+
+					return EVisibility::Hidden;
+				}
+			);
+	}
+
+	return SNullWidget::NullWidget;
+}
+
 void SRCPanelFieldChildNode::Construct(const FArguments& InArgs, const TSharedRef<IDetailTreeNode>& InNode, FRCColumnSizeData InColumnSizeData)
 {
 	TArray<TSharedRef<IDetailTreeNode>> ChildNodes;
diff --git a/Engine/Plugins/VirtualProduction/RemoteControl/Source/RemoteControlUI/Private/UI/SRCPanelExposedField.h b/Engine/Plugins/VirtualProduction/RemoteControl/Source/RemoteControlUI/Private/UI/SRCPanelExposedField.h
--- a/Engine/Plugins/VirtualProduction/RemoteControl/Source/RemoteControlUI/Private/UI/SRCPanelExposedField.h
+++ b/Engine/Plugins/VirtualProduction/RemoteControl/Source/RemoteControlUI/Private/UI/SRCPanelExposedField.h
@@ -110,6 +110,8 @@ private:
 	/** Handles committing a property id label. */
 	void OnPropertyIdTextCommitted(const FText& InText, ETextCommit::Type InCommitInfo);
 
+	TSharedRef<SWidget> CreateInUseWidget();
+
 private:
 	/** Weak pointer to the underlying RC Field. */
 	TWeakPtr<FRemoteControlField> WeakField;
diff --git a/Engine/Plugins/VirtualProduction/RemoteControl/Source/RemoteControlUI/Private/UI/SRCPanelTreeNode.h b/Engine/Plugins/VirtualProduction/RemoteControl/Source/RemoteControlUI/Private/UI/SRCPanelTreeNode.h
--- a/Engine/Plugins/VirtualProduction/RemoteControl/Source/RemoteControlUI/Private/UI/SRCPanelTreeNode.h
+++ b/Engine/Plugins/VirtualProduction/RemoteControl/Source/RemoteControlUI/Private/UI/SRCPanelTreeNode.h
@@ -117,14 +117,14 @@ struct SRCPanelTreeNode : public SCompoundWidget, public IHasProtocolExtensibili
 	/** Executed when the PropertyId of the node is changed. */
 	FOnPropertyIdRenamed& OnPropertyIdRenamed() { return OnPropertyIdRenamedDelegate; }
 
-	bool IsFieldInUse() const 
+	int32 GetFieldInUseCount() const
 	{ 
-		return bIsFieldInUse; 
+		return FieldUseCount;
 	}
 
-	void SetFieldInUse(bool bInFieldInUse) 
+	void SetFieldInUseCount(int32 InUseCount)
 	{ 
-		bIsFieldInUse = bInFieldInUse; 
+		FieldUseCount = InUseCount;
 	}
 
 protected:
@@ -166,7 +166,7 @@ protected:
 	FRCColumnSizeData ColumnSizeData;
 
 	/** If true, an action is bound to this node's entity. */
-	bool bIsFieldInUse = false;
+	int32 FieldUseCount = 0;
 
 private:
 
