Remote Control / Actor Modifier / Property Animator: Bridge plugin
- Actors are now selected when select source is used on a modifier or animator.

diff --git a/Engine/Plugins/Experimental/RemoteControlActorModifierBridge/RemoteControlActorModifierBridge.uplugin b/Engine/Plugins/Experimental/RemoteControlActorModifierBridge/RemoteControlActorModifierBridge.uplugin
new file mode 100644
--- /dev/null
+++ b/Engine/Plugins/Experimental/RemoteControlActorModifierBridge/RemoteControlActorModifierBridge.uplugin
@@ -0,0 +1,41 @@
+{
+	"FileVersion": 3,
+	"Version": 1,
+	"VersionName": "1.0",
+	"FriendlyName": "Remote Control Actor Modifier Bridge",
+	"Description": "Interface between the Remote Control, Actor Modifier and Property Animator plugins.",
+	"Category": "Messaging",
+	"CreatedBy": "Epic Games, Inc.",
+	"CreatedByURL": "https://epicgames.com",
+	"DocsURL": "",
+	"MarketplaceURL": "",
+	"SupportURL": "",
+	"CanContainContent": true,
+	"IsBetaVersion": true,
+	"Installed": false,
+	"Plugins": [
+		{
+			"Name": "ActorModifierCore",
+			"Enabled": true
+		},
+		{
+			"Name": "OperatorStack",
+			"Enabled": true
+		},
+		{
+			"Name": "RemoteControl",
+			"Enabled": true
+		},
+		{
+			"Name": "PropertyAnimatorCore",
+			"Enabled": true
+		}
+	],
+	"Modules": [
+		{
+			"Name": "RemoteControlActorModifierBridge",
+			"Type": "Editor",
+			"LoadingPhase": "Default"
+		}
+	]
+}
\ No newline at end of file
diff --git a/Engine/Plugins/Experimental/RemoteControlActorModifierBridge/Source/Private/RemoteControlActorModifierBridgeModule.cpp b/Engine/Plugins/Experimental/RemoteControlActorModifierBridge/Source/Private/RemoteControlActorModifierBridgeModule.cpp
new file mode 100644
--- /dev/null
+++ b/Engine/Plugins/Experimental/RemoteControlActorModifierBridge/Source/Private/RemoteControlActorModifierBridgeModule.cpp
@@ -0,0 +1,107 @@
+// Copyright Epic Games, Inc. All Rights Reserved.
+
+#include "RemoteControlActorModifierBridgeModule.h"
+
+#include "Animators/PropertyAnimatorCoreBase.h"
+#include "Components/ActorComponent.h"
+#include "Containers/Array.h"
+#include "GameFramework/Actor.h"
+#include "IRemoteControlUIModule.h"
+#include "Modifiers/ActorModifierCoreBase.h"
+#include "PropertyPath.h"
+#include "RemoteControlField.h"
+#include "Subsystems/OperatorStackEditorSubsystem.h"
+#include "UObject/NameTypes.h"
+
+namespace UE::FRemoteControlActorModifierBridge::Private
+{
+	const FLazyName ResolverName = "ActorModifierResolver";
+}
+
+void FRemoteControlActorModifierBridgeModule::StartupModule()
+{
+	using namespace UE::FRemoteControlActorModifierBridge::Private;
+
+	IRemoteControlUIModule::Get().RegisterPropertyResolver(
+		ResolverName,
+		FRCPropertyResolver::CreateRaw(this, &FRemoteControlActorModifierBridgeModule::ResolverActorModifierProperty)
+	);
+}
+
+void FRemoteControlActorModifierBridgeModule::ShutdownModule()
+{
+	using namespace UE::FRemoteControlActorModifierBridge::Private;
+
+	IRemoteControlUIModule::Get().UnregisterPropertyResolver(ResolverName);
+}
+
+bool FRemoteControlActorModifierBridgeModule::ResolverActorModifierProperty(const TSharedRef<FRemoteControlProperty>& InProperty, TArray<UObject*>& InOutBoundObjects, TSharedPtr<FPropertyPath>& OutPropertyPath)
+{
+	const UActorModifierCoreBase* FoundModifier = nullptr;
+	const UPropertyAnimatorCoreBase* FoundAnimator = nullptr;
+
+	for (auto It = InOutBoundObjects.CreateIterator(); It; ++It)
+	{
+		const UObject* BindingObject = *It;
+
+		if (const UActorModifierCoreBase* Modifier = Cast<UActorModifierCoreBase>(BindingObject))
+		{
+			It.RemoveCurrent();
+
+			if (AActor* Actor = Modifier->GetModifiedActor())
+			{
+				InOutBoundObjects.Add(Actor);
+			}
+
+			if (!FoundModifier)
+			{
+				FoundModifier = Modifier;
+			}
+		}
+		else if (const UPropertyAnimatorCoreBase* Animator = Cast<UPropertyAnimatorCoreBase>(BindingObject))
+		{
+			It.RemoveCurrent();
+
+			if (AActor* Actor = Animator->GetAnimatorActor())
+			{
+				InOutBoundObjects.Add(Actor);
+			}
+
+			if (!FoundAnimator)
+			{
+				FoundAnimator = Animator;
+			}
+
+			if (!OutPropertyPath.IsValid())
+			{
+				for (const FPropertyAnimatorCoreData& CoreData : Animator->GetLinkedProperties())
+				{
+					if (FProperty* MemberProperty = CoreData.GetMemberProperty())
+					{
+						if (MemberProperty->HasAnyPropertyFlags(CPF_Edit | CPF_EditConst))
+						{
+							OutPropertyPath = FPropertyPath::Create(MemberProperty);
+							break;
+						}
+					}
+				}
+			}
+		}
+	}
+
+	if (FoundModifier)
+	{
+		UOperatorStackEditorSubsystem::Get()->FocusCustomizationWidget(FoundModifier, TEXT("Modifiers"));
+		return true;
+	}
+
+	if (FoundAnimator)
+	{
+		UOperatorStackEditorSubsystem::Get()->FocusCustomizationWidget(FoundAnimator, TEXT("Animators"));
+		return true;
+	}
+
+	return false;
+}
+
+IMPLEMENT_MODULE(FRemoteControlActorModifierBridgeModule, RemoteControlActorModifierBridge);
diff --git a/Engine/Plugins/Experimental/RemoteControlActorModifierBridge/Source/Private/RemoteControlActorModifierBridgeModule.h b/Engine/Plugins/Experimental/RemoteControlActorModifierBridge/Source/Private/RemoteControlActorModifierBridgeModule.h
new file mode 100644
--- /dev/null
+++ b/Engine/Plugins/Experimental/RemoteControlActorModifierBridge/Source/Private/RemoteControlActorModifierBridgeModule.h
@@ -0,0 +1,24 @@
+// Copyright Epic Games, Inc. All Rights Reserved.
+
+#pragma once
+
+#include "Modules/ModuleInterface.h"
+
+#include "Containers/ContainersFwd.h"
+#include "Templates/SharedPointerFwd.h"
+
+class FPropertyPath;
+class UObject;
+struct FRemoteControlProperty;
+
+class FRemoteControlActorModifierBridgeModule : public IModuleInterface
+{
+public:
+	//~ Begin IModuleInterface
+	virtual void StartupModule() override;
+	virtual void ShutdownModule() override;
+	//~ End IModuleInterface
+
+private:
+	bool ResolverActorModifierProperty(const TSharedRef<FRemoteControlProperty>& InProperty, TArray<UObject*>& InOutBoundObjects, TSharedPtr<FPropertyPath>& OutPropertyPath);
+};
diff --git a/Engine/Plugins/Experimental/RemoteControlActorModifierBridge/Source/RemoteControlActorModifierBridge.build.cs b/Engine/Plugins/Experimental/RemoteControlActorModifierBridge/Source/RemoteControlActorModifierBridge.build.cs
new file mode 100644
--- /dev/null
+++ b/Engine/Plugins/Experimental/RemoteControlActorModifierBridge/Source/RemoteControlActorModifierBridge.build.cs
@@ -0,0 +1,22 @@
+// Copyright Epic Games, Inc. All Rights Reserved.
+
+using UnrealBuildTool;
+
+public class RemoteControlActorModifierBridge : ModuleRules
+{
+	public RemoteControlActorModifierBridge(ReadOnlyTargetRules Target) : base(Target)
+	{
+		PrivateDependencyModuleNames.AddRange(
+			new string[] {
+				"ActorModifierCore",
+				"Core",
+				"CoreUObject",
+				"OperatorStackEditor",
+				"PropertyEditor",
+				"PropertyAnimatorCore",
+				"RemoteControl",
+				"RemoteControlUI"
+			}
+		);
+	}
+}
\ No newline at end of file
diff --git a/Engine/Plugins/VirtualProduction/RemoteControl/Source/RemoteControlUI/Private/RemoteControlUIModule.cpp b/Engine/Plugins/VirtualProduction/RemoteControl/Source/RemoteControlUI/Private/RemoteControlUIModule.cpp
--- a/Engine/Plugins/VirtualProduction/RemoteControl/Source/RemoteControlUI/Private/RemoteControlUIModule.cpp
+++ b/Engine/Plugins/VirtualProduction/RemoteControl/Source/RemoteControlUI/Private/RemoteControlUIModule.cpp
@@ -69,6 +69,11 @@ static const FName DetailsTabIdentifiers_LevelEditor[] = {
 
 static const int32 DetailsTabIdentifiers_LevelEditor_Count = 4;
 
+namespace UE::RemoteControlUI::Private
+{
+	static const FLazyName DisplayClusterConfigDataClassName = "DisplayClusterConfigurationData";
+}
+
 FRCExposesPropertyArgs::FRCExposesPropertyArgs()
 	: Id(FGuid::NewGuid())
 {
@@ -281,6 +286,7 @@ void FRemoteControlUIModule::StartupModule()
 	RegisterStructCustomizations();
 	RegisterSettings();
 	RegisterWidgetFactories();
+	RegisterDefaultPropertyResolvers();
 }
 
 void FRemoteControlUIModule::ShutdownModule()
@@ -1226,63 +1232,137 @@ TSharedPtr<SWidget> FRemoteControlUIModule::CreateCustomControllerWidget(URCVirt
 	return CustomControllerWidget;
 }
 
-void FRemoteControlUIModule::SelectSource(TSharedPtr<FRemoteControlEntity> InEntity)
+void FRemoteControlUIModule::RegisterPropertyResolver(FName InResolverId, const FRCPropertyResolver& InResolver)
 {
-	const TSharedPtr<FRemoteControlField> RCField = StaticCastSharedPtr<FRemoteControlField>(InEntity);
+	PropertyResolvers.FindOrAdd(InResolverId, InResolver);
+}
 
-	if (!RCField.IsValid() || RCField->GetStruct() != FRemoteControlProperty::StaticStruct())
+void FRemoteControlUIModule::UnregisterPropertyResolver(FName InResolverId)
+{
+	PropertyResolvers.Remove(InResolverId);
+}
+
+void FRemoteControlUIModule::RegisterDefaultPropertyResolvers()
+{
+	using namespace UE::RemoteControlUI::Private;
+
+	PropertyResolvers.Add(
+		DisplayClusterConfigDataClassName,
+		FRCPropertyResolver::CreateStatic(&FRemoteControlUIModule::ResolveDisplayClusterConfigurationDataProperty)
+	);
+}
+
+bool FRemoteControlUIModule::ResolveDisplayClusterConfigurationDataProperty(const TSharedRef<FRemoteControlProperty>& InProperty, TArray<UObject*>& InOutBoundObjects, TSharedPtr<FPropertyPath>& OutPropertyPath)
+{
+	for (auto It = InOutBoundObjects.CreateIterator(); It; ++It)
 	{
-		return;
+		const UObject* BindingObject = *It;
+
+		AActor* OwnerActor = BindingObject->GetTypedOuter<AActor>();
+
+		if (!OwnerActor)
+		{
+			continue;
+		}
+
+		// --- Special NDisplay handling because of their customization ---
+		// Since display cluster config data is not created as a default subobject, therefore we have no way of retrieving the CurrentConfigData property from the config object itself.
+		using namespace UE::RemoteControlUI::Private;
+
+		if (BindingObject->GetClass()->GetFName() != DisplayClusterConfigDataClassName)
+		{
+			continue;
+		}
+
+		FProperty* Property = OwnerActor->GetClass()->FindPropertyByName("CurrentConfigData");
+
+		if (!Property)
+		{
+			continue;
+		}
+
+		// Append "CurrentConfigData" to the beginning of the path.
+		TSharedRef<FPropertyPath> CurrentPropertyPath = InProperty->FieldPathInfo.ToPropertyPath();
+		OutPropertyPath = FPropertyPath::Create(Property);
+
+		for (int32 Index = 0; Index < CurrentPropertyPath->GetNumProperties(); Index++)
+		{
+			OutPropertyPath->AddProperty(CurrentPropertyPath->GetPropertyInfo(Index));
+		}
+
+		return true;
 	}
 
-	const TSharedPtr<FRemoteControlProperty> RCProp = StaticCastSharedPtr<FRemoteControlProperty>(RCField);
+	return false;
+}
 
-	// Resolve it to get the property path.
-	RCProp->FieldPathInfo.Resolve(RCProp->GetBoundObject());
+TSharedPtr<FPropertyPath> FRemoteControlUIModule::ResolveProperty(const TSharedRef<FRemoteControlProperty>& InProperty, TArray<UObject*>& InOutBoundObjects)
+{
+	// Resolve it to ensure the property path.
+	InProperty->FieldPathInfo.Resolve(InProperty->GetBoundObject());
 
-	if (!RCProp->FieldPathInfo.IsResolved())
+	if (!InProperty->FieldPathInfo.IsResolved())
 	{
-		return;
+		return nullptr;
 	}
 
-	TSharedRef<FPropertyPath> PropertyPath = RCProp->FieldPathInfo.ToPropertyPath();
-	TArray<UObject*> BoundObjects = RCField->GetBoundObjects();
+	TSharedPtr<FPropertyPath> ResolvedPath;
 
-	for (auto It = BoundObjects.CreateIterator(); It; ++It)
+	for (const TPair<FName, FRCPropertyResolver>& ResolverPair : PropertyResolvers)
 	{
-		if (AActor* OwnerActor = (*It)->GetTypedOuter<AActor>())
+		const bool bIsResolved = ResolverPair.Value.Execute(InProperty, InOutBoundObjects, ResolvedPath);
+
+		if (bIsResolved)
 		{
-			const UObject* BindingObject = *It;
+			break;
+		}
+	}
 
-			// When we encounter a non-component object, 
-			if (!BindingObject->IsA<UActorComponent>())
-			{
-				BoundObjects.Add(OwnerActor);
-				It.RemoveCurrent();
-			}
+	// By default, convert non-actor/actor components to the owning actors.
+	for (auto It = InOutBoundObjects.CreateIterator(); It; ++It)
+	{
+		const UObject* BindingObject = *It;
 
-			// --- Special NDisplay handling because of their customization ---
-			// Since display cluster config data is not created as a default subobject, therefore we have no way of retrieving the CurrentConfigData property from the config object itself.
-			static FName DisplayClusterConfigDataClassName = "DisplayClusterConfigurationData";
-			if (BindingObject->GetClass()->GetFName() == DisplayClusterConfigDataClassName)
-			{
-				if (FProperty* Property = OwnerActor->GetClass()->FindPropertyByName("CurrentConfigData"))
-				{
-					// Append "CurrentConfigData" to the beginning of the path.
-					TSharedRef<FPropertyPath> NewPropertyPath = FPropertyPath::Create(Property);
-					for (int32 Index = 0; Index < PropertyPath->GetNumProperties(); Index++)
-					{
-						NewPropertyPath->AddProperty(PropertyPath->GetPropertyInfo(Index));
-					}
+		AActor* OwnerActor = BindingObject->GetTypedOuter<AActor>();
 
-					PropertyPath = NewPropertyPath;
-				}
-			}
+		if (!OwnerActor)
+		{
+			continue;
+		}
+
+		// When we encounter a non-component object, 
+		if (!BindingObject->IsA<UActorComponent>())
+		{
+			InOutBoundObjects.Add(OwnerActor);
+			It.RemoveCurrent();
 		}
 	}
 
+	// And then return the default path resolution
+	return InProperty->FieldPathInfo.ToPropertyPath();
+}
+
+void FRemoteControlUIModule::SelectSource(TSharedPtr<FRemoteControlEntity> InEntity)
+{
+	const TSharedPtr<FRemoteControlField> RCField = StaticCastSharedPtr<FRemoteControlField>(InEntity);
+
+	if (!RCField.IsValid() || RCField->GetStruct() != FRemoteControlProperty::StaticStruct())
+	{
+		return;
+	}
+
+	const TSharedPtr<FRemoteControlProperty> RCProp = StaticCastSharedPtr<FRemoteControlProperty>(RCField);
+
+	TArray<UObject*> BoundObjects = RCProp->GetBoundObjects();
+	TSharedPtr<FPropertyPath> PropertyPath = ResolveProperty(RCProp.ToSharedRef(), BoundObjects);
+
 	SelectObjects(BoundObjects);
 
+	if (!PropertyPath.IsValid())
+	{
+		return;
+	}
+
 	FTimerHandle Handle;
 	const FTimerDelegate Delegate = FTimerDelegate::CreateLambda(([PropertyPath]()
 		{
diff --git a/Engine/Plugins/VirtualProduction/RemoteControl/Source/RemoteControlUI/Private/RemoteControlUIModule.h b/Engine/Plugins/VirtualProduction/RemoteControl/Source/RemoteControlUI/Private/RemoteControlUIModule.h
--- a/Engine/Plugins/VirtualProduction/RemoteControl/Source/RemoteControlUI/Private/RemoteControlUIModule.h
+++ b/Engine/Plugins/VirtualProduction/RemoteControl/Source/RemoteControlUI/Private/RemoteControlUIModule.h
@@ -73,6 +73,8 @@ public:
 	virtual void HighlightPropertyInDetailsPanel(const FPropertyPath& Path) const override;
 	virtual void SelectObjects(const TArray<UObject*>& Objects) const override;
 	virtual TSharedPtr<SWidget> CreateCustomControllerWidget(URCVirtualPropertyBase* InController, TSharedPtr<IPropertyHandle> InOriginalPropertyHandle) const override;
+	virtual void RegisterPropertyResolver(FName InResolverId, const FRCPropertyResolver& InResolver) override;
+	virtual void UnregisterPropertyResolver(FName InResolverId) override;
 	//~ End IRemoteControlUIModule interface
 
 	void SelectSource(TSharedPtr<FRemoteControlEntity> InEntity);
@@ -264,6 +266,13 @@ private:
 	bool HasChildPropertiesExposed(const FRCExposesPropertyArgs& InPropertyArgs) const;
 
 	void GetAllExposableSubPropertyFromStruct(const FRCExposesPropertyArgs InPropertyArgs, TArray<FRCExposesAllPropertiesArgs>& OutAllProperty) const;
+
+	void RegisterDefaultPropertyResolvers();
+
+	static bool ResolveDisplayClusterConfigurationDataProperty(const TSharedRef<FRemoteControlProperty>& InProperty, TArray<UObject*>& InOutBoundObjects, TSharedPtr<FPropertyPath>& OutPropertyPath);
+
+	TSharedPtr<FPropertyPath> ResolveProperty(const TSharedRef<FRemoteControlProperty>& InProperty, TArray<UObject*>& InOutBoundObjects);
+
 private:
 	/** The custom actions added to the actor context menu. */
 	TSharedPtr<class FRemoteControlPresetActions> RemoteControlPresetActions;
@@ -322,4 +331,7 @@ private:
 	EAssetTypeCategories::Type RemoteControlAssetCategoryBit;
 
 	TArray<TWeakPtr<SRemoteControlPanel>> RegisteredRemoteControlPanels;
+
+	/** List of property sources resolvers. */
+	TMap<FName, FRCPropertyResolver> PropertyResolvers;
 };
diff --git a/Engine/Plugins/VirtualProduction/RemoteControl/Source/RemoteControlUI/Public/IRemoteControlUIModule.h b/Engine/Plugins/VirtualProduction/RemoteControl/Source/RemoteControlUI/Public/IRemoteControlUIModule.h
--- a/Engine/Plugins/VirtualProduction/RemoteControl/Source/RemoteControlUI/Public/IRemoteControlUIModule.h
+++ b/Engine/Plugins/VirtualProduction/RemoteControl/Source/RemoteControlUI/Public/IRemoteControlUIModule.h
@@ -1,7 +1,7 @@
 // Copyright Epic Games, Inc. All Rights Reserved.
 #pragma once
 
-#include "CoreMinimal.h"
+#include "Containers/ContainersFwd.h"
 #include "Misc/Attribute.h"
 #include "Modules/ModuleManager.h"
 #include "Modules/ModuleInterface.h"
@@ -26,6 +26,7 @@ class URCVirtualPropertyBase;
 class URemoteControlPreset;
 struct FOnGenerateGlobalRowExtensionArgs;
 struct FRemoteControlEntity;
+struct FRemoteControlProperty;
 struct SRCPanelTreeNode;
 
 struct FRCColumnSizeData
@@ -169,6 +170,13 @@ DECLARE_MULTICAST_DELEGATE_OneParam(FOnGenerateControllerExtensionsWidgets, FRCC
  */
 DECLARE_MULTICAST_DELEGATE_OneParam(FOnAddControllerExtensionColumn, TArray<FName>& /*OutColumnNames*/);
 
+/**
+ * Callback for custom resolution of exposed property sources.
+ * Returns true if the property is resolved / dealt with.
+ * When successfully resolving, bound objects can be replaced.
+ */
+DECLARE_DELEGATE_RetVal_ThreeParams(bool, FRCPropertyResolver, const TSharedRef<FRemoteControlProperty>& /* In Property */, TArray<UObject*>& /* In/Out Bound Objects */, TSharedPtr<FPropertyPath>& /* Out Property Path */)
+
 /**
  * A Remote Control module that allows exposing objects and properties from the editor.
  */
@@ -310,6 +318,16 @@ public:
 	 */
 	virtual void UnregisterWidgetFactoryForType(UScriptStruct* RemoteControlEntityType) = 0;
 
+	/**
+	 * Registers a property source resolution callback.
+	 */
+	virtual void RegisterPropertyResolver(FName InResolverId, const FRCPropertyResolver& InResolver) = 0;
+
+	/**
+	 * Unregisters a property source resolution callback.
+	 */
+	virtual void UnregisterPropertyResolver(FName InResolverId) = 0;
+
 	/** The section of EditorPerProjectUserSettings in which to save filter settings */
 	static const FString SettingsIniSection;
 
