Remote Control: Added Select Source button to the action list
- Also added icon to right click on exposed property right click menu.

diff --git a/Engine/Plugins/VirtualProduction/RemoteControl/Source/RemoteControlUI/Private/UI/Action/SRCActionPanel.cpp b/Engine/Plugins/VirtualProduction/RemoteControl/Source/RemoteControlUI/Private/UI/Action/SRCActionPanel.cpp
--- a/Engine/Plugins/VirtualProduction/RemoteControl/Source/RemoteControlUI/Private/UI/Action/SRCActionPanel.cpp
+++ b/Engine/Plugins/VirtualProduction/RemoteControl/Source/RemoteControlUI/Private/UI/Action/SRCActionPanel.cpp
@@ -321,11 +321,38 @@ void SRCActionPanel::CreateActionPanelWidgets()
 			];
 	}
 
+	// Add all selected button
+	if (!SelectSourceButton.IsValid())
+	{
+		bIsAnyToolbarItemCreated = true;
+
+		SelectSourceButton = SNew(SButton)
+			.AddMetaData<FTagMetaData>(FTagMetaData(TEXT("Select Action Property Source")))
+			.HAlign(HAlign_Center)
+			.VAlign(VAlign_Center)
+			.ForegroundColor(FSlateColor::UseForeground())
+			.ButtonStyle(&RCPanelStyle->FlatButtonStyle)
+			.ToolTipText(LOCTEXT("RCSelectSourceToolTip", "Select the source of the property this action is bound to."))
+			.OnClicked(this, &SRCActionPanel::OnSelectSource)
+			.Visibility(this, &SRCActionPanel::HandleSelectSourceVisibility)
+			[
+				SNew(SBox)
+				.WidthOverride(RCPanelStyle->IconSize.X)
+				.HeightOverride(RCPanelStyle->IconSize.Y)
+				[
+					SNew(SImage)
+					.ColorAndOpacity(FSlateColor::UseForeground())
+					.Image(FAppStyle::GetBrush("Icons.SelectInViewport"))
+				]
+			];
+	}
+
 	// Adding header toolbar items
 	if (bIsAnyToolbarItemCreated)
 	{
 		ActionDockPanel->ClearHeaderToolbarItems();
 		ActionDockPanel->AddHeaderToolbarItem(EToolbar::Left, AddNewActionButton.ToSharedRef());
+		ActionDockPanel->AddHeaderToolbarItem(EToolbar::Right, SelectSourceButton.ToSharedRef());
 		ActionDockPanel->AddHeaderToolbarItem(EToolbar::Right, AddAllSelectedActionsButton.ToSharedRef());
 		ActionDockPanel->AddHeaderToolbarItem(EToolbar::Right, AddAllActionsButton.ToSharedRef());
 	}
@@ -610,6 +637,16 @@ FReply SRCActionPanel::OnAddAllSelectedFields()
 	return FReply::Handled();
 }
 
+FReply SRCActionPanel::OnSelectSource()
+{
+	if (ActionPanelList.IsValid())
+	{
+		ActionPanelList->SelectSource_Execute();
+	}
+
+	return FReply::Handled();
+}
+
 bool SRCActionPanel::IsListFocused() const
 {
 	return ActionPanelList.IsValid() && ActionPanelList->IsListFocused();
@@ -929,4 +966,17 @@ EVisibility SRCActionPanel::HandleAddAllButtonVisibility() const
 	return EVisibility::Collapsed;
 }
 
+EVisibility SRCActionPanel::HandleSelectSourceVisibility() const
+{
+	if (ActionPanelList.IsValid())
+	{
+		if (ActionPanelList->NumSelectedLogicItems() > 0)
+		{
+			return EVisibility::Visible;
+		}
+	}
+
+	return EVisibility::Hidden;
+}
+
 #undef LOCTEXT_NAMESPACE
diff --git a/Engine/Plugins/VirtualProduction/RemoteControl/Source/RemoteControlUI/Private/UI/Action/SRCActionPanel.h b/Engine/Plugins/VirtualProduction/RemoteControl/Source/RemoteControlUI/Private/UI/Action/SRCActionPanel.h
--- a/Engine/Plugins/VirtualProduction/RemoteControl/Source/RemoteControlUI/Private/UI/Action/SRCActionPanel.h
+++ b/Engine/Plugins/VirtualProduction/RemoteControl/Source/RemoteControlUI/Private/UI/Action/SRCActionPanel.h
@@ -105,6 +105,9 @@ private:
 	/** Determines the visibility Add All Button. */
 	EVisibility HandleAddAllButtonVisibility() const;
 
+	/** Determines the visibility Select Source Button. */
+	EVisibility HandleSelectSourceVisibility() const;
+
 	/** 
 	* Behaviour selection change listener.
 	* Updates the list of actions from newly selected Behaviour
@@ -147,6 +150,9 @@ private:
 	/** Handles click event for Add All Selected button; Adds all selected actions for the active Remote Control Preset*/
 	FReply OnAddAllSelectedFields();
 
+	/** Handles click event for Select Source button; Selects the actor/component in the scene whose properties are bound in this action*/
+	FReply OnSelectSource();
+
 private:
 	/** Helper widget for behavior details. */
 	static TSharedRef<SBox> GetNoneSelectedWidget();
@@ -182,6 +188,9 @@ private:
 	/** Add all selected actions button */
 	TSharedPtr<SWidget> AddAllSelectedActionsButton;
 
+	/** Select source for action button */
+	TSharedPtr<SWidget> SelectSourceButton;
+
 	/** Action Major Panel */
 	TSharedPtr<SRCMajorPanel> ActionsPanel;
 
diff --git a/Engine/Plugins/VirtualProduction/RemoteControl/Source/RemoteControlUI/Private/UI/Action/SRCActionPanelList.h b/Engine/Plugins/VirtualProduction/RemoteControl/Source/RemoteControlUI/Private/UI/Action/SRCActionPanelList.h
--- a/Engine/Plugins/VirtualProduction/RemoteControl/Source/RemoteControlUI/Private/UI/Action/SRCActionPanelList.h
+++ b/Engine/Plugins/VirtualProduction/RemoteControl/Source/RemoteControlUI/Private/UI/Action/SRCActionPanelList.h
@@ -279,6 +279,30 @@ public:
 		}
 	}
 
+	void SelectSource_Execute() override
+	{
+		TArray<TSharedPtr<ActionType>> SelectedItems = ListView->GetSelectedItems();
+		TArray<UObject*> ObjectsToSelect;
+
+		for (const TSharedPtr<ActionType>& SelectedAction : SelectedItems)
+		{
+			if (SelectedAction.IsValid())
+			{
+				if (URCAction* Action = SelectedAction->GetAction())
+				{
+					if (URemoteControlPreset* Preset = Action->PresetWeakPtr.Get())
+					{
+						if (TSharedPtr<FRemoteControlEntity> Entity = Preset->GetExposedEntity(Action->ExposedFieldId).Pin())
+						{
+							FRemoteControlUIModule::Get().SelectSource(Entity);
+							return;
+						}
+					}
+				}
+			}
+		}
+	}
+
 private:
 
 	/** OnGenerateRow delegate for the Actions List View */
@@ -597,30 +621,6 @@ private:
 		SelectedActionItem = InItem;
 	}
 
-	void SelectSource_Execute()
-	{
-		TArray<TSharedPtr<ActionType>> SelectedItems = ListView->GetSelectedItems();
-		TArray<UObject*> ObjectsToSelect;
-
-		for (const TSharedPtr<ActionType>& SelectedAction : SelectedItems)
-		{
-			if (SelectedAction.IsValid())
-			{
-				if (URCAction* Action = SelectedAction->GetAction())
-				{
-					if (URemoteControlPreset* Preset = Action->PresetWeakPtr.Get())
-					{
-						if (TSharedPtr<FRemoteControlEntity> Entity = Preset->GetExposedEntity(Action->ExposedFieldId).Pin())
-						{
-							FRemoteControlUIModule::Get().SelectSource(Entity);
-							return;
-						}
-					}
-				}
-			}
-		}
-	}
-
 	void UpdateEntityUsage()
 	{
 		if (TSharedPtr<SRCActionPanel> ActionPanel = ActionPanelWeakPtr.Pin())
diff --git a/Engine/Plugins/VirtualProduction/RemoteControl/Source/RemoteControlUI/Private/UI/BaseLogicUI/SRCLogicPanelListBase.h b/Engine/Plugins/VirtualProduction/RemoteControl/Source/RemoteControlUI/Private/UI/BaseLogicUI/SRCLogicPanelListBase.h
--- a/Engine/Plugins/VirtualProduction/RemoteControl/Source/RemoteControlUI/Private/UI/BaseLogicUI/SRCLogicPanelListBase.h
+++ b/Engine/Plugins/VirtualProduction/RemoteControl/Source/RemoteControlUI/Private/UI/BaseLogicUI/SRCLogicPanelListBase.h
@@ -65,6 +65,8 @@ public:
 	{
 		return RemoteControlPanelWeakPtr.Pin();
 	}
+	
+	virtual void SelectSource_Execute() {}
 
 private:
 
diff --git a/Engine/Plugins/VirtualProduction/RemoteControl/Source/RemoteControlUI/Private/UI/SRCPanelExposedEntity.cpp b/Engine/Plugins/VirtualProduction/RemoteControl/Source/RemoteControlUI/Private/UI/SRCPanelExposedEntity.cpp
--- a/Engine/Plugins/VirtualProduction/RemoteControl/Source/RemoteControlUI/Private/UI/SRCPanelExposedEntity.cpp
+++ b/Engine/Plugins/VirtualProduction/RemoteControl/Source/RemoteControlUI/Private/UI/SRCPanelExposedEntity.cpp
@@ -90,7 +90,7 @@ TSharedPtr<SWidget> SRCPanelExposedEntity::GetContextMenu()
 	MenuBuilder.AddMenuEntry(FRemoteControlCommands::Get().RenameEntity, NAME_None, TAttribute<FText>(), TAttribute<FText>(), FSlateIcon(FAppStyle::GetAppStyleSetName(), TEXT("GenericCommands.Rename")));
 	MenuBuilder.AddMenuEntry(FRemoteControlCommands::Get().ChangePropId, NAME_None, TAttribute<FText>(), TAttribute<FText>(), FSlateIcon(FAppStyle::GetAppStyleSetName(), TEXT("Icons.Edit")));
 	MenuBuilder.AddMenuEntry(FRemoteControlCommands::Get().DeleteEntity, NAME_None, TAttribute<FText>(), TAttribute<FText>(), FSlateIcon(FAppStyle::GetAppStyleSetName(), TEXT("GenericCommands.Delete")));
-	MenuBuilder.AddMenuEntry(FRemoteControlCommands::Get().SelectSource);
+	MenuBuilder.AddMenuEntry(FRemoteControlCommands::Get().SelectSource, NAME_None, TAttribute<FText>(), TAttribute<FText>(), FSlateIcon(FAppStyle::GetAppStyleSetName(), TEXT("Icons.SelectInViewport")));
 
 	MenuBuilder.EndSection();
 
