Remote Control: Changes to object additions in the source property loops.

diff --git a/Engine/Plugins/VirtualProduction/RemoteControl/Source/RemoteControlUI/Private/RemoteControlUIModule.cpp b/Engine/Plugins/VirtualProduction/RemoteControl/Source/RemoteControlUI/Private/RemoteControlUIModule.cpp
--- a/Engine/Plugins/VirtualProduction/RemoteControl/Source/RemoteControlUI/Private/RemoteControlUIModule.cpp
+++ b/Engine/Plugins/VirtualProduction/RemoteControl/Source/RemoteControlUI/Private/RemoteControlUIModule.cpp
@@ -1318,6 +1318,8 @@ TSharedPtr<FPropertyPath> FRemoteControlUIModule::ResolveProperty(const TSharedR
 		}
 	}
 
+	TArray<UObject*> AdditionalObjects;
+
 	// By default, convert non-actor/actor components to the owning actors.
 	for (auto It = InOutBoundObjects.CreateIterator(); It; ++It)
 	{
@@ -1333,11 +1335,13 @@ TSharedPtr<FPropertyPath> FRemoteControlUIModule::ResolveProperty(const TSharedR
 		// When we encounter a non-component object, 
 		if (!BindingObject->IsA<UActorComponent>())
 		{
-			InOutBoundObjects.Add(OwnerActor);
+			AdditionalObjects.Add(OwnerActor);
 			It.RemoveCurrent();
 		}
 	}
 
+	InOutBoundObjects.Append(AdditionalObjects);
+
 	// And then return the default path resolution
 	return InProperty->FieldPathInfo.ToPropertyPath();
 }
