Remote Control: Changed entity bind numbers over to runtime objects.

diff --git a/Engine/Plugins/VirtualProduction/Avalanche/Source/AvalancheMedia/Internal/Playable/AvaPlayableRemoteControl.h b/Engine/Plugins/VirtualProduction/Avalanche/Source/AvalancheMedia/Internal/Playable/AvaPlayableRemoteControl.h
--- a/Engine/Plugins/VirtualProduction/Avalanche/Source/AvalancheMedia/Internal/Playable/AvaPlayableRemoteControl.h
+++ b/Engine/Plugins/VirtualProduction/Avalanche/Source/AvalancheMedia/Internal/Playable/AvaPlayableRemoteControl.h
@@ -55,9 +55,6 @@ namespace UE::AvaPlayableRemoteControl
 	AVALANCHEMEDIA_API EAvaPlayableRemoteControlResult SetValueOfController(URCVirtualPropertyBase* InController, const FString& InValue);
 	AVALANCHEMEDIA_API EAvaPlayableRemoteControlResult SetValueOfController(URCVirtualPropertyBase* InController, const FString& InValue, bool bInBehaviorsEnabled);
 
-	/** Retrieve the entity ids that are controlled by the given controller. */
-	AVALANCHEMEDIA_API bool GetEntitiesControlledByController(const URemoteControlPreset* InRemoteControlPreset, const URCVirtualPropertyBase* InVirtualProperty, TSet<FGuid>& OutEntityIds);
-
 	inline bool Failed(EAvaPlayableRemoteControlResult InResult)
 	{
 		// Negative result codes are failure.
diff --git a/Engine/Plugins/VirtualProduction/Avalanche/Source/AvalancheMedia/Private/Playable/AvaPlayableRemoteControl.cpp b/Engine/Plugins/VirtualProduction/Avalanche/Source/AvalancheMedia/Private/Playable/AvaPlayableRemoteControl.cpp
--- a/Engine/Plugins/VirtualProduction/Avalanche/Source/AvalancheMedia/Private/Playable/AvaPlayableRemoteControl.cpp
+++ b/Engine/Plugins/VirtualProduction/Avalanche/Source/AvalancheMedia/Private/Playable/AvaPlayableRemoteControl.cpp
@@ -326,54 +326,6 @@ EAvaPlayableRemoteControlResult UE::AvaPlayableRemoteControl::SetValueOfControll
 	return SetValueOfController(InController, InValue);
 }
 
-bool UE::AvaPlayableRemoteControl::GetEntitiesControlledByController(const URemoteControlPreset* InRemoteControlPreset, const URCVirtualPropertyBase* InVirtualProperty, TSet<FGuid>& OutEntityIds)
-{
-	const URCController* Controller = Cast<URCController>(InVirtualProperty);
-	if (!IsValid(Controller))
-	{
-		return false;
-	}
-
-	for (const URCBehaviour* Behavior : Controller->Behaviours)
-	{
-		if (!Behavior)
-		{
-			continue;
-		}
-
-		// Special case: SetAssetByPath behaviors have an extra 'TargetEntity' that these behaviors control outside of RC Actions.
-		if (const URCSetAssetByPathBehaviour* SetAssetByPathBehavior = Cast<URCSetAssetByPathBehaviour>(Behavior))
-		{
-			OutEntityIds.Add(SetAssetByPathBehavior->GetTargetEntityId());
-		}
-
-		for (const URCAction* Action : Behavior->ActionContainer->GetActions())
-		{
-			if (Action)
-			{
-				if (const URCPropertyIdAction* IdentityAction = Cast<URCPropertyIdAction>(Action))
-				{
-					if (const URemoteControlPropertyIdRegistry* IdentityRegistry = InRemoteControlPreset->GetPropertyIdRegistry())
-					{
-						OutEntityIds.Append(IdentityRegistry->GetEntityIdsForPropertyId(IdentityAction->PropertyId));
-
-						// Explicit support for Sub-PropertyId
-						for (const TPair<FPropertyIdContainerKey, TObjectPtr<URCVirtualPropertySelfContainer>>& PropertyContainer : IdentityAction->PropertySelfContainer)
-						{
-							OutEntityIds.Append(IdentityRegistry->GetEntityIdsForPropertyId(PropertyContainer.Key.PropertyId));
-						}
-					}
-				}
-				else
-				{
-					OutEntityIds.Add(Action->ExposedFieldId);
-				}
-			}
-		}
-	}
-	return true;
-}
-
 FString UE::AvaPlayableRemoteControl::EnumToString(EAvaPlayableRemoteControlResult InValue)
 {
 	return StaticEnum<EAvaPlayableRemoteControlResult>()->GetNameStringByValue(static_cast<int64>(InValue));
diff --git a/Engine/Plugins/VirtualProduction/Avalanche/Source/AvalancheMedia/Private/Playable/AvaPlayableRemoteControlPresetInfo.cpp b/Engine/Plugins/VirtualProduction/Avalanche/Source/AvalancheMedia/Private/Playable/AvaPlayableRemoteControlPresetInfo.cpp
--- a/Engine/Plugins/VirtualProduction/Avalanche/Source/AvalancheMedia/Private/Playable/AvaPlayableRemoteControlPresetInfo.cpp
+++ b/Engine/Plugins/VirtualProduction/Avalanche/Source/AvalancheMedia/Private/Playable/AvaPlayableRemoteControlPresetInfo.cpp
@@ -2,6 +2,7 @@
 
 #include "Playable/AvaPlayableRemoteControlPresetInfo.h"
 
+#include "Controller/RCCustomControllerUtilities.h"
 #include "Playable/AvaPlayableRemoteControl.h"
 #include "Playable/AvaPlayableRemoteControlValues.h"
 #include "RCVirtualProperty.h"
@@ -29,7 +30,7 @@ void FAvaPlayableRemoteControlPresetInfo::Refresh(const URemoteControlPreset* In
 			
 			TmpControlledEntities.Reset();
 			
-			if (!UE::AvaPlayableRemoteControl::GetEntitiesControlledByController(InRemoteControlPreset, Controller, TmpControlledEntities))
+			if (!UE::RCCustomControllers::GetEntitiesControlledByController(InRemoteControlPreset, Controller, TmpControlledEntities))
 			{
 				UE_LOG(LogAvaPlayableRemoteControl, Warning, TEXT("Failed to get controlled entities for controller \"%s\" (id:%s)."),
 					*Controller->DisplayName.ToString(), *Controller->Id.ToString());
diff --git a/Engine/Plugins/VirtualProduction/Avalanche/Source/AvalancheMediaEditor/Private/Rundown/DetailsView/RemoteControl/Controllers/AvaRundownPageControllerContextMenu.cpp b/Engine/Plugins/VirtualProduction/Avalanche/Source/AvalancheMediaEditor/Private/Rundown/DetailsView/RemoteControl/Controllers/AvaRundownPageControllerContextMenu.cpp
--- a/Engine/Plugins/VirtualProduction/Avalanche/Source/AvalancheMediaEditor/Private/Rundown/DetailsView/RemoteControl/Controllers/AvaRundownPageControllerContextMenu.cpp
+++ b/Engine/Plugins/VirtualProduction/Avalanche/Source/AvalancheMediaEditor/Private/Rundown/DetailsView/RemoteControl/Controllers/AvaRundownPageControllerContextMenu.cpp
@@ -3,6 +3,7 @@
 #include "AvaRundownPageControllerContextMenu.h"
 #include "AvaRundownRCControllerItem.h"
 #include "Controller/RCController.h"
+#include "Controller/RCCustomControllerUtilities.h"
 #include "IAvaMediaEditorModule.h"
 #include "Playable/AvaPlayableRemoteControl.h"
 #include "Rundown/AvaRundownCommands.h"
@@ -216,7 +217,7 @@ void FAvaRundownPageControllerContextMenu::ResetValuesToDefaults(const bool bInU
 							if (const URemoteControlPreset* Preset = Controller->PresetWeakPtr.Get())
 							{
 								TSet<FGuid> EntityIds;
-								GetEntitiesControlledByController(Preset, Controller, EntityIds);
+								UE::RCCustomControllers::GetEntitiesControlledByController(Preset, Controller, EntityIds);
 								SAvaRundownPageRemoteControlProps::SaveRemoteControlEntitiesToPage(Preset, EntityIds, Rundown, CurrentContext->GetRundownPageId());
 							}
 						}
diff --git a/Engine/Plugins/VirtualProduction/Avalanche/Source/AvalancheMediaEditor/Private/Rundown/DetailsView/RemoteControl/Properties/SAvaRundownPageRemoteControlProps.cpp b/Engine/Plugins/VirtualProduction/Avalanche/Source/AvalancheMediaEditor/Private/Rundown/DetailsView/RemoteControl/Properties/SAvaRundownPageRemoteControlProps.cpp
--- a/Engine/Plugins/VirtualProduction/Avalanche/Source/AvalancheMediaEditor/Private/Rundown/DetailsView/RemoteControl/Properties/SAvaRundownPageRemoteControlProps.cpp
+++ b/Engine/Plugins/VirtualProduction/Avalanche/Source/AvalancheMediaEditor/Private/Rundown/DetailsView/RemoteControl/Properties/SAvaRundownPageRemoteControlProps.cpp
@@ -1,6 +1,7 @@
 // Copyright Epic Games, Inc. All Rights Reserved.
 
 #include "SAvaRundownPageRemoteControlProps.h"
+#include "Controller/RCCustomControllerUtilities.h"
 #include "Item/AvaRundownRCFieldItem.h"
 #include "Modules/ModuleManager.h"
 #include "Playable/AvaPlayableRemoteControl.h"
@@ -494,7 +495,7 @@ void SAvaRundownPageRemoteControlProps::OnRemoteControlControllerModified(URemot
 	TSet<FGuid> EntityIds;
 	for (const FGuid& ControllerId : InModifiedControllerIds)
 	{
-		UE::AvaPlayableRemoteControl::GetEntitiesControlledByController(InPreset, InPreset->GetController(ControllerId), EntityIds);
+		UE::RCCustomControllers::GetEntitiesControlledByController(InPreset, InPreset->GetController(ControllerId), EntityIds);
 	}
 
 	// If a controller changed, we need to propagate the refresh of the field's widgets.
diff --git a/Engine/Plugins/VirtualProduction/RemoteControl/Source/RemoteControl/Private/RCVirtualProperty.cpp b/Engine/Plugins/VirtualProduction/RemoteControl/Source/RemoteControl/Private/RCVirtualProperty.cpp
--- a/Engine/Plugins/VirtualProduction/RemoteControl/Source/RemoteControl/Private/RCVirtualProperty.cpp
+++ b/Engine/Plugins/VirtualProduction/RemoteControl/Source/RemoteControl/Private/RCVirtualProperty.cpp
@@ -2,6 +2,7 @@
 
 #include "RCVirtualProperty.h"
 
+#include "Containers/Set.h"
 #include "IStructDeserializerBackend.h"
 #include "RCVirtualPropertyContainer.h"
 #include "RemoteControlPreset.h"
diff --git a/Engine/Plugins/VirtualProduction/RemoteControl/Source/RemoteControl/Private/RCVirtualPropertyContainer.cpp b/Engine/Plugins/VirtualProduction/RemoteControl/Source/RemoteControl/Private/RCVirtualPropertyContainer.cpp
--- a/Engine/Plugins/VirtualProduction/RemoteControl/Source/RemoteControl/Private/RCVirtualPropertyContainer.cpp
+++ b/Engine/Plugins/VirtualProduction/RemoteControl/Source/RemoteControl/Private/RCVirtualPropertyContainer.cpp
@@ -193,7 +193,7 @@ URCVirtualPropertyInContainer* URCVirtualPropertyContainerBase::DuplicateVirtual
 	return nullptr;
 }
 
-bool URCVirtualPropertyContainerBase::HasVirtualProperty(URCVirtualPropertyBase* InVirtualProperty) const
+bool URCVirtualPropertyContainerBase::HasVirtualProperty(const URCVirtualPropertyBase* InVirtualProperty) const
 {
 	return InVirtualProperty && VirtualProperties.Contains(InVirtualProperty);
 }
diff --git a/Engine/Plugins/VirtualProduction/RemoteControl/Source/RemoteControl/Private/RemoteControlPreset.cpp b/Engine/Plugins/VirtualProduction/RemoteControl/Source/RemoteControl/Private/RemoteControlPreset.cpp
--- a/Engine/Plugins/VirtualProduction/RemoteControl/Source/RemoteControl/Private/RemoteControlPreset.cpp
+++ b/Engine/Plugins/VirtualProduction/RemoteControl/Source/RemoteControl/Private/RemoteControlPreset.cpp
@@ -725,6 +725,8 @@ void URemoteControlPreset::PostLoad()
 	RemoveUnusedBindings();
 
 	PropertyIdRegistry->Initialize();
+
+	UpdateEntityUsage();
 }
 
 void URemoteControlPreset::PostDuplicate(bool bDuplicateForPIE)
@@ -916,7 +918,7 @@ URCVirtualPropertyInContainer* URemoteControlPreset::AddController(TSubclassOf<U
 	return NewController;
 }
 
-bool URemoteControlPreset::HasController(URCVirtualPropertyBase* InController) const
+bool URemoteControlPreset::HasController(const URCVirtualPropertyBase* InController) const
 {
 	if (!ensure(ControllerContainer))
 	{
@@ -2828,4 +2830,44 @@ void URemoteControlPreset::FRCPropertyWatcher::SetLastFrameValue(const FRCFieldR
 	ResolvedData.Field->CopyCompleteValue(LastFrameValue.GetData(), NewValueAddress);
 }
 
+void URemoteControlPreset::UpdateEntityUsage()
+{
+	TMap<FGuid, TSet<UObject*>> EntityToBehavior;
+
+	if (ControllerContainer)
+	{
+		for (URCVirtualPropertyBase* VirtualProperty : ControllerContainer->VirtualProperties)
+		{
+			VirtualProperty->UpdateEntityUsage(this, EntityToBehavior);
+		}
+	}
+
+	BoundEntityIds.Empty(EntityToBehavior.Num());
+
+	for (const TPair<FGuid, TSet<UObject*>>& EntityPair : EntityToBehavior)
+	{
+		BoundEntityIds.Add(EntityPair.Key, EntityPair.Value.Num());
+	}
+}
+
+bool URemoteControlPreset::IsEntityBound(const FGuid& InEntityGuid) const
+{
+	if (InEntityGuid.IsValid())
+	{
+		return BoundEntityIds.FindRef(InEntityGuid) > 0;
+	}
+
+	return false;
+}
+
+int32 URemoteControlPreset::GetEntityBoundCount(const FGuid& InEntityGuid) const
+{
+	if (InEntityGuid.IsValid())
+	{
+		return BoundEntityIds.FindRef(InEntityGuid);
+	}
+
+	return 0;
+}
+
 #undef LOCTEXT_NAMESPACE /* RemoteControlPreset */ 
diff --git a/Engine/Plugins/VirtualProduction/RemoteControl/Source/RemoteControl/Public/RCVirtualProperty.h b/Engine/Plugins/VirtualProduction/RemoteControl/Source/RemoteControl/Public/RCVirtualProperty.h
--- a/Engine/Plugins/VirtualProduction/RemoteControl/Source/RemoteControl/Public/RCVirtualProperty.h
+++ b/Engine/Plugins/VirtualProduction/RemoteControl/Source/RemoteControl/Public/RCVirtualProperty.h
@@ -2,6 +2,8 @@
  
 #pragma once
 
+#include "Containers/ContainersFwd.h"
+#include "Containers/Map.h"
 #include "StructUtils/PropertyBag.h"
 #include "UObject/Object.h"
  
@@ -351,6 +353,8 @@ public:
 	/** Updates the display index of the virtual property (transaction aware) */
 	void SetDisplayIndex(const int32 InDisplayIndex);
 
+	virtual void UpdateEntityUsage(const URemoteControlPreset* InPreset, TMap<FGuid, TSet<UObject*>>& InEntityToBehavior) const {}
+
 public:
 	/** Unique property name */
 	UPROPERTY()
diff --git a/Engine/Plugins/VirtualProduction/RemoteControl/Source/RemoteControl/Public/RCVirtualPropertyContainer.h b/Engine/Plugins/VirtualProduction/RemoteControl/Source/RemoteControl/Public/RCVirtualPropertyContainer.h
--- a/Engine/Plugins/VirtualProduction/RemoteControl/Source/RemoteControl/Public/RCVirtualPropertyContainer.h
+++ b/Engine/Plugins/VirtualProduction/RemoteControl/Source/RemoteControl/Public/RCVirtualPropertyContainer.h
@@ -76,7 +76,7 @@ public:
 	/**
 	 * Returns whether this container contains the given property.
 	 */
-	bool HasVirtualProperty(URCVirtualPropertyBase* InVirtualProperty) const;
+	bool HasVirtualProperty(const URCVirtualPropertyBase* InVirtualProperty) const;
 
 	/**
 	 * Removes a property from the container by name if it exists.
diff --git a/Engine/Plugins/VirtualProduction/RemoteControl/Source/RemoteControl/Public/RemoteControlPreset.h b/Engine/Plugins/VirtualProduction/RemoteControl/Source/RemoteControl/Public/RemoteControlPreset.h
--- a/Engine/Plugins/VirtualProduction/RemoteControl/Source/RemoteControl/Public/RemoteControlPreset.h
+++ b/Engine/Plugins/VirtualProduction/RemoteControl/Source/RemoteControl/Public/RemoteControlPreset.h
@@ -775,7 +775,7 @@ public:
 	URCVirtualPropertyInContainer* AddController(TSubclassOf<URCVirtualPropertyInContainer> InPropertyClass, const EPropertyBagPropertyType InValueType, UObject* InValueTypeObject = nullptr, const FName InPropertyName = NAME_None);
 
 	/** Returns true if the preset contains this controller. */
-	bool HasController(URCVirtualPropertyBase* InController) const;
+	bool HasController(const URCVirtualPropertyBase* InController) const;
 
 	/** Removes a given Virtual Property (by Name) from the Remote Control preset */
 	bool RemoveController(const FName& InPropertyName);
@@ -913,6 +913,21 @@ public:
 	/** Returns how this preset modifies properties when changed by protocols */
 	ERCModifyOperationFlags GetModifyOperationFlagsForProtocols() const { return ModifyOperationFlagsForProtocols; }
 
+	/** Updates the usage of entities by controllers. Specifically the number of controlling behaviors. */
+	void UpdateEntityUsage();
+
+	/** Returns a map of the number of times exposed entities are used by behaviors. */
+	const TMap<FGuid, int32>& GetEntityUsage() const
+	{
+		return BoundEntityIds;
+	}
+
+	/** Returns true if an entity is bound to a controller/behavior. */
+	bool IsEntityBound(const FGuid& InEntityGuid) const;
+
+	/** Returns the number of behaviors an entity is bound to. */
+	int32 GetEntityBoundCount(const FGuid& InEntityGuid) const;
+
 private:
 
 	/** Find a binding that has the same boundobjectmap but that currently points to the object passed as argument. */
@@ -1051,6 +1066,9 @@ private:
 	/** Used for OnObjectPropertyChanged to avoid being called by itself. */
 	bool bReentryGuard = false;
 
+	/** Transient list of the currently bound entity ids and the number of times they are bound. */
+	TMap<FGuid, int32> BoundEntityIds;
+
 	/** Holds manager that handles rebinding unbound entities upon load or map change. */
 	TPimplPtr<FRemoteControlPresetRebindingManager> RebindingManager;
 
diff --git a/Engine/Plugins/VirtualProduction/RemoteControl/Source/RemoteControlLogic/Private/Action/RCActionContainer.cpp b/Engine/Plugins/VirtualProduction/RemoteControl/Source/RemoteControlLogic/Private/Action/RCActionContainer.cpp
--- a/Engine/Plugins/VirtualProduction/RemoteControl/Source/RemoteControlLogic/Private/Action/RCActionContainer.cpp
+++ b/Engine/Plugins/VirtualProduction/RemoteControl/Source/RemoteControlLogic/Private/Action/RCActionContainer.cpp
@@ -271,3 +271,16 @@ void URCActionContainer::EmptyActions()
 {
 	Actions.Empty();
 }
+
+void URCActionContainer::UpdateEntityUsage(const URCBehaviour* InBehavior, TMap<FGuid, TSet<UObject*>>& InEntityToBehavior) const
+{
+	for (URCActionContainer* ActionContainer : ActionContainers)
+	{
+		ActionContainer->UpdateEntityUsage(InBehavior, InEntityToBehavior);
+	}
+
+	for (URCAction* Action : Actions)
+	{
+		Action->UpdateEntityUsage(InBehavior, InEntityToBehavior);
+	}
+}
diff --git a/Engine/Plugins/VirtualProduction/RemoteControl/Source/RemoteControlLogic/Private/Action/RCPropertyAction.cpp b/Engine/Plugins/VirtualProduction/RemoteControl/Source/RemoteControlLogic/Private/Action/RCPropertyAction.cpp
--- a/Engine/Plugins/VirtualProduction/RemoteControl/Source/RemoteControlLogic/Private/Action/RCPropertyAction.cpp
+++ b/Engine/Plugins/VirtualProduction/RemoteControl/Source/RemoteControlLogic/Private/Action/RCPropertyAction.cpp
@@ -40,6 +40,23 @@ void URCAction::UpdateEntityIds(const TMap<FGuid, FGuid>& InEntityIdMap)
 	}
 }
 
+void URCAction::UpdateEntityUsage(const URCBehaviour* InBehavior, TMap<FGuid, TSet<UObject*>>& InEntityToBehavior) const
+{
+	URCBehaviour* Behavior = GetParentBehaviour();
+
+	if (!Behavior || Behavior != InBehavior || !InBehavior->HasAction(this))
+	{
+		return;
+	}
+
+	if (!ExposedFieldId.IsValid())
+	{
+		return;
+	}
+
+	InEntityToBehavior.FindOrAdd(ExposedFieldId).Add(Behavior);
+}
+
 void URCAction::NotifyActionValueChanged()
 {
 	if (URCBehaviour* Behaviour = GetParentBehaviour())
diff --git a/Engine/Plugins/VirtualProduction/RemoteControl/Source/RemoteControlLogic/Private/Behaviour/RCBehaviour.cpp b/Engine/Plugins/VirtualProduction/RemoteControl/Source/RemoteControlLogic/Private/Behaviour/RCBehaviour.cpp
--- a/Engine/Plugins/VirtualProduction/RemoteControl/Source/RemoteControlLogic/Private/Behaviour/RCBehaviour.cpp
+++ b/Engine/Plugins/VirtualProduction/RemoteControl/Source/RemoteControlLogic/Private/Behaviour/RCBehaviour.cpp
@@ -75,7 +75,7 @@ int32 URCBehaviour::GetNumActions() const
 	return ActionContainer->GetActions().Num();
 }
 
-bool URCBehaviour::HasAction(URCAction* InAction) const
+bool URCBehaviour::HasAction(const URCAction* InAction) const
 {
 	return InAction && ActionContainer->GetActions().Contains(InAction);
 }
@@ -149,6 +149,21 @@ URCBehaviourNode* URCBehaviour::GetBehaviourNode()
 	return CachedBehaviourNode;
 }
 
+void URCBehaviour::UpdateEntityUsage(const URCController* InController, TMap<FGuid, TSet<UObject*>>& InEntityToBehavior) const
+{
+	URCController* Controller = ControllerWeakPtr.Get();
+
+	if (!Controller || Controller != InController || !Controller->HasBehavior(this))
+	{
+		return;
+	}
+
+	if (ActionContainer)
+	{
+		ActionContainer->UpdateEntityUsage(this, InEntityToBehavior);
+	}	
+}
+
 void URCBehaviour::ExecuteInternal(const TSet<TObjectPtr<URCAction>>& InActionsToExecute)
 {
 	const URCBehaviourNode* BehaviourNode = GetBehaviourNode();
diff --git a/Engine/Plugins/VirtualProduction/RemoteControl/Source/RemoteControlLogic/Private/Controller/RCController.cpp b/Engine/Plugins/VirtualProduction/RemoteControl/Source/RemoteControlLogic/Private/Controller/RCController.cpp
--- a/Engine/Plugins/VirtualProduction/RemoteControl/Source/RemoteControlLogic/Private/Controller/RCController.cpp
+++ b/Engine/Plugins/VirtualProduction/RemoteControl/Source/RemoteControlLogic/Private/Controller/RCController.cpp
@@ -100,9 +100,9 @@ int32 URCController::RemoveBehaviour(const FGuid InBehaviourId)
 	return RemovedCount;
 }
 
-bool URCController::HasBehaviour(URCBehaviour* InBehaviour) const
+bool URCController::HasBehavior(const URCBehaviour* InBehavior) const
 {
-	return InBehaviour && Behaviours.Contains(InBehaviour);
+	return InBehavior && Behaviours.Contains(InBehavior);
 }
 
 void URCController::EmptyBehaviours()
@@ -151,4 +151,19 @@ URCBehaviour* URCController::DuplicateBehaviour(URCController* InController, URC
 	return NewBehaviour;
 }
 
+void URCController::UpdateEntityUsage(const URemoteControlPreset* InPreset, TMap<FGuid, TSet<UObject*>>& InEntityToBehavior) const
+{
+	URemoteControlPreset* Preset = PresetWeakPtr.Get();
+
+	if (!Preset || Preset != InPreset || !Preset->HasController(this))
+	{
+		return;
+	}
+
+	for (URCBehaviour* Behaviour : Behaviours)
+	{
+		Behaviour->UpdateEntityUsage(this, InEntityToBehavior);
+	}
+}
+
 #undef LOCTEXT_NAMESPACE /* RCController */ 
\ No newline at end of file
diff --git a/Engine/Plugins/VirtualProduction/RemoteControl/Source/RemoteControlLogic/Private/Controller/RCCustomControllerUtilities.cpp b/Engine/Plugins/VirtualProduction/RemoteControl/Source/RemoteControlLogic/Private/Controller/RCCustomControllerUtilities.cpp
--- a/Engine/Plugins/VirtualProduction/RemoteControl/Source/RemoteControlLogic/Private/Controller/RCCustomControllerUtilities.cpp
+++ b/Engine/Plugins/VirtualProduction/RemoteControl/Source/RemoteControlLogic/Private/Controller/RCCustomControllerUtilities.cpp
@@ -1,8 +1,15 @@
-﻿// Copyright Epic Games, Inc. All Rights Reserved.
+// Copyright Epic Games, Inc. All Rights Reserved.
 
 #include "Controller/RCCustomControllerUtilities.h"
+#include "Action/RCAction.h"
+#include "Action/RCPropertyIdAction.h"
+#include "Behaviour/Builtin/Path/RCSetAssetByPathBehaviour.h"
+#include "Behaviour/RCBehaviour.h"
+#include "Controller/RCController.h"
 #include "RCVirtualProperty.h"
 #include "RCVirtualPropertyContainer.h"
+#include "RemoteControlPreset.h"
+#include "RemoteControlPropertyIdRegistry.h"
 
 namespace UE::RCCustomControllers
 {
@@ -83,3 +90,146 @@ FName UE::RCCustomControllers::GetUniqueNameForController(const URCVirtualProper
 	// Worst case - we return the plain type name.
 	return FName(CustomTypeName);
 }
+
+bool UE::RCCustomControllers::GetEntitiesControlledByController(const URemoteControlPreset* InRemoteControlPreset, const URCVirtualPropertyBase* InVirtualProperty, TSet<FGuid>& OutEntityIds)
+{
+	if (!InRemoteControlPreset)
+	{
+		return false;
+	}
+
+	const URCController* Controller = Cast<URCController>(InVirtualProperty);
+
+	if (!IsValid(Controller))
+	{
+		return false;
+	}
+
+	for (const URCBehaviour* Behavior : Controller->Behaviours)
+	{
+		if (!Behavior)
+		{
+			continue;
+		}
+
+		// Special case: SetAssetByPath behaviors have an extra 'TargetEntity' that these behaviors control outside of RC Actions.
+		if (const URCSetAssetByPathBehaviour* SetAssetByPathBehavior = Cast<URCSetAssetByPathBehaviour>(Behavior))
+		{
+			OutEntityIds.Add(SetAssetByPathBehavior->GetTargetEntityId());
+		}
+
+		for (const URCAction* Action : Behavior->ActionContainer->GetActions())
+		{
+			if (!Action)
+			{
+				continue;
+			}
+
+			if (const URCPropertyIdAction* IdentityAction = Cast<URCPropertyIdAction>(Action))
+			{
+				if (const URemoteControlPropertyIdRegistry* IdentityRegistry = InRemoteControlPreset->GetPropertyIdRegistry())
+				{
+					OutEntityIds.Append(IdentityRegistry->GetEntityIdsForPropertyId(IdentityAction->PropertyId));
+
+					// Explicit support for Sub-PropertyId
+					for (const TPair<FPropertyIdContainerKey, TObjectPtr<URCVirtualPropertySelfContainer>>& PropertyContainer : IdentityAction->PropertySelfContainer)
+					{
+						OutEntityIds.Append(IdentityRegistry->GetEntityIdsForPropertyId(PropertyContainer.Key.PropertyId));
+					}
+				}
+			}
+			else
+			{
+				OutEntityIds.Add(Action->ExposedFieldId);
+			}
+		}
+	}
+	return true;
+}
+
+bool UE::RCCustomControllers::GetControllersOfEntity(const URemoteControlPreset* InRemoteControlPreset, const FGuid& InEntity, TSet<const URCVirtualPropertyBase*>& OutControllers)
+{
+	if (!InRemoteControlPreset)
+	{
+		return false;
+	}
+
+	if (URCVirtualPropertyContainerBase* ControllerContainer = InRemoteControlPreset->GetControllerContainer())
+	{
+		for (const URCVirtualPropertyBase* VirtualProperty : ControllerContainer->VirtualProperties)
+		{
+			if (const URCController* Controller = Cast<URCController>(VirtualProperty))
+			{
+				bool bFoundEntity = false;
+
+				for (const URCBehaviour* Behavior : Controller->Behaviours)
+				{
+					if (!Behavior)
+					{
+						continue;
+					}
+
+					// Special case: SetAssetByPath behaviors have an extra 'TargetEntity' that these behaviors control outside of RC Actions.
+					if (const URCSetAssetByPathBehaviour* SetAssetByPathBehavior = Cast<URCSetAssetByPathBehaviour>(Behavior))
+					{
+						if (SetAssetByPathBehavior->GetTargetEntityId() == InEntity)
+						{
+							bFoundEntity = true;
+							break;
+						}
+					}
+
+					if (!bFoundEntity)
+					{
+						for (const URCAction* Action : Behavior->ActionContainer->GetActions())
+						{
+							if (!Action)
+							{
+								continue;
+							}
+
+							if (const URCPropertyIdAction* IdentityAction = Cast<URCPropertyIdAction>(Action))
+							{
+								if (const URemoteControlPropertyIdRegistry* IdentityRegistry = InRemoteControlPreset->GetPropertyIdRegistry())
+								{
+									if (IdentityRegistry->GetEntityIdsForPropertyId(IdentityAction->PropertyId).Contains(InEntity))
+									{
+										bFoundEntity = true;
+										break;
+									}
+
+									// Explicit support for Sub-PropertyId
+									for (const TPair<FPropertyIdContainerKey, TObjectPtr<URCVirtualPropertySelfContainer>>& PropertyContainer : IdentityAction->PropertySelfContainer)
+									{
+										if (IdentityRegistry->GetEntityIdsForPropertyId(PropertyContainer.Key.PropertyId).Contains(InEntity))
+										{
+											bFoundEntity = true;
+											break;
+										}
+									}
+								}
+							}
+							else if (Action->ExposedFieldId == InEntity)
+							{
+								bFoundEntity = true;
+								break;
+							}
+						}
+					}
+
+					if (bFoundEntity)
+					{
+						break;
+					}
+				}
+
+				if (bFoundEntity)
+				{
+					OutControllers.Add(Controller);
+				}
+			}
+		}
+	}
+
+	return true;
+}
diff --git a/Engine/Plugins/VirtualProduction/RemoteControl/Source/RemoteControlLogic/Public/Action/RCAction.h b/Engine/Plugins/VirtualProduction/RemoteControl/Source/RemoteControlLogic/Public/Action/RCAction.h
--- a/Engine/Plugins/VirtualProduction/RemoteControl/Source/RemoteControlLogic/Public/Action/RCAction.h
+++ b/Engine/Plugins/VirtualProduction/RemoteControl/Source/RemoteControlLogic/Public/Action/RCAction.h
@@ -36,6 +36,8 @@ public:
 	 */
 	virtual void UpdateEntityIds(const TMap<FGuid, FGuid>& InEntityIdMap);
 
+	virtual void UpdateEntityUsage(const URCBehaviour* InBehavior, TMap<FGuid, TSet<UObject*>>& InEntityToBehavior) const;
+
 public:
 	/** Exposed Property or Function field Id*/
 	UPROPERTY()
diff --git a/Engine/Plugins/VirtualProduction/RemoteControl/Source/RemoteControlLogic/Public/Action/RCActionContainer.h b/Engine/Plugins/VirtualProduction/RemoteControl/Source/RemoteControlLogic/Public/Action/RCActionContainer.h
--- a/Engine/Plugins/VirtualProduction/RemoteControl/Source/RemoteControlLogic/Public/Action/RCActionContainer.h
+++ b/Engine/Plugins/VirtualProduction/RemoteControl/Source/RemoteControlLogic/Public/Action/RCActionContainer.h
@@ -98,6 +98,8 @@ public:
 	/** Call the given function on all actions of the container. */
 	void ForEachAction(TFunctionRef<void(URCAction*)> InActionFunction, bool bInRecursive);
 
+	void UpdateEntityUsage(const URCBehaviour* InBehavior, TMap<FGuid, TSet<UObject*>>& InEntityToBehavior) const;
+
 private:
 	/** Add remote control property action  */
 	URCPropertyAction* AddPropertyAction(const TSharedRef<const FRemoteControlProperty> InRemoteControlProperty);
diff --git a/Engine/Plugins/VirtualProduction/RemoteControl/Source/RemoteControlLogic/Public/Behaviour/RCBehaviour.h b/Engine/Plugins/VirtualProduction/RemoteControl/Source/RemoteControlLogic/Public/Behaviour/RCBehaviour.h
--- a/Engine/Plugins/VirtualProduction/RemoteControl/Source/RemoteControlLogic/Public/Behaviour/RCBehaviour.h
+++ b/Engine/Plugins/VirtualProduction/RemoteControl/Source/RemoteControlLogic/Public/Behaviour/RCBehaviour.h
@@ -57,7 +57,7 @@ public:
 	/** Get number of action associated with behaviour */
 	int32 GetNumActions() const;
 
-	bool HasAction(URCAction* InAction) const;
+	bool HasAction(const URCAction* InAction) const;
 
 	/** Whether we can create an action pertaining to a given remote control field for the current behaviour */
 	virtual bool CanHaveActionForField(const TSharedPtr<FRemoteControlField> InRemoteControlField) const;
@@ -105,6 +105,8 @@ public:
 	 * Or just return cached one
 	 */
 	URCBehaviourNode* GetBehaviourNode();
+	
+	void UpdateEntityUsage(const URCController* InController, TMap<FGuid, TSet<UObject*>>& InEntityToBehavior) const;
 
 protected:
 	/** Execute all the action if not provided a valid Action otherwise will only execute the given action */
diff --git a/Engine/Plugins/VirtualProduction/RemoteControl/Source/RemoteControlLogic/Public/Controller/RCController.h b/Engine/Plugins/VirtualProduction/RemoteControl/Source/RemoteControlLogic/Public/Controller/RCController.h
--- a/Engine/Plugins/VirtualProduction/RemoteControl/Source/RemoteControlLogic/Public/Controller/RCController.h
+++ b/Engine/Plugins/VirtualProduction/RemoteControl/Source/RemoteControlLogic/Public/Controller/RCController.h
@@ -45,7 +45,7 @@ public:
 	virtual int32 RemoveBehaviour(const FGuid InBehaviourId);
 
 	/** Returns true if the behaviour exist in this controller. */
-	bool HasBehaviour(URCBehaviour* InBehaviour) const;
+	bool HasBehavior(const URCBehaviour* InBehavior) const;
 
 	/** Removes all behaviours. */
 	virtual void EmptyBehaviours();
@@ -65,6 +65,10 @@ public:
 	/** Delegate that notifies changes to the list of behaviours*/
 	FOnBehaviourListModified OnBehaviourListModified;
 
+	//~ Begin URCVirtualPropertyBase
+	virtual void UpdateEntityUsage(const URemoteControlPreset* InPreset, TMap<FGuid, TSet<UObject*>>& InEntityToBehavior) const override;
+	//~ End URCVirtualPropertyBase
+
 public:
 	/** Set of the behaviours */
 	UPROPERTY()
diff --git a/Engine/Plugins/VirtualProduction/RemoteControl/Source/RemoteControlLogic/Public/Controller/RCCustomControllerUtilities.h b/Engine/Plugins/VirtualProduction/RemoteControl/Source/RemoteControlLogic/Public/Controller/RCCustomControllerUtilities.h
--- a/Engine/Plugins/VirtualProduction/RemoteControl/Source/RemoteControlLogic/Public/Controller/RCCustomControllerUtilities.h
+++ b/Engine/Plugins/VirtualProduction/RemoteControl/Source/RemoteControlLogic/Public/Controller/RCCustomControllerUtilities.h
@@ -1,8 +1,11 @@
 // Copyright Epic Games, Inc. All Rights Reserved.
 
 #pragma once
+
+#include "Containers/ContainersFwd.h"
 #include "StructUtils/PropertyBag.h"
 
+class URemoteControlPreset;
 class URCVirtualPropertyBase;
 class URCVirtualPropertyInContainer;
 
@@ -36,6 +39,10 @@ namespace UE::RCCustomControllers
 	 * Useful when creating a new controller.
 	 */
 	REMOTECONTROLLOGIC_API FName GetUniqueNameForController(const URCVirtualPropertyInContainer* InController);
+
+	REMOTECONTROLLOGIC_API bool GetEntitiesControlledByController(const URemoteControlPreset* InRemoteControlPreset, const URCVirtualPropertyBase* InVirtualProperty, TSet<FGuid>& OutEntityIds);
+
+	REMOTECONTROLLOGIC_API bool GetControllersOfEntity(const URemoteControlPreset* InRemoteControlPreset, const FGuid& InEntity, TSet<const URCVirtualPropertyBase*>& OutControllers);
 	
 	inline static FName CustomControllerNameKey = "CustomControllerTypeName";
 	inline static FName CustomControllerExecuteOnLoadNameKey = "CustomControllerExecuteOnLoad";
diff --git a/Engine/Plugins/VirtualProduction/RemoteControl/Source/RemoteControlUI/Private/UI/SRCPanelExposedEntitiesList.cpp b/Engine/Plugins/VirtualProduction/RemoteControl/Source/RemoteControlUI/Private/UI/SRCPanelExposedEntitiesList.cpp
--- a/Engine/Plugins/VirtualProduction/RemoteControl/Source/RemoteControlUI/Private/UI/SRCPanelExposedEntitiesList.cpp
+++ b/Engine/Plugins/VirtualProduction/RemoteControl/Source/RemoteControlUI/Private/UI/SRCPanelExposedEntitiesList.cpp
@@ -198,7 +198,7 @@ void SRCPanelExposedEntitiesList::Construct(const FArguments& InArgs, URemoteCon
 
 	bFilterApplicationRequested = false;
 	bSearchRequested = false;
-	bUpdateEntityUsage = true; // Perform this on creation
+	bUpdateEntityUsage = false; // Perform this on creation
 	
 	RecallListSettings();
 
@@ -552,7 +552,11 @@ void SRCPanelExposedEntitiesList::Tick(const FGeometry& AllottedGeometry, const
 
 	if (bUpdateEntityUsage)
 	{
-		RunUpdateEntityUsage();
+		if (Preset.IsValid())
+		{
+			Preset->UpdateEntityUsage();
+		}
+
 		bUpdateEntityUsage = false;
 	}
 }
@@ -889,92 +893,14 @@ void SRCPanelExposedEntitiesList::SelectSource_Execute()
 	FRemoteControlUIModule::Get().SelectSource(Entity);
 }
 
-void SRCPanelExposedEntitiesList::RunUpdateEntityUsage()
+void SRCPanelExposedEntitiesList::UpdateEntityUsage_Execute()
 {
 	if (!Preset.IsValid())
 	{
 		return;
 	}
 
-	TArray<UObject*> Subobjects;
-	GetObjectsWithOuter(Preset.Get(), Subobjects);
-
-	TMap<FGuid, TSet<URCBehaviour*>> EntityToBehaviour;
-	TSet<UObject*> ValidContainers;
-	TSet<UObject*> InvalidContainers;
-
-	for (UObject* Object : Subobjects)
-	{
-		URCAction* Action = Cast<URCAction>(Object);
-
-		if (!Action || !Action->ExposedFieldId.IsValid())
-		{
-			continue;
-		}
-
-		// Check to see if the action is still part of the behaviour.
-		URCBehaviour* Behaviour = Action->GetParentBehaviour();
-			
-		if (!Behaviour || InvalidContainers.Contains(Behaviour) || !Behaviour->HasAction(Action))
-		{
-			continue;
-		}
-				
-		if (!ValidContainers.Contains(Behaviour))
-		{
-			// Make sure the behaviour is still part of the controller.
-			URCController* Controller = Behaviour->ControllerWeakPtr.Get();
-
-			if (!Controller || InvalidContainers.Contains(Controller) || !Controller->HasBehaviour(Behaviour))
-			{
-				InvalidContainers.Add(Behaviour);
-				continue;
-			}
-
-			if (!ValidContainers.Contains(Controller))
-			{
-				// Make sure the controller is still part of the preset
-				URemoteControlPreset* ControllerPreset = Controller->PresetWeakPtr.Get();
-
-				if (!ControllerPreset || ControllerPreset != Preset.Get() || !ControllerPreset->HasController(Controller))
-				{
-					InvalidContainers.Add(Behaviour);
-					InvalidContainers.Add(Controller);
-					continue;
-				}
-
-				ValidContainers.Add(Behaviour);
-				ValidContainers.Add(Controller);
-			}
-		}
-
-		//BoundEntityIds.FindOrAdd(Action->ExposedFieldId) += 1;
-		EntityToBehaviour.FindOrAdd(Action->ExposedFieldId).Add(Behaviour);
-	}
-
-	BoundEntityIds.Empty(EntityToBehaviour.Num());
-
-	for (const TPair<FGuid, TSet<URCBehaviour*>>& EntityPair : EntityToBehaviour)
-	{
-		BoundEntityIds.Add(EntityPair.Key, EntityPair.Value.Num());
-	}
-
-	for (const TPair <FGuid, TSharedPtr<SRCPanelTreeNode>>& Node : FieldWidgetMap)
-	{
-		const TSharedPtr<SRCPanelTreeNode>& NodeWidget = Node.Value;
-
-		if (NodeWidget.IsValid())
-		{
-			if (NodeWidget->GetRCType() == SRCPanelTreeNode::ENodeType::Field && NodeWidget->GetRCId().IsValid())
-			{
-				NodeWidget->SetFieldInUseCount(BoundEntityIds.FindRef(NodeWidget->GetRCId()));
-			}
-			else
-			{
-				NodeWidget->SetFieldInUseCount(0);
-			}
-		}
-	}
+	Preset->UpdateEntityUsage();
 }
 
 void SRCPanelExposedEntitiesList::UpdateEntityUsage()
@@ -982,11 +908,6 @@ void SRCPanelExposedEntitiesList::UpdateEntityUsage()
 	bUpdateEntityUsage = true;
 }
 
-bool SRCPanelExposedEntitiesList::IsEntityBound(const FGuid& InEntityGuid) const
-{
-	return InEntityGuid.IsValid() && BoundEntityIds.Contains(InEntityGuid);
-}
-
 void SRCPanelExposedEntitiesList::OnObjectPropertyChange(UObject* InObject, FPropertyChangedEvent& InChangeEvent)
 {
 	static constexpr EPropertyChangeType::Type TypesNeedingRefresh = EPropertyChangeType::ArrayAdd | EPropertyChangeType::ArrayClear | EPropertyChangeType::ArrayRemove | EPropertyChangeType::ValueSet | EPropertyChangeType::Duplicate;
@@ -1261,15 +1182,6 @@ TSharedRef<ITableRow> SRCPanelExposedEntitiesList::OnGenerateRow(TSharedPtr<SRCP
 	}
 	else
 	{
-		if (Node->GetRCType() == SRCPanelTreeNode::ENodeType::Field && Node->GetRCId().IsValid())
-		{
-			Node->SetFieldInUseCount(BoundEntityIds.FindRef(Node->GetRCId()));
-		}
-		else
-		{
-			Node->SetFieldInUseCount(0);
-		}
-
 		constexpr float LeftPadding = 3.f;
 		const FMargin Margin = Node->GetRCType() == SRCPanelTreeNode::FieldChild ? FMargin(LeftPadding + 10.f, 1.f, 1.f, 1.f) : FMargin(LeftPadding, 1.f, 1.f, 1.f);
 		Node->OnPropertyIdRenamed().BindSP(this, &SRCPanelExposedEntitiesList::OnPropertyIdRenamed, Node);
diff --git a/Engine/Plugins/VirtualProduction/RemoteControl/Source/RemoteControlUI/Private/UI/SRCPanelExposedEntitiesList.h b/Engine/Plugins/VirtualProduction/RemoteControl/Source/RemoteControlUI/Private/UI/SRCPanelExposedEntitiesList.h
--- a/Engine/Plugins/VirtualProduction/RemoteControl/Source/RemoteControlUI/Private/UI/SRCPanelExposedEntitiesList.h
+++ b/Engine/Plugins/VirtualProduction/RemoteControl/Source/RemoteControlUI/Private/UI/SRCPanelExposedEntitiesList.h
@@ -137,9 +137,6 @@ public:
 	/** Update the In Use flag on entities. */
 	void UpdateEntityUsage();
 
-	/** Returns true if the entity was bound in the last binding check. */
-	bool IsEntityBound(const FGuid& InEntityGuid) const;
-
 private:
 	/** Handles label to be shown in the entity list header.  */
 	FText HandleEntityListHeaderLabel() const;
@@ -282,8 +279,7 @@ private:
 	/** Selects the source of the currently selected entity. */
 	void SelectSource_Execute();
 
-	/** Updates the usage of entities in this list. */
-	void RunUpdateEntityUsage();
+	void UpdateEntityUsage_Execute();
 
 private:
 	/** Command list to use if specified */
@@ -340,8 +336,6 @@ private:
 	bool bSearchRequested;
 	/** If true, the entity list will have its in use flag updated next frame. */
 	bool bUpdateEntityUsage;
-	/** Transient list of the currently bound entity ids and the number of times they are bound. */
-	TMap<FGuid, int32> BoundEntityIds;
 	/** Panel Style reference. */
 	const FRCPanelStyle* RCPanelStyle;
 	/** Holds the header row of entities list. */
diff --git a/Engine/Plugins/VirtualProduction/RemoteControl/Source/RemoteControlUI/Private/UI/SRCPanelExposedField.cpp b/Engine/Plugins/VirtualProduction/RemoteControl/Source/RemoteControlUI/Private/UI/SRCPanelExposedField.cpp
--- a/Engine/Plugins/VirtualProduction/RemoteControl/Source/RemoteControlUI/Private/UI/SRCPanelExposedField.cpp
+++ b/Engine/Plugins/VirtualProduction/RemoteControl/Source/RemoteControlUI/Private/UI/SRCPanelExposedField.cpp
@@ -942,7 +942,7 @@ TSharedRef<SWidget> SRCPanelExposedField::CreateInUseWidget()
 				{
 					if (TSharedPtr<SRCPanelExposedField> This = StaticCastSharedPtr<SRCPanelExposedField>(ThisWeak.Pin()))
 					{
-						if (This->GetFieldInUseCount() > 1)
+						if (This->GetBoundCount() > 1)
 						{
 							return FAppStyle::Get().GetBrush("Icons.WarningWithColor");
 						}
@@ -956,9 +956,9 @@ TSharedRef<SWidget> SRCPanelExposedField::CreateInUseWidget()
 				{
 					if (TSharedPtr<SRCPanelExposedField> This = StaticCastSharedPtr<SRCPanelExposedField>(ThisWeak.Pin()))
 					{
-						if (This->GetFieldInUseCount() > 1)
+						if (This->GetBoundCount() > 1)
 						{
-							return LOCTEXT("PropertyBoundMultiple", "WARNING: Property is bound to multiple behaviours!");
+							return LOCTEXT("PropertyBoundMultiple", "WARNING: Property is bound to multiple behaviors!");
 						}
 					}
 
@@ -970,7 +970,7 @@ TSharedRef<SWidget> SRCPanelExposedField::CreateInUseWidget()
 				{
 					if (TSharedPtr<SRCPanelExposedField> This = StaticCastSharedPtr<SRCPanelExposedField>(ThisWeak.Pin()))
 					{
-						if (This->GetFieldInUseCount() > 0)
+						if (This->GetBoundCount() > 0)
 						{
 							return EVisibility::Visible;
 						}
@@ -984,6 +984,24 @@ TSharedRef<SWidget> SRCPanelExposedField::CreateInUseWidget()
 	return SNullWidget::NullWidget;
 }
 
+int32 SRCPanelExposedField::GetBoundCount() const
+{
+	if (TSharedPtr<FRemoteControlField> Field = WeakField.Pin())
+	{
+		if (URemoteControlPreset* NodePreset = Field->GetOwner())
+		{
+			const FGuid Id = Field->GetId();
+
+			if (Id.IsValid())
+			{
+				return NodePreset->GetEntityBoundCount(Id);
+			}
+		}
+	}
+
+	return 0;
+}
+
 void SRCPanelFieldChildNode::Construct(const FArguments& InArgs, const TSharedRef<IDetailTreeNode>& InNode, FRCColumnSizeData InColumnSizeData)
 {
 	TArray<TSharedRef<IDetailTreeNode>> ChildNodes;
diff --git a/Engine/Plugins/VirtualProduction/RemoteControl/Source/RemoteControlUI/Private/UI/SRCPanelExposedField.h b/Engine/Plugins/VirtualProduction/RemoteControl/Source/RemoteControlUI/Private/UI/SRCPanelExposedField.h
--- a/Engine/Plugins/VirtualProduction/RemoteControl/Source/RemoteControlUI/Private/UI/SRCPanelExposedField.h
+++ b/Engine/Plugins/VirtualProduction/RemoteControl/Source/RemoteControlUI/Private/UI/SRCPanelExposedField.h
@@ -110,8 +110,12 @@ private:
 	/** Handles committing a property id label. */
 	void OnPropertyIdTextCommitted(const FText& InText, ETextCommit::Type InCommitInfo);
 
+	/** Creates a widget which says whether this field is bound to a controller. */
 	TSharedRef<SWidget> CreateInUseWidget();
 
+	/** Returns the number of behaviors bound to this field. */
+	int32 GetBoundCount() const;
+
 private:
 	/** Weak pointer to the underlying RC Field. */
 	TWeakPtr<FRemoteControlField> WeakField;
diff --git a/Engine/Plugins/VirtualProduction/RemoteControl/Source/RemoteControlUI/Private/UI/SRCPanelTreeNode.h b/Engine/Plugins/VirtualProduction/RemoteControl/Source/RemoteControlUI/Private/UI/SRCPanelTreeNode.h
--- a/Engine/Plugins/VirtualProduction/RemoteControl/Source/RemoteControlUI/Private/UI/SRCPanelTreeNode.h
+++ b/Engine/Plugins/VirtualProduction/RemoteControl/Source/RemoteControlUI/Private/UI/SRCPanelTreeNode.h
@@ -117,16 +117,6 @@ struct SRCPanelTreeNode : public SCompoundWidget, public IHasProtocolExtensibili
 	/** Executed when the PropertyId of the node is changed. */
 	FOnPropertyIdRenamed& OnPropertyIdRenamed() { return OnPropertyIdRenamedDelegate; }
 
-	int32 GetFieldInUseCount() const
-	{ 
-		return FieldUseCount;
-	}
-
-	void SetFieldInUseCount(int32 InUseCount)
-	{ 
-		FieldUseCount = InUseCount;
-	}
-
 protected:
 	struct FMakeNodeWidgetArgs
 	{
@@ -165,9 +155,6 @@ protected:
 	/** Holds the row's columns' width. */
 	FRCColumnSizeData ColumnSizeData;
 
-	/** If true, an action is bound to this node's entity. */
-	int32 FieldUseCount = 0;
-
 private:
 
 	/** Cached widget of PropertyId identifier. */
