Remote Control: Horde Issue 983587

diff --git a/Engine/Plugins/VirtualProduction/RemoteControl/Source/RemoteControlUI/Private/UI/RCPanelWidgetRegistry.cpp b/Engine/Plugins/VirtualProduction/RemoteControl/Source/RemoteControlUI/Private/UI/RCPanelWidgetRegistry.cpp
--- a/Engine/Plugins/VirtualProduction/RemoteControl/Source/RemoteControlUI/Private/UI/RCPanelWidgetRegistry.cpp
+++ b/Engine/Plugins/VirtualProduction/RemoteControl/Source/RemoteControlUI/Private/UI/RCPanelWidgetRegistry.cpp
@@ -186,6 +186,11 @@ TSharedPtr<IDetailTreeNode> FRCPanelWidgetRegistry::GetStructTreeNode(const TSha
 
 void FRCPanelWidgetRegistry::Refresh(UObject* InObject, bool bForceRefresh)
 {
+	if (!InObject)
+	{
+		return;
+	}
+
 	if (TSharedPtr<IPropertyRowGenerator>* Generator = ObjectToRowGenerator.Find({InObject}))
 	{
 		const TArray<TWeakObjectPtr<UObject>>& SelectedObjects = (*Generator)->GetSelectedObjects();
diff --git a/Engine/Plugins/VirtualProduction/RemoteControl/Source/RemoteControlUI/Private/UI/SRCPanelExposedEntitiesList.cpp b/Engine/Plugins/VirtualProduction/RemoteControl/Source/RemoteControlUI/Private/UI/SRCPanelExposedEntitiesList.cpp
--- a/Engine/Plugins/VirtualProduction/RemoteControl/Source/RemoteControlUI/Private/UI/SRCPanelExposedEntitiesList.cpp
+++ b/Engine/Plugins/VirtualProduction/RemoteControl/Source/RemoteControlUI/Private/UI/SRCPanelExposedEntitiesList.cpp
@@ -910,6 +910,11 @@ void SRCPanelExposedEntitiesList::UpdateEntityUsage()
 
 void SRCPanelExposedEntitiesList::OnObjectPropertyChange(UObject* InObject, FPropertyChangedEvent& InChangeEvent)
 {
+	if (!InObject)
+	{
+		return;
+	}
+
 	static constexpr EPropertyChangeType::Type TypesNeedingRefresh = EPropertyChangeType::ArrayAdd | EPropertyChangeType::ArrayClear | EPropertyChangeType::ArrayRemove | EPropertyChangeType::ValueSet | EPropertyChangeType::Duplicate;
 	auto IsRelevantProperty = [](FFieldClass* PropertyClass)
 	{
