Remote Control: Added Catgory "controller".

diff --git a/Engine/Plugins/VirtualProduction/RemoteControl/Source/RemoteControl/Public/RCVirtualPropertyContainer.h b/Engine/Plugins/VirtualProduction/RemoteControl/Source/RemoteControl/Public/RCVirtualPropertyContainer.h
--- a/Engine/Plugins/VirtualProduction/RemoteControl/Source/RemoteControl/Public/RCVirtualPropertyContainer.h
+++ b/Engine/Plugins/VirtualProduction/RemoteControl/Source/RemoteControl/Public/RCVirtualPropertyContainer.h
@@ -4,6 +4,8 @@
 
 #include "CoreMinimal.h"
 
+#include "Internationalization/Text.h"
+#include "Misc/Guid.h"
 #include "StructUtils/PropertyBag.h"
 
 #include "RCVirtualPropertyContainer.generated.h"
@@ -13,6 +15,24 @@ class URCVirtualPropertyBase;
 class URemoteControlPreset;
 class FStructOnScope;
 
+/** 
+ * Describes a category into which properties can be placed.
+ */
+USTRUCT()
+struct FRCVirtualPropertyCategory
+{
+	GENERATED_BODY()
+
+	UPROPERTY()
+	FGuid Id = FGuid();
+	
+	UPROPERTY()
+	FText Name = FText::GetEmpty();
+
+	UPROPERTY()
+	int32 DisplayIndex = INDEX_NONE;
+};
+
 /**
  * Container for more then one Virtual Property
  */
@@ -203,6 +223,11 @@ public:
 	/** Delegate that notifies changes to the virtual property container*/
 	FOnVirtualPropertyContainerModified OnVirtualPropertyContainerModifiedDelegate;
 
+#if WITH_EDITORONLY_DATA
+	UPROPERTY()
+	TArray<FRCVirtualPropertyCategory> Categories;
+#endif
+
 private:
 	/** Map of Controller Name to GUID. */
 	UPROPERTY(Transient)
diff --git a/Engine/Plugins/VirtualProduction/RemoteControl/Source/RemoteControlLogic/RemoteControlLogic.Build.cs b/Engine/Plugins/VirtualProduction/RemoteControl/Source/RemoteControlLogic/RemoteControlLogic.Build.cs
--- a/Engine/Plugins/VirtualProduction/RemoteControl/Source/RemoteControlLogic/RemoteControlLogic.Build.cs
+++ b/Engine/Plugins/VirtualProduction/RemoteControl/Source/RemoteControlLogic/RemoteControlLogic.Build.cs
@@ -21,6 +21,7 @@ public class RemoteControlLogic : ModuleRules
 				"HTTP",
 				"ImageCore",
 				"RemoteControl",
+				"RemoteControlCommon",
 				"RenderCore",
 				"RHI",
 				"Serialization",
diff --git a/Engine/Plugins/VirtualProduction/RemoteControl/Source/RemoteControlUI/Private/RemoteControlUIModule.cpp b/Engine/Plugins/VirtualProduction/RemoteControl/Source/RemoteControlUI/Private/RemoteControlUIModule.cpp
--- a/Engine/Plugins/VirtualProduction/RemoteControl/Source/RemoteControlUI/Private/RemoteControlUIModule.cpp
+++ b/Engine/Plugins/VirtualProduction/RemoteControl/Source/RemoteControlUI/Private/RemoteControlUIModule.cpp
@@ -9,6 +9,7 @@
 #include "Commands/RemoteControlCommands.h"
 #include "Components/SkeletalMeshComponent.h"
 #include "Components/StaticMeshComponent.h"
+#include "Controller/RCController.h"
 #include "Controller/RCCustomControllerUtilities.h"
 #include "DetailRowMenuContext.h"
 #include "DetailTreeNode.h"
@@ -38,8 +39,8 @@
 #include "UI/Action/SRCActionPanel.h"
 #include "UI/Behaviour/Builtin/Conditional/SRCBehaviourConditional.h"
 #include "UI/Behaviour/SRCBehaviourPanel.h"
-#include "UI/Controller/CustomControllers/SCustomTextureControllerWidget.h"
 #include "UI/Controller/CustomControllers/SCustomStringListControllerWidget.h"
+#include "UI/Controller/CustomControllers/SCustomTextureControllerWidget.h"
 #include "UI/Customizations/FPassphraseCustomization.h"
 #include "UI/Customizations/NetworkAddressCustomization.h"
 #include "UI/Customizations/RCAssetPathElementCustomization.h"
@@ -54,6 +55,7 @@
 #include "UI/SRCPanelExposedField.h"
 #include "UI/SRemoteControlPanel.h"
 #include "Widgets/DeclarativeSyntaxSupport.h"
+#include "Widgets/SNullWidget.h"
 
 #define LOCTEXT_NAMESPACE "RemoteControlUI"
 
@@ -1220,13 +1222,14 @@ TSharedPtr<SWidget> FRemoteControlUIModule::CreateCustomControllerWidget(URCVirt
 {
 	TSharedPtr<SWidget> CustomControllerWidget;
 	const FString& CustomName = UE::RCCustomControllers::GetCustomControllerTypeName(InController);
+	
 	if (CustomName == UE::RCCustomControllers::CustomTextureControllerName)
 	{
-		CustomControllerWidget = SNew(SCustomTextureControllerWidget, InController, InOriginalPropertyHandle);
+		CustomControllerWidget = SNew(SCustomTextureControllerWidget, Cast<URCController>(InController), InOriginalPropertyHandle);
 	}
 	else if (CustomName == UE::RCCustomControllers::CustomStringListControllerName)
 	{
-		CustomControllerWidget = SNew(SCustomStringListControllerWidget, InController, InOriginalPropertyHandle);
+		CustomControllerWidget = SNew(SCustomStringListControllerWidget, Cast<URCController>(InController), InOriginalPropertyHandle);
 	}
 
 	return CustomControllerWidget;
diff --git a/Engine/Plugins/VirtualProduction/RemoteControl/Source/RemoteControlUI/Private/UI/Controller/CustomControllers/SCustomTextureControllerWidget.cpp b/Engine/Plugins/VirtualProduction/RemoteControl/Source/RemoteControlUI/Private/UI/Controller/CustomControllers/SCustomTextureControllerWidget.cpp
--- a/Engine/Plugins/VirtualProduction/RemoteControl/Source/RemoteControlUI/Private/UI/Controller/CustomControllers/SCustomTextureControllerWidget.cpp
+++ b/Engine/Plugins/VirtualProduction/RemoteControl/Source/RemoteControlUI/Private/UI/Controller/CustomControllers/SCustomTextureControllerWidget.cpp
@@ -21,7 +21,7 @@
 
 #define LOCTEXT_NAMESPACE "ExternalTextureControllerWidget"
 
-void SCustomTextureControllerWidget::Construct(const FArguments& InArgs, URCVirtualPropertyBase* InController, const TSharedPtr<IPropertyHandle>& InOriginalPropertyHandle)
+void SCustomTextureControllerWidget::Construct(const FArguments& InArgs, URCController* InController, const TSharedPtr<IPropertyHandle>& InOriginalPropertyHandle)
 {
 	if (!InOriginalPropertyHandle.IsValid())
 	{
diff --git a/Engine/Plugins/VirtualProduction/RemoteControl/Source/RemoteControlUI/Private/UI/Controller/CustomControllers/SCustomTextureControllerWidget.h b/Engine/Plugins/VirtualProduction/RemoteControl/Source/RemoteControlUI/Private/UI/Controller/CustomControllers/SCustomTextureControllerWidget.h
--- a/Engine/Plugins/VirtualProduction/RemoteControl/Source/RemoteControlUI/Private/UI/Controller/CustomControllers/SCustomTextureControllerWidget.h
+++ b/Engine/Plugins/VirtualProduction/RemoteControl/Source/RemoteControlUI/Private/UI/Controller/CustomControllers/SCustomTextureControllerWidget.h
@@ -6,7 +6,7 @@
 #include "Widgets/SCompoundWidget.h"
 
 class IPropertyHandle;
-class URCVirtualPropertyBase;
+class URCController;
 class UTexture2D;
 
 class SCustomTextureControllerWidget : public SCompoundWidget
@@ -22,7 +22,7 @@ public:
 	 * Constructs this widget with InArgs
 	 * @param InOriginalPropertyHandle Original PropertyHandle of this CustomWidget
 	 */
-	void Construct(const FArguments& InArgs, URCVirtualPropertyBase* InController, const TSharedPtr<IPropertyHandle>& InOriginalPropertyHandle);
+	void Construct(const FArguments& InArgs, URCController* InController, const TSharedPtr<IPropertyHandle>& InOriginalPropertyHandle);
 
 private:
 	/** Handle toggling Texture Controller type change: External vs. Internal */
diff --git a/Engine/Plugins/VirtualProduction/RemoteControl/Source/RemoteControlUI/Private/UI/Controller/RCCategoryModel.cpp b/Engine/Plugins/VirtualProduction/RemoteControl/Source/RemoteControlUI/Private/UI/Controller/RCCategoryModel.cpp
new file mode 100644
--- /dev/null
+++ b/Engine/Plugins/VirtualProduction/RemoteControl/Source/RemoteControlUI/Private/UI/Controller/RCCategoryModel.cpp
@@ -0,0 +1,140 @@
+// Copyright Epic Games, Inc. All Rights Reserved.
+
+#include "RCCategoryModel.h"
+#include "RCVirtualPropertyContainer.h"
+#include "RemoteControlPreset.h"
+#include "ScopedTransaction.h"
+#include "UI/SRemoteControlPanel.h"
+#include "Widgets/Text/SInlineEditableTextBlock.h"
+
+#define LOCTEXT_NAMESPACE "FRCCategoryModel"
+
+FRCCategoryModel::FRCCategoryModel(const FGuid& InCateogryId, const TSharedPtr<SRemoteControlPanel> InRemoteControlPanel)
+	: FRCLogicModeBase(InRemoteControlPanel)
+	, Id(InCateogryId)
+	, CachedIndex(INDEX_NONE)
+{
+}
+
+TOptional<FRCVirtualPropertyCategory> FRCCategoryModel::GetCategory() const
+{
+	CheckCachedIndex();
+
+	if (CachedIndex != INDEX_NONE)
+	{
+		return GetPreset()->GetControllerContainer()->Categories[CachedIndex];
+	}
+
+	return {};
+}
+
+void FRCCategoryModel::Initialize()
+{
+	SAssignNew(CategoryNameTextBox, SInlineEditableTextBlock)
+		.Text(this, &FRCCategoryModel::GetCategoryDisplayName)
+		.MultiLine(false)
+		.OnTextCommitted(this, &FRCCategoryModel::OnCategoryNameCommitted);
+}
+
+TSharedRef<SWidget> FRCCategoryModel::GetWidget() const
+{
+	static const FMargin SlotMargin(10.0f, 2.0f);
+
+	return SNew(SBox)
+		.Padding(SlotMargin)
+		[
+			CategoryNameTextBox.ToSharedRef()
+		];
+}
+
+TSharedRef<SWidget> FRCCategoryModel::GetNameWidget() const
+{
+	return SNew(SBox)
+		.VAlign(VAlign_Center)
+		.Padding(10.f, 2.f)
+		[
+			CategoryNameTextBox.ToSharedRef()
+		];
+}
+
+void FRCCategoryModel::OnCategoryNameCommitted(const FText& InNewCategoryName, ETextCommit::Type InCommitInfo)
+{
+	if (!InNewCategoryName.IsEmpty())
+	{
+		CheckCachedIndex();
+
+		if (CachedIndex != INDEX_NONE)
+		{
+			// Validated by having a valid CachedIndex
+			URCVirtualPropertyContainerBase* ControllerContainer = GetPreset()->GetControllerContainer();
+
+			FScopedTransaction Transaction(LOCTEXT("RenameCategory", "Rename Category"));
+			ControllerContainer->Modify();
+			ControllerContainer->Categories[CachedIndex].Name = InNewCategoryName;
+			CategoryNameTextBox->SetText(InNewCategoryName);
+		}
+	}
+}
+
+void FRCCategoryModel::CheckCachedIndex() const
+{
+	if (URemoteControlPreset* Preset = GetPreset())
+	{
+		if (URCVirtualPropertyContainerBase* ControllerContainer = Preset->GetControllerContainer())
+		{
+			if (CachedIndex != INDEX_NONE && ControllerContainer->Categories.IsValidIndex(CachedIndex) 
+				&& ControllerContainer->Categories[CachedIndex].Id == Id)
+			{
+				return;
+			}
+
+			for (int32 Index = 0; Index < ControllerContainer->Categories.Num(); ++Index)
+			{
+				if (ControllerContainer->Categories[Index].Id == Id)
+				{
+					CachedIndex = Index;
+					return;
+				}
+			}
+		}
+	}
+
+	CachedIndex = INDEX_NONE;
+}
+
+void FRCCategoryModel::EnterNameEditingMode()
+{	
+	CategoryNameTextBox->EnterEditingMode();
+}
+
+FText FRCCategoryModel::GetCategoryDisplayName() const
+{
+	CheckCachedIndex();
+
+	if (CachedIndex != INDEX_NONE)
+	{
+		return GetPreset()->GetControllerContainer()->Categories[CachedIndex].Name;
+	}
+
+	return FText::GetEmpty();
+}
+
+void FRCCategoryModel::SetDisplayIndex(int32 InDisplayIndex)
+{
+	CheckCachedIndex();
+
+	if (CachedIndex != INDEX_NONE)
+	{
+		// Validated by having a valid CachedIndex
+		URCVirtualPropertyContainerBase* ControllerContainer = GetPreset()->GetControllerContainer();
+
+		if (GUndo)
+		{
+			ControllerContainer->Modify();
+		}
+
+		ControllerContainer->Categories[CachedIndex].DisplayIndex = InDisplayIndex;
+	}
+}
+
+#undef LOCTEXT_NAMESPACE
diff --git a/Engine/Plugins/VirtualProduction/RemoteControl/Source/RemoteControlUI/Private/UI/Controller/RCCategoryModel.h b/Engine/Plugins/VirtualProduction/RemoteControl/Source/RemoteControlUI/Private/UI/Controller/RCCategoryModel.h
new file mode 100644
--- /dev/null
+++ b/Engine/Plugins/VirtualProduction/RemoteControl/Source/RemoteControlUI/Private/UI/Controller/RCCategoryModel.h
@@ -0,0 +1,71 @@
+// Copyright Epic Games, Inc. All Rights Reserved.
+
+#pragma once
+
+#include "UI/BaseLogicUI/RCLogicModeBase.h"
+
+#include "Misc/Optional.h"
+#include "Widgets/DeclarativeSyntaxSupport.h"
+#include "Widgets/SWidget.h"
+
+class SInlineEditableTextBlock;
+class SRemoteControlPanel;
+struct FRCVirtualPropertyCategory;
+
+/*
+* ~ FRCCategoryModel ~
+*
+* UI model for representing a Category
+* Contains a row widget with Category Name
+*/
+class FRCCategoryModel : public FRCLogicModeBase
+{
+public:
+	DECLARE_MULTICAST_DELEGATE_OneParam(FOnValueChanged, TSharedPtr<FRCCategoryModel> /* InCategoryModel */);
+
+	FRCCategoryModel(const FGuid& InCateogryId, const TSharedPtr<SRemoteControlPanel> InRemoteControlPanel);
+
+	TOptional<FRCVirtualPropertyCategory> GetCategory() const;
+
+	/** Called right after object construction to initialize the model (and allow using SharedThis, etc) */
+	void Initialize();
+
+	/**
+	 * The widget to be rendered for this Category
+	 * Used to represent a single row when added to the Controllers Panel List
+	 */
+	virtual TSharedRef<SWidget> GetWidget() const override;
+
+	/** The widget showing the Name of this Category */
+	TSharedRef<SWidget> GetNameWidget() const;
+
+	/** Allows users to enter text into the Category Name Box */
+	void EnterNameEditingMode();
+
+	/** User-friendly Name of the underlying Category */
+	FText GetCategoryDisplayName() const;
+
+	/**Fetches the unique Id for this UI item */
+	FGuid GetId() const { return Id; }
+
+	/** Set the display index of the underlying category item */
+	void SetDisplayIndex(int32 InDisplayIndex);
+
+	/** Triggered when the Value this controller changes */
+	FOnValueChanged OnValueChanged;
+private:
+	/** Text commit event for Category Name text box */
+	void OnCategoryNameCommitted(const FText& InNewCategoryName, ETextCommit::Type InCommitInfo);
+
+	/** Checks cached index fro validity and sets it if not. */
+	void CheckCachedIndex() const;
+
+	/** Category name - editable text box */
+	TSharedPtr<SInlineEditableTextBlock> CategoryNameTextBox;
+
+	/** Unique Id for this UI item */
+	FGuid Id;
+
+	/** Cached index for the category for faster lookup. */
+	mutable int32 CachedIndex;
+};
diff --git a/Engine/Plugins/VirtualProduction/RemoteControl/Source/RemoteControlUI/Private/UI/Controller/RCControllerPanelListItem.h b/Engine/Plugins/VirtualProduction/RemoteControl/Source/RemoteControlUI/Private/UI/Controller/RCControllerPanelListItem.h
new file mode 100644
--- /dev/null
+++ b/Engine/Plugins/VirtualProduction/RemoteControl/Source/RemoteControlUI/Private/UI/Controller/RCControllerPanelListItem.h
@@ -0,0 +1,119 @@
+// Copyright Epic Games, Inc. All Rights Reserved.
+
+#pragma once
+
+#include "Containers/UnrealString.h"
+#include "Framework/Views/TableViewTypeTraits.h"
+#include "HAL/Platform.h"
+
+enum class ERCControllerPanelListItemType : uint8
+{
+	None,
+	Category,
+	Controller
+};
+
+struct FRCControllerPanelListItem
+{
+	ERCControllerPanelListItemType Type;
+	int32 Index;
+
+	operator bool() const
+	{
+		return Type != ERCControllerPanelListItemType::None;
+	}
+
+	bool operator==(const FRCControllerPanelListItem& InOther) const
+	{
+		return Type == InOther.Type
+			&& Index == InOther.Index;
+	}
+
+	bool operator<(const FRCControllerPanelListItem& InOther) const
+	{
+		if (static_cast<uint8>(Type) < static_cast<uint8>(InOther.Type))
+		{
+			return true;
+		}
+		else if (static_cast<uint8>(Type) > static_cast<uint8>(InOther.Type))
+		{
+			return false;
+		}
+
+		return Index < InOther.Index;
+	}
+
+	friend uint32 GetTypeHash(const FRCControllerPanelListItem& InItem)
+	{
+		return static_cast<uint32>(InItem.Index) | (1 << (32 - static_cast<uint32>(InItem.Type)));
+	}
+};
+
+namespace UE::RemoteControl::UI::Private
+{
+	namespace ControllerPanelListItem
+	{
+		const FRCControllerPanelListItem None = {ERCControllerPanelListItemType::None, INDEX_NONE};
+	}
+}
+
+template <>
+struct TIsValidListItem<FRCControllerPanelListItem>
+{
+	enum
+	{
+		Value = true
+	};
+};
+
+/**
+ * Functionality (e.g. setting the value, testing invalid value) specialized for FGroupItem.
+ */
+template <>
+struct TListTypeTraits<FRCControllerPanelListItem>
+{
+public:
+	typedef FRCControllerPanelListItem NullableType;
+
+	using MapKeyFuncs = TDefaultMapHashableKeyFuncs<FRCControllerPanelListItem, TSharedRef<ITableRow>, false>;
+	using MapKeyFuncsSparse = TDefaultMapHashableKeyFuncs<FRCControllerPanelListItem, FSparseItemInfo, false>;
+	using SetKeyFuncs = DefaultKeyFuncs<FRCControllerPanelListItem>;
+
+	template<typename U>
+	static void AddReferencedObjects(FReferenceCollector& Collector,
+		TArray<FRCControllerPanelListItem>& ItemsWithGeneratedWidgets,
+		TSet<FRCControllerPanelListItem>& SelectedItems,
+		TMap<const U*, FRCControllerPanelListItem>& WidgetToItemMap)
+	{
+	}
+
+	static bool IsPtrValid(const FRCControllerPanelListItem& InValue)
+	{
+		return InValue.Type != ERCControllerPanelListItemType::None;
+	}
+
+	static void ResetPtr(FRCControllerPanelListItem& InValue)
+	{
+		InValue.Type = ERCControllerPanelListItemType::None;
+		InValue.Index = INDEX_NONE;
+	}
+
+	static FRCControllerPanelListItem MakeNullPtr()
+	{
+		return UE::RemoteControl::UI::Private::ControllerPanelListItem::None;
+	}
+
+	static FRCControllerPanelListItem NullableItemTypeConvertToItemType(const FRCControllerPanelListItem& InValue)
+	{
+		return InValue;
+	}
+
+	static FString DebugDump(FRCControllerPanelListItem InValue)
+	{
+		return FString::FromInt(static_cast<int32>(InValue.Type)) + TEXT(": ") + FString::FromInt(InValue.Index);
+	}
+
+	class SerializerType
+	{
+	};
+};
diff --git a/Engine/Plugins/VirtualProduction/RemoteControl/Source/RemoteControlUI/Private/UI/Controller/SRCControllerItemRow.cpp b/Engine/Plugins/VirtualProduction/RemoteControl/Source/RemoteControlUI/Private/UI/Controller/SRCControllerItemRow.cpp
deleted file mode 100644
diff --git a/Engine/Plugins/VirtualProduction/RemoteControl/Source/RemoteControlUI/Private/UI/Controller/SRCControllerItemRow.h b/Engine/Plugins/VirtualProduction/RemoteControl/Source/RemoteControlUI/Private/UI/Controller/SRCControllerItemRow.h
deleted file mode 100644
diff --git a/Engine/Plugins/VirtualProduction/RemoteControl/Source/RemoteControlUI/Private/UI/Controller/SRCControllerPanel.cpp b/Engine/Plugins/VirtualProduction/RemoteControl/Source/RemoteControlUI/Private/UI/Controller/SRCControllerPanel.cpp
--- a/Engine/Plugins/VirtualProduction/RemoteControl/Source/RemoteControlUI/Private/UI/Controller/SRCControllerPanel.cpp
+++ b/Engine/Plugins/VirtualProduction/RemoteControl/Source/RemoteControlUI/Private/UI/Controller/SRCControllerPanel.cpp
@@ -6,6 +6,7 @@
 #include "Controller/RCCustomControllerUtilities.h"
 #include "Engine/Texture.h"
 #include "Framework/MultiBox/MultiBoxBuilder.h"
+#include "GraphEditorSettings.h"
 #include "IDetailTreeNode.h"
 #include "IPropertyRowGenerator.h"
 #include "Materials/MaterialInterface.h"
@@ -14,6 +15,7 @@
 #include "PropertyEditorModule.h"
 #include "RCControllerModel.h"
 #include "RCVirtualPropertyContainer.h"
+#include "RemoteControlCommon.h"
 #include "RemoteControlLogicConfig.h"
 #include "RemoteControlPreset.h"
 #include "SlateOptMacros.h"
@@ -38,6 +40,14 @@
 
 #define LOCTEXT_NAMESPACE "SRCControllerPanel"
 
+namespace UE::RemoteControl::UI::Private
+{
+	const FLazyName FRCControllerPropertyInfo::NonPropertyTypeName = TEXT("NonPropertyType");
+	const FLazyName FRCControllerPropertyInfo::CategoryName = TEXT("Category");
+	const FLazyName FRCControllerPropertyInfo::InitialDescriptionName = TEXT("InitialDescription");
+	const FString FRCControllerPropertyInfo::CategoryMeta = TEXT("Category");
+}
+
 BEGIN_SLATE_FUNCTION_BUILD_OPTIMIZATION
 
 void SRCControllerPanel::Construct(const FArguments& InArgs, const TSharedRef<SRemoteControlPanel>& InPanel)
@@ -460,41 +470,15 @@ TSharedRef<SWidget> SRCControllerPanel::GetControllerMenuContentWidget() const
 		MenuBuilder.AddMenuEntry(Action, MenuItemWidget);
 	}
 
-	// All structs
-	const FName DefaultStructName = TEXT("Structs");
-	FInstancedPropertyBag StructBag;
-	StructBag.AddProperty(DefaultStructName, EPropertyBagPropertyType::Struct, TBaseStructure<FVector2D>().Get());
+	const UGraphEditorSettings* GraphEditorSettings = GetDefault<UGraphEditorSettings>();
 
-	if (const FPropertyBagPropertyDesc* StructBagPropertyDesc = StructBag.FindPropertyDescByName(DefaultStructName))
-	{
-		TSharedRef<SWidget> MenuItemWidget =
-			SNew(SHorizontalBox)
-			+ SHorizontalBox::Slot()
-			.AutoWidth()
-			.HAlign(HAlign_Left)
-			.VAlign(VAlign_Fill)
-			[
-				UE::RCUIHelpers::GetTypeColorWidget(StructBagPropertyDesc->CachedProperty)
-			]
-			+ SHorizontalBox::Slot()
-			.AutoWidth()
-			.Padding(FMargin(6.f, 0.f))
-			[
-				SNew(STextBlock)
-				.Text(FText::FromName(DefaultStructName))
-			];
+	using namespace UE::RCCustomControllers;
 
-		FNewMenuDelegate Delegate = FNewMenuDelegate::CreateSP(this, &SRCControllerPanel::GetControllerMenuStructContentWidget);
-		MenuBuilder.AddSubMenu(MenuItemWidget, Delegate);
-	}
 
 	// All enums
-	const FName DefaultEnumName = TEXT("Enums");
-	FInstancedPropertyBag EnumBag;
-	EnumBag.AddProperty(DefaultEnumName, EPropertyBagPropertyType::Enum, StaticEnum<EStyleColor>());
-
-	if (const FPropertyBagPropertyDesc* EnumBagPropertyDesc = EnumBag.FindPropertyDescByName(DefaultEnumName))
 	{
+		const FText EnumsName = LOCTEXT("Enums", "Enums");
+
 		{
 			TSharedRef<SWidget> MenuItemWidget =
 				SNew(SHorizontalBox)
@@ -503,14 +487,14 @@ TSharedRef<SWidget> SRCControllerPanel::GetControllerMenuContentWidget() const
 				.HAlign(HAlign_Left)
 				.VAlign(VAlign_Fill)
 				[
-					UE::RCUIHelpers::GetTypeColorWidget(EnumBagPropertyDesc->CachedProperty)
+					UE::RCUIHelpers::GetTypeColorWidget(GraphEditorSettings->BytePinTypeColor, EnumsName)
 				]
 				+ SHorizontalBox::Slot()
 				.AutoWidth()
 				.Padding(FMargin(6.f, 0.f))
 				[
 					SNew(STextBlock)
-					.Text(FText::FromName(DefaultEnumName))
+					.Text(EnumsName)
 				];
 
 			FNewMenuDelegate Delegate = FNewMenuDelegate::CreateSP(this, &SRCControllerPanel::GetControllerMenuEnumContentWidget);
@@ -519,12 +503,9 @@ TSharedRef<SWidget> SRCControllerPanel::GetControllerMenuContentWidget() const
 	}
 
 	// String List
-	const FName DefaultStringName = TEXT("Strings");
-	FInstancedPropertyBag StringBag;
-	StringBag.AddProperty(DefaultStringName, EPropertyBagPropertyType::String);
-
-	if (const FPropertyBagPropertyDesc* StringBagPropertyDesc = StringBag.FindPropertyDescByName(DefaultStringName))
 	{
+		const FText StringListName = LOCTEXT("StringList", "String List");
+
 		TSharedRef<SWidget> MenuItemWidget =
 			SNew(SHorizontalBox)
 			+ SHorizontalBox::Slot()
@@ -532,14 +513,14 @@ TSharedRef<SWidget> SRCControllerPanel::GetControllerMenuContentWidget() const
 			.HAlign(HAlign_Left)
 			.VAlign(VAlign_Fill)
 			[
-				UE::RCUIHelpers::GetTypeColorWidget(StringBagPropertyDesc->CachedProperty)
+				UE::RCUIHelpers::GetTypeColorWidget(GraphEditorSettings->StringPinTypeColor, StringListName)
 			]
 			+ SHorizontalBox::Slot()
 			.AutoWidth()
 			.Padding(FMargin(6.f, 0.f))
 			[
 				SNew(STextBlock)
-				.Text(LOCTEXT("StringList", "String List"))
+				.Text(StringListName)
 			];
 
 		using namespace UE::RCCustomControllers;
@@ -553,6 +534,31 @@ TSharedRef<SWidget> SRCControllerPanel::GetControllerMenuContentWidget() const
 		MenuBuilder.AddMenuEntry(Action, MenuItemWidget);
 	}
 
+	// Category
+	{
+		const FText CategoryName = LOCTEXT("Category", "Category");
+
+		TSharedRef<SWidget> MenuItemWidget =
+			SNew(SHorizontalBox)
+			+ SHorizontalBox::Slot()
+			.AutoWidth()
+			.HAlign(HAlign_Left)
+			.VAlign(VAlign_Fill)
+			[
+				UE::RCUIHelpers::GetTypeColorWidget(GraphEditorSettings->Int64PinTypeColor, CategoryName)
+			]
+			+ SHorizontalBox::Slot()
+			.AutoWidth()
+			.Padding(FMargin(6.f, 0.f))
+			[
+				SNew(STextBlock)
+				.Text(CategoryName)
+			];
+
+		FUIAction Action(FExecuteAction::CreateSP(this, &SRCControllerPanel::OnAddCategoryClicked));
+		MenuBuilder.AddMenuEntry(Action, MenuItemWidget);
+	}
+
 	return MenuBuilder.MakeWidget();
 }
 
@@ -595,12 +601,6 @@ void SRCControllerPanel::GetControllerMenuStructContentWidget(FMenuBuilder& InMe
 			continue;
 		}
 
-		// Uncustomised structs do not show correctly.
-		if (!PropertyEditorModule.IsCustomizedStruct(Struct, FCustomPropertyTypeLayoutMap()))
-		{
-			continue;
-		}
-
 		Structs.Add(Struct);
 	}
 
@@ -624,7 +624,7 @@ void SRCControllerPanel::GetControllerMenuStructContentWidget(FMenuBuilder& InMe
 
 	InMenuBuilder.AddWidget(
 		SNew(SRCControllerTypeSelector, Types)
-			.OnTypeSelected(FRCControllerTypeBase::FOnTypeSelected::CreateSP(this, &SRCControllerPanel::OnAddControllerSelected)),
+		.OnTypeSelected(FRCControllerTypeBase::FOnTypeSelected::CreateSP(this, &SRCControllerPanel::OnAddControllerSelected)),
 		FText::GetEmpty(),
 		/* No Indent */ false
 	);
@@ -706,94 +706,108 @@ void SRCControllerPanel::OnAddControllerClicked(const EPropertyBagPropertyType I
 void SRCControllerPanel::OnAddControllerClicked(const EPropertyBagPropertyType InValueType, const UE::RemoteControl::UI::Private::FRCControllerPropertyInfo InControllerPropertyInfo, TSubclassOf<URCController> InControllerClass) const
 {
 	// Add to the asset
-	if (URemoteControlPreset* Preset = GetPreset())
+	URemoteControlPreset* Preset = GetPreset();
+
+	if (!Preset)
 	{
-		FScopedTransaction Transaction(LOCTEXT("AddController", "Add Controller"));
-		Preset->Modify();
+		return;
+	}
 
-		URCVirtualPropertyInContainer* NewVirtualProperty = Preset->AddController(InControllerClass, InValueType, InControllerPropertyInfo.ValueTypeObject);
+	FScopedTransaction Transaction(LOCTEXT("AddController", "Add Controller"));
+	Preset->Modify();
 
-		// Copy over the custom metadata of this property
-		for (const TPair<FName, FString>& Pair : InControllerPropertyInfo.CustomMetaData)
-		{
-			NewVirtualProperty->SetMetadataValue(Pair.Key, Pair.Value);
-		}
+	using namespace UE::RemoteControl::UI::Private;
 
-		// Additional logic for Custom Controllers naming
-		if (UE::RCCustomControllers::IsCustomController(NewVirtualProperty))
+	FName InitialPropertyName = NAME_None;
+	
+	if (const FString* NonPropertyTypeValue = InControllerPropertyInfo.CustomMetaData.Find(FRCControllerPropertyInfo::NonPropertyTypeName))
+	{
+		if ((*NonPropertyTypeValue) == FRCControllerPropertyInfo::CategoryMeta)
 		{
-			const FName& NewControllerName = UE::RCCustomControllers::GetUniqueNameForController(NewVirtualProperty);
-			NewVirtualProperty->DisplayName = NewControllerName;
+			InitialPropertyName = **NonPropertyTypeValue;
 		}
+	}
 
-		if (ControllerPanelList.IsValid())
-		{
-			NewVirtualProperty->DisplayIndex = ControllerPanelList->NumControllerItems();
-		}
+	URCVirtualPropertyInContainer* NewVirtualProperty = Preset->AddController(URCController::StaticClass(), InValueType, InControllerPropertyInfo.ValueTypeObject, InitialPropertyName);
 
-		// Refresh list
-		const TSharedPtr<SRemoteControlPanel> RemoteControlPanel = GetRemoteControlPanel();
-		check(RemoteControlPanel);
-		RemoteControlPanel->OnControllerAdded.Broadcast(NewVirtualProperty->PropertyName);
+	if (!NewVirtualProperty)
+	{
+		return;
 	}
-}
 
-void SRCControllerPanel::OnAddControllerSelected(TSharedPtr<UE::RemoteControl::UI::Private::FRCControllerPropertyInfo> InControllerPropertyInfo) const
-{
-	if (UStruct* Struct = Cast<UStruct>(InControllerPropertyInfo->ValueTypeObject))
+	// Copy over the custom metadata of this property
+	// Initial description should not be copied
+	for (const TPair<FName, FString>& Pair : InControllerPropertyInfo.CustomMetaData)
 	{
-		FInstancedPropertyBag Bag;
-		Bag.AddProperty("Default", EPropertyBagPropertyType::Struct, Struct);
+		if (Pair.Key != FRCControllerPropertyInfo::InitialDescriptionName)
+		{
+			NewVirtualProperty->SetMetadataValue(Pair.Key, Pair.Value);
+		}
+	}
 
-		TSharedRef<FStructOnScope> StructOnScope = MakeShared<FStructOnScope>(Bag.GetPropertyBagStruct(), Bag.GetMutableValue().GetMemory());
+	// Additional logic for Custom Controllers naming
+	if (UE::RCCustomControllers::IsCustomController(NewVirtualProperty))
+	{
+		const FName& NewControllerName = UE::RCCustomControllers::GetUniqueNameForController(NewVirtualProperty);
+		NewVirtualProperty->DisplayName = NewControllerName;
+	}
 
-		FPropertyRowGeneratorArgs Args;
-		Args.bShouldShowHiddenProperties = true;
-		TSharedRef<IPropertyRowGenerator> PropertyRowGenerator = FModuleManager::GetModuleChecked<FPropertyEditorModule>("PropertyEditor").CreatePropertyRowGenerator(Args);
-		PropertyRowGenerator->SetStructure(StructOnScope);
+	if (const FString* InitialDescription = InControllerPropertyInfo.CustomMetaData.Find(FRCControllerPropertyInfo::InitialDescriptionName))
+	{
+		if (!InitialDescription->IsEmpty())
+		{
+			NewVirtualProperty->Description = FText::FromString(*InitialDescription);
+		}
+	}
 
-		TArray<TSharedRef<IDetailTreeNode>> RootNodes = PropertyRowGenerator->GetRootTreeNodes();
+	if (ControllerPanelList.IsValid())
+	{
+		NewVirtualProperty->DisplayIndex = ControllerPanelList->NumControllerItems();
+	}
 
-		bool bInvalidNodes = false;
+	// Refresh list
+	const TSharedPtr<SRemoteControlPanel> RemoteControlPanel = GetRemoteControlPanel();
+	check(RemoteControlPanel);
+	RemoteControlPanel->OnControllerAdded.Broadcast(NewVirtualProperty->PropertyName);
+}
 
-		for (const TSharedRef<IDetailTreeNode>& RootNode : RootNodes)
-		{
-			TArray<TSharedRef<IDetailTreeNode>> Children;
-			RootNode->GetChildren(Children);
+void SRCControllerPanel::OnAddCategoryClicked() const
+{
+	// Add to the asset
+	URemoteControlPreset* Preset = GetPreset();
 
-			for (const TSharedRef<IDetailTreeNode>& Child : Children)
-			{
-				TSharedPtr<IPropertyHandle> ChildHandle = Child->CreatePropertyHandle();
+	if (!Preset)
+	{
+		return;
+	}
 
-				if (!ChildHandle.IsValid())
-				{
-					bInvalidNodes = true;
-					break;
-				}
+	URCVirtualPropertyContainerBase* ControllerContainer = Preset->GetControllerContainer();
 
-				if (ChildHandle->GetPropertyPath().Contains(TEXT("->")))
-				{
-					bInvalidNodes = true;
-					break;
-				}
-			}
+	if (!ControllerContainer)
+	{
+		return;
+	}
 
-			if (bInvalidNodes)
-			{
-				break;
-			}
-		}
+	FScopedTransaction Transaction(LOCTEXT("AddCategory", "Add Category"));
+	ControllerContainer->Modify();
 
-		if (bInvalidNodes)
-		{
-			FMessageDialog::Open(
-				EAppMsgType::Ok,
-				LOCTEXT("InvalidNodesError", "This struct is not supported due to its customization.")
-			);
-			return;
-		}
+	FRCVirtualPropertyCategory& NewCategory = ControllerContainer->Categories.AddDefaulted_GetRef();
+	NewCategory.Id = FGuid::NewGuid();
+	NewCategory.Name = LOCTEXT("NewCategory", "New Category");
+	
+	if (ControllerPanelList.IsValid())
+	{
+		NewCategory.DisplayIndex = ControllerPanelList->NumControllerItems();
 	}
 
+	// Refresh list
+	const TSharedPtr<SRemoteControlPanel> RemoteControlPanel = GetRemoteControlPanel();
+	check(RemoteControlPanel);
+	RemoteControlPanel->OnControllerAdded.Broadcast(*NewCategory.Id.ToString());
+}
+
+void SRCControllerPanel::OnAddControllerSelected(TSharedPtr<UE::RemoteControl::UI::Private::FRCControllerPropertyInfo> InControllerPropertyInfo) const
+{
 	OnAddControllerClicked(InControllerPropertyInfo->Type, *InControllerPropertyInfo);
 }
 
diff --git a/Engine/Plugins/VirtualProduction/RemoteControl/Source/RemoteControlUI/Private/UI/Controller/SRCControllerPanel.h b/Engine/Plugins/VirtualProduction/RemoteControl/Source/RemoteControlUI/Private/UI/Controller/SRCControllerPanel.h
--- a/Engine/Plugins/VirtualProduction/RemoteControl/Source/RemoteControlUI/Private/UI/Controller/SRCControllerPanel.h
+++ b/Engine/Plugins/VirtualProduction/RemoteControl/Source/RemoteControlUI/Private/UI/Controller/SRCControllerPanel.h
@@ -26,6 +26,11 @@ namespace UE::RemoteControl::UI::Private
  */
 struct FRCControllerPropertyInfo
 {
+	static const FLazyName NonPropertyTypeName;
+	static const FLazyName CategoryName;
+	static const FLazyName InitialDescriptionName;
+	static const FString CategoryMeta;
+
 	FRCControllerPropertyInfo() {}
 	FRCControllerPropertyInfo(UObject* InValueTypeObject) : ValueTypeObject(InValueTypeObject) {}
 
@@ -117,6 +122,9 @@ private:
 	void OnAddControllerClicked(const EPropertyBagPropertyType InValueType, const UE::RemoteControl::UI::Private::FRCControllerPropertyInfo InControllerPropertyInfo) const;
 	void OnAddControllerClicked(const EPropertyBagPropertyType InValueType, const UE::RemoteControl::UI::Private::FRCControllerPropertyInfo InControllerPropertyInfo, TSubclassOf<URCController> InControllerClass) const;
 
+	/** Handles adding a new category */
+	void OnAddCategoryClicked() const;
+
 	/** Handles selection of controller type in submenus. */
 	void OnAddControllerSelected(TSharedPtr<UE::RemoteControl::UI::Private::FRCControllerPropertyInfo> InControllerPropertyInfo) const;
 
diff --git a/Engine/Plugins/VirtualProduction/RemoteControl/Source/RemoteControlUI/Private/UI/Controller/SRCControllerPanelCategoryRow.cpp b/Engine/Plugins/VirtualProduction/RemoteControl/Source/RemoteControlUI/Private/UI/Controller/SRCControllerPanelCategoryRow.cpp
new file mode 100644
--- /dev/null
+++ b/Engine/Plugins/VirtualProduction/RemoteControl/Source/RemoteControlUI/Private/UI/Controller/SRCControllerPanelCategoryRow.cpp
@@ -0,0 +1,155 @@
+// Copyright Epic Games, Inc. All Rights Reserved.
+
+#include "SRCControllerPanelCategoryRow.h"
+#include "Controller/RCController.h"
+#include "Controller/RCControllerUtilities.h"
+#include "Controller/RCCustomControllerUtilities.h"
+#include "RemoteControlPreset.h"
+#include "ScopedTransaction.h"
+#include "SRCControllerPanelList.h"
+#include "UI/Controller/RCCategoryModel.h"
+#include "UI/Controller/SRCControllerPanel.h"
+#include "UI/RCUIHelpers.h"
+#include "UI/SRCPanelExposedEntity.h"
+#include "Widgets/SNullWidget.h"
+#include "Widgets/Views/SExpanderArrow.h"
+
+#define LOCTEXT_NAMESPACE "SRCControllerPanelCategoryRow"
+
+void SRCControllerPanelCategoryRow::Construct(const FArguments& InArgs, const TSharedRef<STableViewBase>& InOwnerTableView,
+	TSharedRef<SRCControllerPanelList> InControllerPanelList, FRCControllerPanelListItem InCategoryItem)
+{
+	CategoryItem = InCategoryItem;
+	ControllerPanelListWeak = InControllerPanelList.ToWeakPtr();
+
+	TSharedPtr<FRCCategoryModel> CategoryModel = GetCategory();
+
+	STableRow::Construct(
+		STableRow::FArguments()
+		.Style(InArgs._Style)
+		.Padding(4.5f)
+		.OnCanAcceptDrop(this, &SRCControllerPanelCategoryRow::CanAcceptDrop)
+		.OnAcceptDrop(this, &SRCControllerPanelCategoryRow::OnAcceptDrop)
+		.Content()
+		[
+			CategoryModel.IsValid() 
+				? CategoryModel->GetNameWidget()
+				: SNullWidget::NullWidget
+		],
+		InOwnerTableView
+	);
+}
+
+TSharedPtr<FRCCategoryModel> SRCControllerPanelCategoryRow::GetCategory() const
+{
+	if (TSharedPtr<SRCControllerPanelList> ControllerPanelList = ControllerPanelListWeak.Pin())
+	{
+		return ControllerPanelList->FindCategoryItemByListItem(CategoryItem);
+	}
+
+	return {};
+}
+
+TOptional<EItemDropZone> SRCControllerPanelCategoryRow::CanAcceptDrop(const FDragDropEvent& InDragDropEvent, EItemDropZone InDropZone, FRCControllerPanelListItem InItem)
+{
+	TSharedPtr<FRCCategoryModel> CategoryModel = GetCategory();
+	if (!CategoryModel.IsValid())
+	{
+		return TOptional<EItemDropZone>();
+	}
+
+	const TSharedPtr<SRCControllerPanelList> ControllerPanelList = ControllerPanelListWeak.Pin();
+	if (!ControllerPanelList.IsValid())
+	{
+		return TOptional<EItemDropZone>();
+	}
+
+	// Fetch the Exposed Entities
+	URemoteControlPreset* const Preset = ControllerPanelList->GetPreset();
+	if (!Preset)
+	{
+		return TOptional<EItemDropZone>();
+	}
+
+	// Dragging Categorys onto Categorys (Reordering). Only allow Above/Below re-ordering (no-onto support)
+	if (TSharedPtr<FRCControllerPanelDragDrop> ControllerDragDrop = InDragDropEvent.GetOperationAs<FRCControllerPanelDragDrop>())
+	{
+		if (!ControllerDragDrop->GetItems().Contains(CategoryItem))
+		{
+			return InDropZone;
+		}
+	}
+
+	return TOptional<EItemDropZone>();
+}
+
+FReply SRCControllerPanelCategoryRow::OnAcceptDrop(const FDragDropEvent& InDragDropEvent, EItemDropZone InDropZone, FRCControllerPanelListItem InItem)
+{
+	TSharedPtr<FRCCategoryModel> CategoryModel = GetCategory();
+	if (!CategoryModel.IsValid())
+	{
+		return FReply::Unhandled();
+	}
+
+	const TSharedPtr<SRCControllerPanelList> ControllerPanelList = ControllerPanelListWeak.Pin();
+	if (!ControllerPanelList.IsValid())
+	{
+		return FReply::Unhandled();
+	}
+
+	ControllerPanelList->bIsAnyControllerItemEligibleForDragDrop = false;
+
+	URemoteControlPreset* Preset = ControllerPanelList->GetPreset();
+	if (!Preset)
+	{
+		return FReply::Unhandled();
+	}
+
+	if (TSharedPtr<FRCControllerPanelDragDrop> ControllerDragDrop = InDragDropEvent.GetOperationAs<FRCControllerPanelDragDrop>())
+	{
+		FScopedTransaction Transaction(LOCTEXT("ReorderControllers", "Reorder Controllers"));
+
+		if (!ControllerPanelList->ReorderItems(ControllerDragDrop->GetItems(), CategoryItem, InDropZone))
+		{
+			Transaction.Cancel();
+		}
+
+		return FReply::Handled();
+	}
+
+	return FReply::Unhandled();
+}
+
+FReply SRCControllerPanelCategoryRow::OnDragDetected(const FGeometry& InMyGeometry, const FPointerEvent& InMouseEvent)
+{
+	if (FSuperRowType::OnDragDetected(InMyGeometry, InMouseEvent).IsEventHandled())
+	{
+		return FReply::Handled();
+	}
+
+	const TSharedPtr<SRCControllerPanelList> ControllerPanelList = ControllerPanelListWeak.Pin();
+	if (!ControllerPanelList.IsValid())
+	{
+		return FReply::Unhandled();
+	}
+
+	if (InMouseEvent.IsMouseButtonDown(EKeys::LeftMouseButton))
+	{
+		TSharedRef<FRCControllerPanelDragDrop> ControllerDragDrop = MakeShared<FRCControllerPanelDragDrop>(ControllerPanelList->GetSelectedItems());
+		return FReply::Handled().BeginDragDrop(ControllerDragDrop);
+	}
+	return FReply::Unhandled();
+}
+
+FReply SRCControllerPanelCategoryRow::OnMouseButtonDoubleClick(const FGeometry& InMyGeometry, const FPointerEvent& InMouseEvent)
+{
+	if (TSharedPtr<FRCCategoryModel> CategoryModel = GetCategory())
+	{
+		CategoryModel->EnterNameEditingMode();
+		return FReply::Handled();
+	}
+
+	return FSuperRowType::OnMouseButtonDoubleClick(InMyGeometry, InMouseEvent);
+}
+
+#undef LOCTEXT_NAMESPACE
diff --git a/Engine/Plugins/VirtualProduction/RemoteControl/Source/RemoteControlUI/Private/UI/Controller/SRCControllerPanelCategoryRow.h b/Engine/Plugins/VirtualProduction/RemoteControl/Source/RemoteControlUI/Private/UI/Controller/SRCControllerPanelCategoryRow.h
new file mode 100644
--- /dev/null
+++ b/Engine/Plugins/VirtualProduction/RemoteControl/Source/RemoteControlUI/Private/UI/Controller/SRCControllerPanelCategoryRow.h
@@ -0,0 +1,52 @@
+// Copyright Epic Games, Inc. All Rights Reserved.
+
+#pragma once
+
+#include "Misc/Optional.h"
+#include "RCVirtualPropertyContainer.h"
+#include "UI/Controller/RCControllerPanelListItem.h"
+#include "Widgets/Views/SListView.h"
+#include "Widgets/Views/STableRow.h"
+
+class FExposedEntityDragDrop;
+class FRCCategoryModel;
+class SRCControllerPanelList;
+class URemoteControlPreset;
+
+class SRCControllerPanelCategoryRow : public STableRow<FRCControllerPanelListItem>
+{
+	using FSuperRowType = STableRow<FRCControllerPanelListItem>;
+
+public:
+	SLATE_BEGIN_ARGS(SRCControllerPanelCategoryRow) {}
+		SLATE_STYLE_ARGUMENT(FTableRowStyle, Style)
+	SLATE_END_ARGS()
+
+	void Construct(const FArguments& InArgs, const TSharedRef<STableViewBase>& InOwnerTableView,
+		TSharedRef<SRCControllerPanelList> InControllerPanelList, FRCControllerPanelListItem InCategoryItem);
+
+	TSharedPtr<FRCCategoryModel> GetCategory() const;
+
+private:
+	/** Called to validate whether a drag drop event is accepted for a controller item row */
+	TOptional<EItemDropZone> CanAcceptDrop(const FDragDropEvent& InDragDropEvent, EItemDropZone InDropZone, FRCControllerPanelListItem InItem);
+
+	/** Called when a drop happens on this controller item row */
+	FReply OnAcceptDrop(const FDragDropEvent& InDragDropEvent, EItemDropZone InDropZone, FRCControllerPanelListItem InItem);
+
+	struct FDragDropContext
+	{
+		URemoteControlPreset* Preset;
+		const TSharedPtr<SRCControllerPanelList>& ControllerPanelList;
+		const FRCControllerPanelListItem& Item;
+		EItemDropZone DropZone;
+	};
+
+	//~ Begin SWidget
+	virtual FReply OnDragDetected(const FGeometry& InMyGeometry, const FPointerEvent& InMouseEvent) override;
+	virtual FReply OnMouseButtonDoubleClick(const FGeometry& InMyGeometry, const FPointerEvent& InMouseEvent) override;
+	//~ End SWidget
+
+	TWeakPtr<SRCControllerPanelList> ControllerPanelListWeak;
+	FRCControllerPanelListItem CategoryItem;
+};
\ No newline at end of file
diff --git a/Engine/Plugins/VirtualProduction/RemoteControl/Source/RemoteControlUI/Private/UI/Controller/SRCControllerPanelControllerRow.cpp b/Engine/Plugins/VirtualProduction/RemoteControl/Source/RemoteControlUI/Private/UI/Controller/SRCControllerPanelControllerRow.cpp
new file mode 100644
--- /dev/null
+++ b/Engine/Plugins/VirtualProduction/RemoteControl/Source/RemoteControlUI/Private/UI/Controller/SRCControllerPanelControllerRow.cpp
@@ -0,0 +1,437 @@
+// Copyright Epic Games, Inc. All Rights Reserved.
+
+#include "SRCControllerPanelControllerRow.h"
+#include "Behaviour/Builtin/Bind/RCBehaviourBind.h"
+#include "Controller/RCController.h"
+#include "Controller/RCControllerUtilities.h"
+#include "Controller/RCCustomControllerUtilities.h"
+#include "RemoteControlPreset.h"
+#include "SRCControllerPanelList.h"
+#include "ScopedTransaction.h"
+#include "UI/Controller/SRCControllerPanel.h"
+#include "UI/RCUIHelpers.h"
+#include "UI/SRCPanelExposedEntity.h"
+
+#define LOCTEXT_NAMESPACE "SRCControllerPanelControllerRow"
+
+const FLazyName FRCControllerColumns::TypeColor = TEXT("TypeColor");
+const FLazyName FRCControllerColumns::ControllerId = TEXT("Controller Id");
+const FLazyName FRCControllerColumns::Description = TEXT("Controller Description");
+const FLazyName FRCControllerColumns::Value = TEXT("Controller Value");
+const FLazyName FRCControllerColumns::FieldId = TEXT("Controller Field Id");
+const FLazyName FRCControllerColumns::ValueTypeSelection = TEXT("Value Type Selection");
+
+void SRCControllerPanelControllerRow::Construct(const FArguments& InArgs, const TSharedRef<STableViewBase>& InOwnerTableView,
+	TSharedRef<SRCControllerPanelList> InControllerPanelList, FRCControllerPanelListItem InControllerItem)
+{
+	ControllerItem = InControllerItem;
+	ControllerPanelListWeak = InControllerPanelList.ToWeakPtr();
+
+	FSuperRowType::Construct(FSuperRowType::FArguments()
+		.Style(InArgs._Style)
+		.Padding(4.5f)
+		.OnCanAcceptDrop(this, &SRCControllerPanelControllerRow::CanAcceptDrop)
+		.OnAcceptDrop(this, &SRCControllerPanelControllerRow::OnAcceptDrop)
+		, InOwnerTableView);
+}
+
+TSharedRef<SWidget> SRCControllerPanelControllerRow::GenerateWidgetForColumn(const FName& InColumnName)
+{
+	TSharedPtr<FRCControllerModel> ControllerModel = GetControllerModel();
+
+	if (!ensure(ControllerModel.IsValid()))
+	{
+		return SNullWidget::NullWidget;
+	}
+
+	if (InColumnName == FRCControllerColumns::TypeColor)
+	{
+		return SNew(SBox)
+			.HAlign(HAlign_Left)
+			[
+				GetTypeColorWidget()
+			];
+	}
+	else if (InColumnName == FRCControllerColumns::FieldId)
+	{
+		return ControllerModel->GetFieldIdWidget();
+	}
+	else if (InColumnName == FRCControllerColumns::ValueTypeSelection)
+	{
+		return ControllerModel->GetTypeSelectionWidget();
+	}
+	else if (InColumnName == FRCControllerColumns::ControllerId)
+	{
+		return ControllerModel->GetNameWidget();
+	}
+	else if (InColumnName == FRCControllerColumns::Description)
+	{
+		return ControllerModel->GetDescriptionWidget();
+	}
+	else if (InColumnName == FRCControllerColumns::Value)
+	{
+		return ControllerModel->GetWidget();
+	}
+	else if (const TSharedPtr<SRCControllerPanelList> ControlPanelList = ControllerPanelListWeak.Pin())
+	{
+		if (ControlPanelList->GetCustomColumns().Contains(InColumnName))
+		{
+			return ControllerModel->GetControllerExtensionWidget(InColumnName);
+		}
+	}
+	return SNullWidget::NullWidget;
+}
+
+TSharedPtr<FRCControllerModel> SRCControllerPanelControllerRow::GetControllerModel() const
+{
+	if (TSharedPtr<SRCControllerPanelList> ControllerPanelList = ControllerPanelListWeak.Pin())
+	{
+		return ControllerPanelList->FindControllerItemByListItem(ControllerItem);
+	}
+
+	return nullptr;
+}
+
+TOptional<EItemDropZone> SRCControllerPanelControllerRow::CanAcceptDrop(const FDragDropEvent& InDragDropEvent, EItemDropZone InDropZone, FRCControllerPanelListItem InItem)
+{
+	TSharedPtr<FRCControllerModel> ControllerModel = GetControllerModel();
+	if (!ControllerModel.IsValid())
+	{
+		return TOptional<EItemDropZone>();
+	}
+
+	const TSharedPtr<SRCControllerPanelList> ControllerPanelList = ControllerPanelListWeak.Pin();
+	if (!ControllerPanelList.IsValid())
+	{
+		return TOptional<EItemDropZone>();
+	}
+
+	// Fetch the Exposed Entities
+	URemoteControlPreset* const Preset = ControllerPanelList->GetPreset();
+	if (!Preset)
+	{
+		return TOptional<EItemDropZone>();
+	}
+
+	// Dragging Controllers onto Controllers (Reordering). Only allow Above/Below re-ordering (no-onto support)
+	if (TSharedPtr<FRCControllerPanelDragDrop> ControllerDragDrop = InDragDropEvent.GetOperationAs<FRCControllerPanelDragDrop>())
+	{
+		if (InDropZone == EItemDropZone::OntoItem)
+		{
+			return TOptional<EItemDropZone>();
+		}
+		
+		if (!ControllerDragDrop->GetItems().Contains(ControllerItem))
+		{
+			return InDropZone;
+		}
+
+		return TOptional<EItemDropZone>();
+	}
+
+	if (TSharedPtr<FExposedEntityDragDrop> EntityDragDrop = InDragDropEvent.GetOperationAs<FExposedEntityDragDrop>())
+	{
+		// Support when adding entities onto this controller, ensure that at least one entity can be added via bind action 
+		if (InDropZone == EItemDropZone::OntoItem)
+		{
+			for (const FGuid& ExposedEntityId : EntityDragDrop->GetSelectedFieldsId())
+			{
+				if (TSharedPtr<const FRemoteControlField> RemoteControlField = Preset->GetExposedEntity<FRemoteControlField>(ExposedEntityId).Pin())
+				{
+					if (URCController* Controller = Cast<URCController>(ControllerModel->GetVirtualProperty()))
+					{
+						const bool bAllowNumericInputAsStrings = true;
+						if (URCBehaviourBind::CanHaveActionForField(Controller, RemoteControlField.ToSharedRef(), bAllowNumericInputAsStrings))
+						{
+							ControllerPanelList->bIsAnyControllerItemEligibleForDragDrop = true;
+							return InDropZone;
+						}
+					}
+				}
+			}
+		}
+		// When adding entities above/below controllers (i.e. create a new controller above/below this one), at least one has to be supported
+		else
+		{
+			for (const FGuid& ExposedEntityId : EntityDragDrop->GetSelectedFieldsId())
+			{
+				if (TSharedPtr<const FRemoteControlProperty> RemoteControlProperty = Preset->GetExposedEntity<FRemoteControlProperty>(ExposedEntityId).Pin())
+				{
+					if (UE::RCControllers::CanCreateControllerFromEntity(RemoteControlProperty))
+					{
+						ControllerPanelList->bIsAnyControllerItemEligibleForDragDrop = true;
+						return InDropZone;
+					}
+				}
+			}
+		}
+	}
+
+	return TOptional<EItemDropZone>();
+}
+
+FReply SRCControllerPanelControllerRow::OnAcceptDrop(const FDragDropEvent& InDragDropEvent, EItemDropZone InDropZone, FRCControllerPanelListItem InItem)
+{
+	TSharedPtr<FRCControllerModel> ControllerModel = GetControllerModel();
+	if (!ControllerModel.IsValid())
+	{
+		return FReply::Unhandled();
+	}
+
+	const TSharedPtr<SRCControllerPanelList> ControllerPanelList = ControllerPanelListWeak.Pin();
+	if (!ControllerPanelList.IsValid())
+	{
+		return FReply::Unhandled();
+	}
+
+	ControllerPanelList->bIsAnyControllerItemEligibleForDragDrop = false;
+
+	URemoteControlPreset* Preset = ControllerPanelList->GetPreset();
+	if (!Preset)
+	{
+		return FReply::Unhandled();
+	}
+
+	if (TSharedPtr<FRCControllerPanelDragDrop> ControllerDragDrop = InDragDropEvent.GetOperationAs<FRCControllerPanelDragDrop>())
+	{
+		FScopedTransaction Transaction(LOCTEXT("ReorderControllers", "Reorder Controllers"));
+
+		if (!ControllerPanelList->ReorderItems(ControllerDragDrop->GetItems(), ControllerItem, InDropZone))
+		{
+			Transaction.Cancel();
+		}
+
+		return FReply::Handled();
+	}
+
+	if (TSharedPtr<FExposedEntityDragDrop> ExposedEntityDragDrop = InDragDropEvent.GetOperationAs<FExposedEntityDragDrop>())
+	{
+		const FDragDropContext DragDropContext
+			{
+				.Preset = Preset,
+				.ControllerPanelList = ControllerPanelList,
+				.Item = InItem,
+				.DropZone = InDropZone,
+			};
+
+		if (InDropZone == EItemDropZone::OntoItem)
+		{
+			CreateBindBehaviorsFromEntities(*ExposedEntityDragDrop, DragDropContext);
+		}
+		else
+		{
+			CreateControllersFromEntities(*ExposedEntityDragDrop, DragDropContext);
+		}
+		return FReply::Handled();
+	}
+
+	return FReply::Unhandled();
+}
+
+void SRCControllerPanelControllerRow::CreateBindBehaviorsFromEntities(const FExposedEntityDragDrop& InExposedEntityDragDrop, const FDragDropContext& InContext)
+{
+	TSharedPtr<FRCControllerModel> ControllerModel = GetControllerModel();
+	if (!ControllerModel.IsValid())
+	{
+		return;
+	}
+	
+	const TArray<FGuid>& DroppedFieldIds = InExposedEntityDragDrop.GetSelectedFieldsId();
+
+	if (!ControllerModel.IsValid() || DroppedFieldIds.IsEmpty())
+	{
+		return;
+	}
+
+	URCController* const Controller = Cast<URCController>(ControllerModel->GetVirtualProperty());
+	if (!Controller)
+	{
+		return;
+	}
+
+	FScopedTransaction Transaction(LOCTEXT("BindPropertiesToController", "Bind properties to Controller"));
+
+	bool bModified = false;
+
+	// Update controller description if empty to match the first dragged property
+	if (Controller->Description.IsEmpty())
+	{
+		if (TSharedPtr<const FRemoteControlProperty> RemoteControlProperty = InContext.Preset->GetExposedEntity<FRemoteControlProperty>(DroppedFieldIds[0]).Pin())
+		{
+			Controller->Modify();
+			Controller->Description = UE::RCUIHelpers::GenerateControllerDescriptionFromEntity(RemoteControlProperty);
+			bModified = true;
+		}
+	}
+
+	for (const FGuid& ExposedEntityId : DroppedFieldIds)
+	{
+		if (TSharedPtr<const FRemoteControlProperty> RemoteControlProperty = InContext.Preset->GetExposedEntity<FRemoteControlProperty>(ExposedEntityId).Pin())
+		{
+			InContext.ControllerPanelList->CreateBindBehaviourAndAssignTo(Controller, RemoteControlProperty.ToSharedRef(), true);
+			bModified = true;
+		}
+	}
+
+	if (!bModified)
+	{
+		Transaction.Cancel();
+	}
+}
+
+void SRCControllerPanelControllerRow::CreateControllersFromEntities(const FExposedEntityDragDrop& InExposedEntityDragDrop, const FDragDropContext& InContext)
+{
+	TSharedPtr<FRCControllerModel> ControllerModel = GetControllerModel();
+	if (!ControllerModel.IsValid())
+	{
+		return;
+	}
+
+	FScopedTransaction Transaction(LOCTEXT("AutoBindEntities", "Auto bind entities to controllers"));
+
+	const TArray<FGuid>& SelectedFieldIds = InExposedEntityDragDrop.GetSelectedFieldsId();
+
+	TArray<URCController*, TInlineAllocator<1>> CreatedControllers;
+	TArray<TSharedRef<const FRemoteControlProperty>, TInlineAllocator<1>> SourcePropertyEntities;
+
+	CreatedControllers.Reserve(SelectedFieldIds.Num());
+	SourcePropertyEntities.Reserve(SelectedFieldIds.Num());
+
+	for (const FGuid& ExposedEntityId : SelectedFieldIds)
+	{
+		if (TSharedPtr<const FRemoteControlProperty> RemoteControlProperty = InContext.Preset->GetExposedEntity<FRemoteControlProperty>(ExposedEntityId).Pin())
+		{
+			if (URCController* NewController = UE::RCUIHelpers::CreateControllerFromEntity(InContext.Preset, RemoteControlProperty))
+			{
+				CreatedControllers.Emplace(NewController);
+				SourcePropertyEntities.Add(RemoteControlProperty.ToSharedRef());
+			}
+		}
+	}
+
+	if (CreatedControllers.IsEmpty())
+	{
+		Transaction.Cancel();
+		return;
+	}
+
+	// Refresh panel list so that the controller models are created
+	InContext.ControllerPanelList->Reset();
+
+	// Reorder the controller items to match the target drop zone
+	{
+		const FRCControllerPanelListItem NewControllerItem = InContext.ControllerPanelList->FindListItemByControllerItem(ControllerModel.ToSharedRef());
+		const TArray<TSharedPtr<FRCControllerModel>> ControllersToMove = InContext.ControllerPanelList->FindControllerItemsByObject(CreatedControllers);
+
+		TArray<FRCControllerPanelListItem> ControllerItemsToMove;
+		ControllerItemsToMove.Reserve(CreatedControllers.Num());
+
+		for (const TSharedPtr<FRCControllerModel>& ControllerToMove : ControllersToMove)
+		{
+			if (ControllerToMove.IsValid())
+			{
+				if (const FRCControllerPanelListItem ControllerModelItem = InContext.ControllerPanelList->FindListItemByControllerItem(ControllerToMove.ToSharedRef()))
+				{
+					ControllerItemsToMove.Add(ControllerModelItem);
+				}
+			}
+		}
+
+		InContext.ControllerPanelList->ReorderItems(ControllerItemsToMove, NewControllerItem, EItemDropZone::BelowItem);
+	}
+
+	check(CreatedControllers.Num() == SourcePropertyEntities.Num());
+
+	// Create Bind Behaviour and Bind to the property
+	for (int32 Index = 0; Index < CreatedControllers.Num(); ++Index)
+	{
+		constexpr bool bExecuteBind = true;
+		InContext.ControllerPanelList->CreateBindBehaviourAndAssignTo(CreatedControllers[Index], SourcePropertyEntities[Index], bExecuteBind);
+	}
+}
+
+TSharedRef<SWidget> SRCControllerPanelControllerRow::GetTypeColorWidget()
+{
+	TSharedPtr<FRCControllerModel> ControllerModel = GetControllerModel();
+
+	if (!ControllerModel.IsValid())
+	{
+		return SNullWidget::NullWidget;
+	}
+
+	URCController* Controller = Cast<URCController>(ControllerModel->GetVirtualProperty());
+	
+	if (!Controller)
+	{
+		return SNullWidget::NullWidget;
+	}
+
+	using namespace UE::RCCustomControllers;
+	using namespace UE::RemoteControl::UI::Private;
+
+	if (const FProperty* Property = Controller->GetProperty())
+	{
+		return SNew(SBox)
+			.Padding(this, &SRCControllerPanelControllerRow::GetTypeColorPadding)
+			[
+				UE::RCUIHelpers::GetTypeColorWidget(Property)
+			];
+	}
+
+	return SNullWidget::NullWidget;
+}
+
+FMargin SRCControllerPanelControllerRow::GetTypeColorPadding() const
+{
+	TSharedPtr<FRCControllerModel> ControllerModel = GetControllerModel();
+
+	if (!ControllerModel.IsValid())
+	{
+		return FMargin(0.f);
+	}
+
+	using namespace UE::RemoteControl::UI::Private;
+
+	if (URCController* Controller = Cast<URCController>(ControllerModel->GetVirtualProperty()))
+	{
+		if (!Controller->GetMetadataValue(FRCControllerPropertyInfo::CategoryName).IsEmpty())
+		{ 
+			return FMargin(10.f, 0.f, 0.f, 0.f);
+		}
+	}
+
+	return FMargin(0.f);
+}
+
+FReply SRCControllerPanelControllerRow::OnDragDetected(const FGeometry& InMyGeometry, const FPointerEvent& InMouseEvent)
+{
+	if (FSuperRowType::OnDragDetected(InMyGeometry, InMouseEvent).IsEventHandled())
+	{
+		return FReply::Handled();
+	}
+
+	const TSharedPtr<SRCControllerPanelList> ControllerPanelList = ControllerPanelListWeak.Pin();
+	if (!ControllerPanelList.IsValid())
+	{
+		return FReply::Unhandled();
+	}
+
+	if (InMouseEvent.IsMouseButtonDown(EKeys::LeftMouseButton))
+	{
+		TSharedRef<FRCControllerPanelDragDrop> ControllerDragDrop = MakeShared<FRCControllerPanelDragDrop>(ControllerPanelList->GetSelectedItems());
+		return FReply::Handled().BeginDragDrop(ControllerDragDrop);
+	}
+	return FReply::Unhandled();
+}
+
+FReply SRCControllerPanelControllerRow::OnMouseButtonDoubleClick(const FGeometry& InMyGeometry, const FPointerEvent& InMouseEvent)
+{
+	if (TSharedPtr<FRCControllerModel> ControllerModel = GetControllerModel())
+	{
+		ControllerModel->EnterDescriptionEditingMode();
+		return FReply::Handled();
+	}
+	return FSuperRowType::OnMouseButtonDoubleClick(InMyGeometry, InMouseEvent);
+}
+
+#undef LOCTEXT_NAMESPACE
diff --git a/Engine/Plugins/VirtualProduction/RemoteControl/Source/RemoteControlUI/Private/UI/Controller/SRCControllerPanelControllerRow.h b/Engine/Plugins/VirtualProduction/RemoteControl/Source/RemoteControlUI/Private/UI/Controller/SRCControllerPanelControllerRow.h
new file mode 100644
--- /dev/null
+++ b/Engine/Plugins/VirtualProduction/RemoteControl/Source/RemoteControlUI/Private/UI/Controller/SRCControllerPanelControllerRow.h
@@ -0,0 +1,70 @@
+// Copyright Epic Games, Inc. All Rights Reserved.
+
+#pragma once
+
+#include "UI/Controller/RCControllerPanelListItem.h"
+#include "Widgets/Views/SListView.h"
+#include "Widgets/Views/STableRow.h"
+
+class FExposedEntityDragDrop;
+class FRCControllerModel;
+class SRCControllerPanelList;
+class URemoteControlPreset;
+
+struct FRCControllerColumns
+{
+	static const FLazyName TypeColor;
+	static const FLazyName ControllerId;
+	static const FLazyName Description;
+	static const FLazyName Value;
+	static const FLazyName FieldId;
+	static const FLazyName ValueTypeSelection;
+};
+
+class SRCControllerPanelControllerRow : public SMultiColumnTableRow<FRCControllerPanelListItem>
+{
+public:
+	SLATE_BEGIN_ARGS(SRCControllerPanelControllerRow) {}
+		SLATE_STYLE_ARGUMENT(FTableRowStyle, Style)
+	SLATE_END_ARGS()
+
+	void Construct(const FArguments& InArgs, const TSharedRef<STableViewBase>& InOwnerTableView,
+		TSharedRef<SRCControllerPanelList> InControllerPanelList, FRCControllerPanelListItem InControllerItem);
+
+	TSharedRef<SWidget> GenerateWidgetForColumn(const FName& InColumnName) override;
+
+	TSharedPtr<FRCControllerModel> GetControllerModel() const;
+
+private:
+	/** Called to validate whether a drag drop event is accepted for a controller item row */
+	TOptional<EItemDropZone> CanAcceptDrop(const FDragDropEvent& InDragDropEvent, EItemDropZone InDropZone, FRCControllerPanelListItem InItem);
+
+	/** Called when a drop happens on this controller item row */
+	FReply OnAcceptDrop(const FDragDropEvent& InDragDropEvent, EItemDropZone InDropZone, FRCControllerPanelListItem InItem);
+
+	struct FDragDropContext
+	{
+		URemoteControlPreset* Preset;
+		const TSharedPtr<SRCControllerPanelList>& ControllerPanelList;
+		const FRCControllerPanelListItem& Item;
+		EItemDropZone DropZone;
+	};
+
+	/** Creates a bind behavior on the existing controller items for all the given entities found in the drag drop */
+	void CreateBindBehaviorsFromEntities(const FExposedEntityDragDrop& InExposedEntityDragDrop, const FDragDropContext& InContext);
+
+	/** Creates new controllers next to this item row (above/below depending on drop zone) and creates a bind behavior for each new controller for all the given entities found in the drag drop */
+	void CreateControllersFromEntities(const FExposedEntityDragDrop& InExposedEntityDragDrop, const FDragDropContext& InContext);
+
+	TSharedRef<SWidget> GetTypeColorWidget();
+
+	FMargin GetTypeColorPadding() const;
+
+	//~ Begin SWidget
+	virtual FReply OnDragDetected(const FGeometry& InMyGeometry, const FPointerEvent& InMouseEvent) override;
+	virtual FReply OnMouseButtonDoubleClick(const FGeometry& InMyGeometry, const FPointerEvent& InMouseEvent) override;
+	//~ End SWidget
+
+	TWeakPtr<SRCControllerPanelList> ControllerPanelListWeak;
+	FRCControllerPanelListItem ControllerItem;
+};
\ No newline at end of file
diff --git a/Engine/Plugins/VirtualProduction/RemoteControl/Source/RemoteControlUI/Private/UI/Controller/SRCControllerPanelList.cpp b/Engine/Plugins/VirtualProduction/RemoteControl/Source/RemoteControlUI/Private/UI/Controller/SRCControllerPanelList.cpp
--- a/Engine/Plugins/VirtualProduction/RemoteControl/Source/RemoteControlUI/Private/UI/Controller/SRCControllerPanelList.cpp
+++ b/Engine/Plugins/VirtualProduction/RemoteControl/Source/RemoteControlUI/Private/UI/Controller/SRCControllerPanelList.cpp
@@ -15,6 +15,7 @@
 #include "IPropertyRowGenerator.h"
 #include "IRemoteControlModule.h"
 #include "Materials/MaterialInterface.h"
+#include "RCCategoryModel.h"
 #include "RCControllerModel.h"
 #include "RCMultiController.h"
 #include "RCVirtualProperty.h"
@@ -24,11 +25,13 @@
 #include "ScopedTransaction.h"
 #include "SDropTarget.h"
 #include "SlateOptMacros.h"
-#include "SRCControllerItemRow.h"
+#include "SRCControllerPanelCategoryRow.h"
+#include "SRCControllerPanelControllerRow.h"
 #include "SRCControllerPanel.h"
 #include "Styling/RemoteControlStyles.h"
 #include "UI/Action/SRCActionPanelList.h"
 #include "UI/BaseLogicUI/RCLogicModeBase.h"
+#include "UI/Controller/SRCControllerPanel.h"
 #include "UI/RCUIHelpers.h"
 #include "UI/RemoteControlPanelStyle.h"
 #include "UI/SRCPanelDragHandle.h"
@@ -36,24 +39,10 @@
 #include "UI/SRCPanelExposedEntity.h"
 #include "UI/SRemoteControlPanel.h"
 #include "Widgets/Views/SHeaderRow.h"
+#include "Widgets/Views/STreeView.h"
 
 #define LOCTEXT_NAMESPACE "SRCControllerPanelList"
 
-TArray<TSharedPtr<FRCControllerModel>> FRCControllerDragDrop::ResolveControllers() const
-{
-	TArray<TSharedPtr<FRCControllerModel>> ResolvedControllers;
-	ResolvedControllers.Reserve(ControllersWeak.Num());
-
-	for (const TWeakPtr<FRCControllerModel>& WeakController : ControllersWeak)
-	{
-		if (TSharedPtr<FRCControllerModel> ResolvedController = WeakController.Pin())
-		{
-			ResolvedControllers.Add(ResolvedController);
-		}
-	}
-	return ResolvedControllers;
-}
-
 void SRCControllerPanelList::Construct(const FArguments& InArgs, const TSharedRef<SRCControllerPanel> InControllerPanel, const TSharedRef<SRemoteControlPanel> InRemoteControlPanel)
 {
 	SRCLogicPanelListBase::Construct(SRCLogicPanelListBase::FArguments(), InControllerPanel, InRemoteControlPanel);
@@ -62,10 +51,11 @@ void SRCControllerPanelList::Construct(const FArguments& InArgs, const TSharedRe
 
 	RCPanelStyle = &FRemoteControlPanelStyle::Get()->GetWidgetStyle<FRCPanelStyle>("RemoteControlPanel.LogicControllersPanel");
 
-	ListView = SNew(SListView<TSharedPtr<FRCControllerModel>>)
-		.ListItemsSource(&ControllerItems)
+	TreeView = SNew(STreeView<FRCControllerPanelListItem>)
+		.TreeItemsSource(&RootItems)
+		.OnGetChildren(this, &SRCControllerPanelList::OnGetChildren)
 		.OnSelectionChanged(this, &SRCControllerPanelList::OnTreeSelectionChanged)
-		.OnGenerateRow(this, &SRCControllerPanelList::OnGenerateWidgetForList)
+		.OnGenerateRow(this, &SRCControllerPanelList::OnGenerateWidgetForTree)
 		.SelectionMode(ESelectionMode::Multi)
 		.OnContextMenuOpening(this, &SRCLogicPanelListBase::GetContextMenuWidget)
 		.HeaderRow(
@@ -74,7 +64,7 @@ void SRCControllerPanelList::Construct(const FArguments& InArgs, const TSharedRe
 
 			+ SHeaderRow::Column(FRCControllerColumns::TypeColor)
 			.DefaultLabel(FText())
-			.FixedWidth(15)
+			.FixedWidth(25)
 			.HeaderContentPadding(RCPanelStyle->HeaderRowPadding)
 
 			+ SHeaderRow::Column(FRCControllerColumns::ControllerId)
@@ -97,11 +87,11 @@ void SRCControllerPanelList::Construct(const FArguments& InArgs, const TSharedRe
 		SNew(SDropTarget)
 		.VerticalImage(FRemoteControlPanelStyle::Get()->GetBrush("RemoteControlPanel.VerticalDash"))
 		.HorizontalImage(FRemoteControlPanelStyle::Get()->GetBrush("RemoteControlPanel.HorizontalDash"))
-		.OnDropped_Lambda([this](const FGeometry& InGeometry, const FDragDropEvent& InDragDropEvent) { return SRCControllerPanelList::OnControllerListViewDragDrop(InDragDropEvent.GetOperation()); })
+		.OnDropped_Lambda([this](const FGeometry& InGeometry, const FDragDropEvent& InDragDropEvent) { return SRCControllerPanelList::OnControllerTreeViewDragDrop(InDragDropEvent.GetOperation()); })
 		.OnAllowDrop(this, &SRCControllerPanelList::OnAllowDrop)
 		.OnIsRecognized(this, &SRCControllerPanelList::OnAllowDrop)
 		[
-			ListView.ToSharedRef()
+			TreeView.ToSharedRef()
 		]
 	];
 
@@ -137,7 +127,7 @@ int32 SRCControllerPanelList::Num() const
 
 int32 SRCControllerPanelList::NumSelectedLogicItems() const
 {
-	return ListView->GetNumItemsSelected();
+	return TreeView->GetNumItemsSelected();
 }
 
 void SRCControllerPanelList::Reset()
@@ -151,7 +141,7 @@ void SRCControllerPanelList::Reset()
 	}
 
 	// Cache Controller Selection
-	const TArray<TSharedPtr<FRCControllerModel>> SelectedControllers = ListView->GetSelectedItems();
+	const TArray<FRCControllerPanelListItem> SelectedItem = TreeView->GetSelectedItems();
 
 	ControllerItems.Empty();
 
@@ -181,11 +171,18 @@ void SRCControllerPanelList::Reset()
 		TArray<TSharedRef<IDetailTreeNode>> Children;
 		CategoryNode->GetChildren(Children);
 
-		ControllerItems.SetNumZeroed(Children.Num());
+		ControllerItems.Reserve(Children.Num());
 
 		for (TSharedRef<IDetailTreeNode>& Child : Children)
 		{
-			FProperty* Property = Child->CreatePropertyHandle()->GetProperty();
+			TSharedPtr<IPropertyHandle> ChildHandle = Child->CreatePropertyHandle();
+
+			if (!ChildHandle.IsValid())
+			{
+				continue;
+			}
+
+			FProperty* Property = ChildHandle->GetProperty();
 			check(Property);
 
 			if (URCVirtualPropertyBase* Controller = Preset->GetController(Property->GetFName()))
@@ -200,28 +197,25 @@ void SRCControllerPanelList::Reset()
 					// there's at least one Field Id set, let's show their column
 					bShowFieldIdsColumn = true;
 				}
-				
+
 				// MultiController Mode: only showing one Controller per Field Id
 				if (bIsInMultiControllerMode && Preset->GetControllersByFieldId(FieldId).Num() > 1)
-				{					
+				{
 					bIsMultiController = MultiControllers.TryToAddAsMultiController(Controller);
 					bIsVisible = bIsMultiController;
 				}
 
 				if (bIsVisible)
 				{
-					if (ensureAlways(ControllerItems.IsValidIndex(Controller->DisplayIndex)))
+					const TSharedRef<FRCControllerModel> ControllerModel = MakeShared<FRCControllerModel>(Controller, Child, RemoteControlPanel);
+					ControllerModel->Initialize();
+					ControllerModel->OnValueChanged.AddSP(this, &SRCControllerPanelList::OnControllerValueChanged, bIsMultiController);
+					if (bIsMultiController)
 					{
-						const TSharedRef<FRCControllerModel> ControllerModel = MakeShared<FRCControllerModel>(Controller, Child, RemoteControlPanel);
-						ControllerModel->Initialize();
-						ControllerModel->OnValueChanged.AddSP(this, &SRCControllerPanelList::OnControllerValueChanged, bIsMultiController);
-						if (bIsMultiController)
-						{
-							ControllerModel->SetMultiController(bIsMultiController);
-							ControllerModel->OnValueTypeChanged.AddSP(this, &SRCControllerPanelList::OnControllerValueTypeChanged);
-						}
-						ControllerItems[Controller->DisplayIndex] = ControllerModel;
+						ControllerModel->SetMultiController(bIsMultiController);
+						ControllerModel->OnValueTypeChanged.AddSP(this, &SRCControllerPanelList::OnControllerValueTypeChanged);
 					}
+					ControllerItems.Add(ControllerModel);
 				}
 			}
 		}
@@ -231,20 +225,64 @@ void SRCControllerPanelList::Reset()
 	if (bIsInMultiControllerMode)
 	{
 		Algo::Sort(ControllerItems, [](const TSharedPtr<FRCControllerModel>& A, const TSharedPtr<FRCControllerModel>& B)
-		{
-			if (A.IsValid() && B.IsValid())
 			{
-				const URCVirtualPropertyBase* ControllerA = A->GetVirtualProperty();
-				const URCVirtualPropertyBase* ControllerB = B->GetVirtualProperty();
+				if (A.IsValid() && B.IsValid())
+				{
+					const URCVirtualPropertyBase* ControllerA = A->GetVirtualProperty();
+					const URCVirtualPropertyBase* ControllerB = B->GetVirtualProperty();
 
-				if (ControllerA && ControllerB)
+					if (ControllerA && ControllerB)
+					{
+						return ControllerA->FieldId.FastLess(ControllerB->FieldId);
+					}
+				}
+
+				return false;
+			});
+	}
+	else
+	{
+		Algo::Sort(ControllerItems, [](const TSharedPtr<FRCControllerModel>& A, const TSharedPtr<FRCControllerModel>& B)
+			{
+				if (A.IsValid() && B.IsValid())
 				{
-					return ControllerA->FieldId.FastLess(ControllerB->FieldId);
+					const URCVirtualPropertyBase* ControllerA = A->GetVirtualProperty();
+					const URCVirtualPropertyBase* ControllerB = B->GetVirtualProperty();
+
+					if (ControllerA && ControllerB)
+					{
+						return ControllerA->DisplayIndex < ControllerB->DisplayIndex;
+					}
 				}
-			}
 
-			return false;
-		});
+				return B.IsValid();
+			});
+	}
+
+	using namespace UE::RemoteControl::UI::Private;
+
+	TSet<FString> ValidCategories;
+
+	CategoryItems.Empty();
+
+	if (URCVirtualPropertyContainerBase* ControllerContainer = Preset->GetControllerContainer())
+	{
+		CategoryItems.Reserve(ControllerContainer->Categories.Num());
+
+		for (const FRCVirtualPropertyCategory& Category : ControllerContainer->Categories)
+		{
+			CategoryItems.Add(MakeShared<FRCCategoryModel>(Category.Id, RemoteControlPanel));
+			CategoryItems.Last()->Initialize();
+			
+			ValidCategories.Add(Category.Id.ToString());
+		}
+
+		Algo::Sort(
+			CategoryItems, 
+			[](const TSharedPtr<FRCCategoryModel>& A, const TSharedPtr<FRCCategoryModel>& B)
+			{
+				return A->GetCategory()->DisplayIndex < B->GetCategory()->DisplayIndex;
+			});
 	}
 
 	ShowFieldIdHeaderColumn(bShowFieldIdsColumn);
@@ -267,29 +305,52 @@ void SRCControllerPanelList::Reset()
 					.FillWidth(0.2f)
 					.HeaderContentPadding(RCPanelStyle->HeaderRowPadding)
 				);
-			}			
+			}
 		}
 	}
-	
-	ListView->RebuildList();
 
-	// Restore Controller Selection
-	for (const TSharedPtr<FRCControllerModel>& ControllerModel : SelectedControllers)
+	for (const TSharedPtr<FRCControllerModel>& ControllerModel : ControllerItems)
 	{
 		if (ControllerModel.IsValid())
 		{
-			const FName SelectedControllerName = ControllerModel->GetPropertyName();
+			if (URCVirtualPropertyBase* VirtualProperty = ControllerModel->GetVirtualProperty())
+			{
+				const FString ParentCategoryName = VirtualProperty->GetMetadataValue(FRCControllerPropertyInfo::CategoryName);
 
-			const TSharedPtr<FRCControllerModel>* SelectedController = ControllerItems.FindByPredicate([&SelectedControllerName]
-				(const TSharedPtr<FRCControllerModel>& InControllerModel)
+				if (!ParentCategoryName.IsEmpty() && !ValidCategories.Contains(ParentCategoryName))
 				{
-					// Internal PropertyName is unique, so we use that
-					return InControllerModel.IsValid() && SelectedControllerName == InControllerModel->GetPropertyName();
-				});
+					VirtualProperty->RemoveMetadataValue(FRCControllerPropertyInfo::CategoryName);
+				}
+			}
+		}
+	}
+
+	UpdateRootItems();
+
+	TreeView->RebuildList();
+
+	// Restore Controller Selection
+	for (const FRCControllerPanelListItem& Item : SelectedItem)
+	{
+		if (Item.Type == ERCControllerPanelListItemType::Controller && ControllerItems.IsValidIndex(Item.Index))
+		{
+			const TSharedPtr<FRCControllerModel> ControllerModel = ControllerItems[Item.Index];
 
-			if (SelectedController)
+			if (ControllerModel.IsValid())
 			{
-				ListView->SetItemSelection(*SelectedController, true);
+				const FName SelectedControllerName = ControllerModel->GetPropertyName();
+
+				const TSharedPtr<FRCControllerModel>* SelectedController = ControllerItems.FindByPredicate([&SelectedControllerName]
+				(const TSharedPtr<FRCControllerModel>& InControllerModel)
+					{
+						// Internal PropertyName is unique, so we use that
+						return InControllerModel.IsValid() && SelectedControllerName == InControllerModel->GetPropertyName();
+					});
+
+				if (SelectedController)
+				{
+					TreeView->SetItemSelection(Item, true);
+				}
 			}
 		}
 	}
@@ -297,42 +358,258 @@ void SRCControllerPanelList::Reset()
 	UpdateEntityUsage();
 }
 
-TSharedRef<ITableRow> SRCControllerPanelList::OnGenerateWidgetForList(TSharedPtr<FRCControllerModel> InItem, const TSharedRef<STableViewBase>& OwnerTable)
+void SRCControllerPanelList::UpdateRootItems()
 {
-	return SNew(SRCControllerItemRow, OwnerTable, InItem.ToSharedRef(), SharedThis(this))
-		.Style(&RCPanelStyle->TableRowStyle);
+	using namespace UE::RemoteControl::UI::Private;
+
+	RootItems.Empty(CategoryItems.Num() + ControllerItems.Num());
+
+	for (int32 Index = 0; Index < CategoryItems.Num(); ++Index)
+	{
+		RootItems.Emplace(
+			ERCControllerPanelListItemType::Category,
+			Index
+		);
+	}
+
+	for (int32 Index = 0; Index < ControllerItems.Num(); ++Index)
+	{
+		if (ControllerItems[Index].IsValid())
+		{
+			if (URCVirtualPropertyBase* VirtualProperty = ControllerItems[Index]->GetVirtualProperty())
+			{
+				if (VirtualProperty->GetMetadataValue(FRCControllerPropertyInfo::CategoryName).IsEmpty())
+				{
+					RootItems.Emplace(
+						ERCControllerPanelListItemType::Controller,
+						Index
+					);
+				}
+			}
+		}
+	}
+
+	auto GetDisplayIndex = [this](const FRCControllerPanelListItem& InItem) -> int32
+		{
+			switch (InItem.Type)
+			{
+				default:
+					return INDEX_NONE;
+
+				case ERCControllerPanelListItemType::Category:
+					if (CategoryItems.IsValidIndex(InItem.Index) && CategoryItems[InItem.Index].IsValid())
+					{
+						if (TOptional<FRCVirtualPropertyCategory> Category = CategoryItems[InItem.Index]->GetCategory())
+						{
+							return Category->DisplayIndex;
+						}
+					}
+
+					return INDEX_NONE;
+
+				case ERCControllerPanelListItemType::Controller:
+					if (ControllerItems.IsValidIndex(InItem.Index) && ControllerItems[InItem.Index].IsValid())
+					{
+						if (URCVirtualPropertyBase* Controller = ControllerItems[InItem.Index]->GetVirtualProperty())
+						{
+							return Controller->DisplayIndex;
+						}					
+					}
+
+					return INDEX_NONE;
+			}
+		};
+
+	Algo::Sort(
+		RootItems,
+		[&GetDisplayIndex](const FRCControllerPanelListItem& A, const FRCControllerPanelListItem& B)
+		{
+			return GetDisplayIndex(A) < GetDisplayIndex(B);
+	});
 }
 
-void SRCControllerPanelList::OnTreeSelectionChanged(TSharedPtr<FRCControllerModel> InItem, ESelectInfo::Type InSelectInfo)
+void SRCControllerPanelList::OnItemRemoved(const FRCControllerPanelListItem& InRemovedItem)
 {
-	if (TSharedPtr<SRCControllerPanel> ControllerPanel = ControllerPanelWeakPtr.Pin())
+	// Remove category and put all items into the base list
+	if (InRemovedItem.Type != ERCControllerPanelListItemType::Category)
+	{
+		return;
+	}
+
+	if (!CategoryItems.IsValidIndex(InRemovedItem.Index) || !CategoryItems[InRemovedItem.Index].IsValid())
+	{
+		return;
+	}
+
+	TOptional<FRCVirtualPropertyCategory> Category = CategoryItems[InRemovedItem.Index]->GetCategory();
+
+	if (!Category.IsSet())
 	{
-		if (const TSharedPtr<SRemoteControlPanel> RemoteControlPanel = ControllerPanel->GetRemoteControlPanel())
+		return;
+	}
+
+	const FString CategoryPropertyName = Category->Id.ToString();
+
+	// Shift all display indices higher than the deleted index down by 1
+	for (int32 ControllerIndex = 0; ControllerIndex < ControllerItems.Num(); ControllerIndex++)
+	{
+		if (ensure(ControllerItems[ControllerIndex]))
 		{
-			// Tree selection gets updated and notifies with the first element in the selection set.
-			// Controller Behavior / Actions only support viewing one behavior at a time
-			if (InItem != SelectedControllerItemWeakPtr)
+			if (URCVirtualPropertyBase* Controller = ControllerItems[ControllerIndex]->GetVirtualProperty())
 			{
-				SelectedControllerItemWeakPtr = InItem;
-				RemoteControlPanel->OnControllerSelectionChanged.Broadcast(InItem, InSelectInfo);
-				RemoteControlPanel->OnBehaviourSelectionChanged.Broadcast(InItem.IsValid() ? InItem->GetSelectedBehaviourModel() : nullptr);
+				using namespace UE::RemoteControl::UI::Private;
+
+				if (Controller->GetMetadataValue(FRCControllerPropertyInfo::CategoryName) == CategoryPropertyName)
+				{
+					if (GUndo)
+					{
+						Controller->Modify();
+					}
+
+					Controller->RemoveMetadataValue(FRCControllerPropertyInfo::CategoryName);
+				}
 			}
 		}
 	}
+
+	UpdateRootItems();
+}
+
+TSharedRef<ITableRow> SRCControllerPanelList::OnGenerateWidgetForTree(FRCControllerPanelListItem InItem, const TSharedRef<STableViewBase>& OwnerTable)
+{
+	switch (InItem.Type)
+	{
+		default:
+			break;
+
+		case ERCControllerPanelListItemType::Category:
+			if (CategoryItems.IsValidIndex(InItem.Index) && CategoryItems[InItem.Index].IsValid())
+			{
+				return SNew(SRCControllerPanelCategoryRow, OwnerTable, SharedThis(this), InItem)
+					.Style(&RCPanelStyle->TableRowStyle);
+			}
+			break;
+
+		case ERCControllerPanelListItemType::Controller:
+			if (ControllerItems.IsValidIndex(InItem.Index) && ControllerItems[InItem.Index].IsValid())
+			{
+				return SNew(SRCControllerPanelControllerRow, OwnerTable, SharedThis(this), InItem)
+					.Style(&RCPanelStyle->TableRowStyle);
+			}
+			break;
+	}
+
+	return SNew(STableRow<FRCControllerPanelListItem>, TreeView.ToSharedRef());
+}
+
+bool SRCControllerPanelList::IsItemCategoryExpanded(FRCControllerPanelListItem InItem) const
+{
+	using namespace UE::RemoteControl::UI::Private;
+
+	if (InItem.Type != ERCControllerPanelListItemType::Category)
+	{
+		return true;
+	}
+
+	if (TSharedPtr<ITableRow> Row = TreeView->WidgetFromItem(InItem))
+	{
+		return !Row->DoesItemHaveChildren() || Row->IsItemExpanded();
+	}
+
+	return false;
+}
+
+void SRCControllerPanelList::OnGetChildren(FRCControllerPanelListItem InItem, TArray<FRCControllerPanelListItem>& OutChildren)
+{
+	using namespace UE::RemoteControl::UI::Private;
+
+	if (InItem.Type != ERCControllerPanelListItemType::Category)
+	{
+		return;
+	}
+
+	if (!CategoryItems.IsValidIndex(InItem.Index) || !CategoryItems[InItem.Index].IsValid())
+	{
+		return;
+	}
+
+	TOptional<FRCVirtualPropertyCategory> Category = CategoryItems[InItem.Index]->GetCategory();
+
+	if (!Category.IsSet())
+	{
+		return;
+	}
+
+	const FString CategoryPropertyName = Category->Id.ToString();
+
+	for (int32 Index = 0; Index < ControllerItems.Num(); ++Index)
+	{
+		if (ControllerItems[Index].IsValid())
+		{
+			if (URCVirtualPropertyBase* VirtualProperty = ControllerItems[Index]->GetVirtualProperty())
+			{
+				if (VirtualProperty->GetMetadataValue(FRCControllerPropertyInfo::CategoryName) == CategoryPropertyName)
+				{
+					OutChildren.Emplace(
+						ERCControllerPanelListItemType::Controller,
+						Index
+					);
+				}
+			}
+		}
+	}
+}
+
+void SRCControllerPanelList::OnTreeSelectionChanged(FRCControllerPanelListItem InItem, ESelectInfo::Type InSelectInfo)
+{
+	TSharedPtr<SRCControllerPanel> ControllerPanel = ControllerPanelWeakPtr.Pin();
+
+	if (!ControllerPanel.IsValid())
+	{
+		return;
+	}
+
+	const TSharedPtr<SRemoteControlPanel> RemoteControlPanel = ControllerPanel->GetRemoteControlPanel();
+
+	if (!RemoteControlPanel.IsValid())
+	{
+		return;
+	}
+
+	TSharedPtr<FRCControllerModel> ControllerModel;
+
+	if (InItem.Type == ERCControllerPanelListItemType::Controller
+		&& ControllerItems.IsValidIndex(InItem.Index)
+		&& ControllerItems[InItem.Index].IsValid())
+	{
+		ControllerModel = ControllerItems[InItem.Index].ToSharedRef();
+	}
+
+	// Tree selection gets updated and notifies with the first element in the selection set.
+	// Controller Behavior / Actions only support viewing one behavior at a time
+	if (ControllerModel != SelectedControllerItemWeakPtr)
+	{
+		SelectedControllerItemWeakPtr = ControllerModel;
+		RemoteControlPanel->OnControllerSelectionChanged.Broadcast(ControllerModel, InSelectInfo);
+		RemoteControlPanel->OnBehaviourSelectionChanged.Broadcast(ControllerModel ? ControllerModel->GetSelectedBehaviourModel() : nullptr);
+	}
 }
 
 void SRCControllerPanelList::SelectController(URCController* InController)
 {
-	for (TSharedPtr<FRCControllerModel> ControllerItem : ControllerItems)
+	for (int32 Index = 0; Index < ControllerItems.Num(); ++Index)
 	{
-		if (!ensure(ControllerItem))
+		if (!ensure(ControllerItems[Index].IsValid()))
 		{
 			continue;
 		}
 
-		if (ControllerItem->GetVirtualProperty() == InController)
+		if (ControllerItems[Index]->GetVirtualProperty() == InController)
 		{
-			ListView->SetSelection(ControllerItem);
+			TreeView->SetSelection({
+				ERCControllerPanelListItemType::Controller,
+				Index
+			});
 		}
 	}
 }
@@ -445,9 +722,15 @@ int32 SRCControllerPanelList::RemoveModel(const TSharedPtr<FRCLogicModeBase> InM
 					Preset->Modify();
 
 					// Remove Model from Data Container
+					const FString ControllerPropertyName = SelectedController->GetPropertyName().ToString();
+					const bool bHasPropertyName = !ControllerPropertyName.IsEmpty();
+
 					const bool bRemoved = Preset->RemoveController(SelectedController->GetPropertyName());
+
 					if (bRemoved)
 					{
+						int32 FoundIndex = INDEX_NONE;
+
 						// Shift all display indices higher than the deleted index down by 1
 						for (int32 ControllerIndex = 0; ControllerIndex < ControllerItems.Num(); ControllerIndex++)
 						{
@@ -455,16 +738,39 @@ int32 SRCControllerPanelList::RemoveModel(const TSharedPtr<FRCLogicModeBase> InM
 							{
 								if (URCVirtualPropertyBase* Controller = ControllerItems[ControllerIndex]->GetVirtualProperty())
 								{
-									if (Controller->DisplayIndex > DisplayIndexToRemove)
+									if (Controller->DisplayIndex == DisplayIndexToRemove)
+									{
+										FoundIndex = Controller->DisplayIndex;
+									}
+									else if (Controller->DisplayIndex > DisplayIndexToRemove)
 									{
 										Controller->Modify();
-
 										Controller->SetDisplayIndex(Controller->DisplayIndex - 1);
+										FoundIndex = ControllerIndex;
+									}
+
+									if (bHasPropertyName)
+									{
+										using namespace UE::RemoteControl::UI::Private;
+
+										if (Controller->GetMetadataValue(FRCControllerPropertyInfo::CategoryName) == ControllerPropertyName)
+										{
+											Controller->Modify();
+											Controller->RemoveMetadataValue(FRCControllerPropertyInfo::CategoryName);
+										}
 									}
 								}
 							}
 						}
 
+						if (FoundIndex != INDEX_NONE)
+						{
+							OnItemRemoved({
+								ERCControllerPanelListItemType::Controller,
+								FoundIndex
+							});
+						}
+
 						return 1; // Remove Count
 					}
 				}
@@ -477,15 +783,46 @@ int32 SRCControllerPanelList::RemoveModel(const TSharedPtr<FRCLogicModeBase> InM
 
 bool SRCControllerPanelList::IsListFocused() const
 {
-	return ListView->HasAnyUserFocus().IsSet() || ContextMenuWidgetCached.IsValid();
+	return TreeView->HasAnyUserFocus().IsSet() || ContextMenuWidgetCached.IsValid();
 }
 
 void SRCControllerPanelList::DeleteSelectedPanelItems()
 {
 	FScopedTransaction Transaction(LOCTEXT("DeleteSelectedItems", "Delete Selected Items"));
+	TArray<FRCControllerPanelListItem> SelectedItems = TreeView->GetSelectedItems();
+
+	TArray<TSharedPtr<FRCControllerModel>> SelectedControllers;
+	SelectedControllers.Reserve(SelectedItems.Num());
+
+	TArray<int32> SelectedControllerIndices;
+	SelectedControllerIndices.Reserve(SelectedItems.Num());
 
-	if (DeleteItemsFromLogicPanel<FRCControllerModel>(ControllerItems, ListView->GetSelectedItems()))
+	for (const FRCControllerPanelListItem& Item : SelectedItems)
 	{
+		if (Item.Type != ERCControllerPanelListItemType::Controller)
+		{
+			continue;
+		}
+
+		if (!ControllerItems.IsValidIndex(Item.Index) || !ControllerItems[Item.Index].IsValid())
+		{
+			continue;
+		}
+
+		SelectedControllers.Add(ControllerItems[Item.Index]);
+		SelectedControllerIndices.Add(Item.Index);
+	}
+	
+	if (DeleteItemsFromLogicPanel<FRCControllerModel>(ControllerItems, SelectedControllers))
+	{
+		for (int32 Index = 0; Index < SelectedControllerIndices.Num(); ++Index)
+		{
+			OnItemRemoved({
+				ERCControllerPanelListItemType::Controller,
+				SelectedControllerIndices[Index]
+			});
+		}
+
 		UpdateEntityUsage();
 	}
 	else
@@ -494,14 +831,39 @@ void SRCControllerPanelList::DeleteSelectedPanelItems()
 	}
 }
 
+TArray<FRCControllerPanelListItem> SRCControllerPanelList::GetSelectedItems() const
+{
+	return TreeView->GetSelectedItems();
+}
+
 TArray<TSharedPtr<FRCControllerModel>> SRCControllerPanelList::GetSelectedControllers() const
 {
-	return ListView->GetSelectedItems();
+	TArray<FRCControllerPanelListItem> SelectedItems = TreeView->GetSelectedItems();
+
+	TArray<TSharedPtr<FRCControllerModel>> SelectedControllers;
+	SelectedControllers.Reserve(SelectedItems.Num());
+
+	for (const FRCControllerPanelListItem& Item : SelectedItems)
+	{
+		if (Item.Type != ERCControllerPanelListItemType::Controller)
+		{
+			continue;
+		}
+
+		if (!ControllerItems.IsValidIndex(Item.Index) || !ControllerItems[Item.Index].IsValid())
+		{
+			continue;
+		}
+
+		SelectedControllers.Add(ControllerItems[Item.Index]);
+	}
+
+	return SelectedControllers;
 }
 
 TArray<TSharedPtr<FRCLogicModeBase>> SRCControllerPanelList::GetSelectedLogicItems()
 {
-	return TArray<TSharedPtr<FRCLogicModeBase>>(ListView->GetSelectedItems());
+	return TArray<TSharedPtr<FRCLogicModeBase>>(GetSelectedControllers());
 }
 
 void SRCControllerPanelList::NotifyPreChange(FEditPropertyChain* PropertyAboutToChange)
@@ -525,6 +887,29 @@ void SRCControllerPanelList::EnterRenameMode()
 	}
 }
 
+TSharedPtr<FRCCategoryModel>  SRCControllerPanelList::FindCategoryItemByListItem(const FRCControllerPanelListItem& InItem) const
+{
+	if (InItem.Type == ERCControllerPanelListItemType::Category)
+	{
+		if (CategoryItems.IsValidIndex(InItem.Index))
+		{
+			return CategoryItems[InItem.Index];
+		}
+	}
+
+	return nullptr;
+}
+
+TSharedPtr<FRCControllerModel> SRCControllerPanelList::FindControllerItemByListItem(const FRCControllerPanelListItem& InItem) const
+{
+	if (InItem.Type == ERCControllerPanelListItemType::Controller && ControllerItems.IsValidIndex(InItem.Index))
+	{
+		return ControllerItems[InItem.Index];
+	}
+
+	return nullptr;
+}
+
 TSharedPtr<FRCControllerModel> SRCControllerPanelList::FindControllerItemById(const FGuid& InId) const
 {
 	for (TSharedPtr<FRCControllerModel> ControllerItem : ControllerItems)
@@ -538,6 +923,24 @@ TSharedPtr<FRCControllerModel> SRCControllerPanelList::FindControllerItemById(co
 	return nullptr;
 }
 
+TSharedPtr<FRCControllerModel> SRCControllerPanelList::FindControllerItemByPropertyName(const FName& InPropertyName) const
+{
+	if (InPropertyName.IsNone())
+	{
+		return nullptr;
+	}
+
+	for (TSharedPtr<FRCControllerModel> ControllerItem : ControllerItems)
+	{
+		if (ControllerItem && ControllerItem->GetPropertyName() == InPropertyName)
+		{
+			return ControllerItem;
+		}
+	}
+
+	return nullptr;
+}
+
 TArray<TSharedPtr<FRCControllerModel>> SRCControllerPanelList::FindControllerItemsById(TConstArrayView<FGuid> InIds)
 {
 	TArray<TSharedPtr<FRCControllerModel>> FoundControllerItems;
@@ -570,65 +973,372 @@ TArray<TSharedPtr<FRCControllerModel>> SRCControllerPanelList::FindControllerIte
 	return FoundControllerItems;
 }
 
-int32 SRCControllerPanelList::GetDropIndex(const TSharedPtr<FRCControllerModel>& InItem, EItemDropZone InDropZone) const
+FRCControllerPanelListItem SRCControllerPanelList::FindListItemByControllerItem(const TSharedRef<FRCControllerModel>& InModel) const
 {
-	int32 Index = ControllerItems.Find(InItem);
-	if (Index == INDEX_NONE)
+	for (int32 Index = 0; Index < ControllerItems.Num(); ++Index)
 	{
-		return INDEX_NONE;
+		if (ControllerItems[Index] == InModel)
+		{
+			return {ERCControllerPanelListItemType::Controller, Index};
+		}
 	}
 
-	if (InDropZone == EItemDropZone::BelowItem)
+	return UE::RemoteControl::UI::Private::ControllerPanelListItem::None;
+}
+
+FRCControllerPanelListItem SRCControllerPanelList::FindListItemByCategoryItem(const TSharedPtr<FRCCategoryModel>& InModel) const
+{
+	for (int32 Index = 0; Index < CategoryItems.Num(); ++Index)
 	{
-		++Index;
+		if (CategoryItems[Index] == InModel)
+		{
+			return {ERCControllerPanelListItemType::Category, Index};
+		}
 	}
-	return Index;
+
+	return UE::RemoteControl::UI::Private::ControllerPanelListItem::None;
 }
 
-bool SRCControllerPanelList::ReorderControllerItems(TConstArrayView<TSharedPtr<FRCControllerModel>> InItemsToMove, int32 InTargetIndex)
+bool SRCControllerPanelList::ReorderItems(TConstArrayView<FRCControllerPanelListItem> InItemsToMove, FRCControllerPanelListItem InDroppedOnItem, EItemDropZone InDropZone)
 {
-	// Only allow valid indices and the index 1 more of the end index (i.e. index == num)
-	if (InTargetIndex == INDEX_NONE || InTargetIndex > ControllerItems.Num())
+	if (InItemsToMove.Contains(InDroppedOnItem))
 	{
 		return false;
 	}
 
-	// Remove the items to move.
-	// if these items come before the target index, the target index needs to shift to compensate for the removed items
-	for (TArray<TSharedPtr<FRCControllerModel>>::TIterator Iter(ControllerItems); Iter; ++Iter)
+	// Check we have a valid drop target
+	switch (InDroppedOnItem.Type)
+	{
+		case ERCControllerPanelListItemType::None:
+			return false;
+
+		case ERCControllerPanelListItemType::Category:
+			if (!CategoryItems.IsValidIndex(InDroppedOnItem.Index) || !CategoryItems[InDroppedOnItem.Index].IsValid())
+			{
+				return false;
+			}
+			break;
+
+		case ERCControllerPanelListItemType::Controller:
+			if (!ControllerItems.IsValidIndex(InDroppedOnItem.Index) || !ControllerItems[InDroppedOnItem.Index].IsValid())
+			{
+				return false;
+			}
+			break;
+	}
+
+	// Rebuild list assuming that everything is already in the correct display order
+	TMap<FName, TArray<FRCControllerPanelListItem>> Tree;
+	Tree.Add(NAME_None);
+
+	for (const TSharedPtr<FRCCategoryModel>& CategoryModel : CategoryItems)
+	{
+		if (TOptional<FRCVirtualPropertyCategory> Category = CategoryModel->GetCategory())
+		{
+			Tree.Add(*Category->Id.ToString());
+		}
+	}
+
+	auto AddDroppedItems = [this, &InItemsToMove, &Tree](FName InCategoryName = NAME_None)
+		{
+			for (const FRCControllerPanelListItem& DroppedItem : InItemsToMove)
+			{
+				// Force categories into the root list
+				if (DroppedItem.Type == ERCControllerPanelListItemType::Category)
+				{
+					Tree[NAME_None].Add(DroppedItem);
+				}
+				else
+				{
+					if (!Tree.Contains(InCategoryName))
+					{
+						InCategoryName = NAME_None;
+					}
+
+					Tree[InCategoryName].Add(DroppedItem);
+
+					if (URCVirtualPropertyBase* Controller = ControllerItems[DroppedItem.Index]->GetVirtualProperty())
+					{
+						using namespace UE::RemoteControl::UI::Private;
+
+						if (InCategoryName.IsNone())
+						{
+							Controller->RemoveMetadataValue(FRCControllerPropertyInfo::CategoryName);
+						}
+						else
+						{
+							Controller->SetMetadataValue(FRCControllerPropertyInfo::CategoryName, InCategoryName.ToString());
+						}
+					}
+				}
+			}
+		};
+
+	int32 CategoryIndex = 0;
+	int32 ControllerIndex = 0;
+
+	auto GetNextItem = [this, &InItemsToMove, &CategoryIndex, &ControllerIndex]() -> FRCControllerPanelListItem
+		{
+			int32 NextCategoryIndex = INDEX_NONE;
+			int32 NextControllerIndex = INDEX_NONE;
+				
+			// We may possibly have invalid or missing categories
+			while (CategoryIndex < CategoryItems.Num())
+			{
+				const FRCControllerPanelListItem CategoryItem = {ERCControllerPanelListItemType::Category, CategoryIndex};
+
+				if (CategoryItems[CategoryIndex].IsValid() && !InItemsToMove.Contains(CategoryItem))
+				{
+					if (TOptional<FRCVirtualPropertyCategory> Category = CategoryItems[CategoryIndex]->GetCategory())
+					{
+						NextCategoryIndex = Category->DisplayIndex;
+						break;
+					}
+				}
+
+				++CategoryIndex;
+			}
+
+			// We may possibly have invalid or missing controllers
+			while (ControllerIndex < ControllerItems.Num())
+			{
+				const FRCControllerPanelListItem ControllerItem = {ERCControllerPanelListItemType::Controller, ControllerIndex};
+
+				if (ControllerItems[ControllerIndex].IsValid() && !InItemsToMove.Contains(ControllerItem))
+				{
+					if (URCVirtualPropertyBase* Controller = ControllerItems[ControllerIndex]->GetVirtualProperty())
+					{
+						NextControllerIndex = Controller->DisplayIndex;
+						break;
+					}
+				}
+
+				++ControllerIndex;
+			}
+
+			if (NextCategoryIndex != INDEX_NONE && NextControllerIndex != INDEX_NONE)
+			{
+				if (NextCategoryIndex <= NextControllerIndex)
+				{
+					return {ERCControllerPanelListItemType::Category, CategoryIndex++};
+				}
+				else
+				{
+					return {ERCControllerPanelListItemType::Controller, ControllerIndex++};
+				}
+			}
+			else if (NextCategoryIndex != INDEX_NONE)
+			{
+				return {ERCControllerPanelListItemType::Category, CategoryIndex++};
+			}
+			else if (NextControllerIndex != INDEX_NONE)
+			{
+				return {ERCControllerPanelListItemType::Controller, ControllerIndex++};
+			}
+			else
+			{
+				return UE::RemoteControl::UI::Private::ControllerPanelListItem::None;
+			}
+		};
+
+	using namespace UE::RemoteControl::UI::Private;
+
+	FRCControllerPanelListItem NextItem = GetNextItem();
+
+	while (NextItem.Type != ERCControllerPanelListItemType::None)
 	{
-		if (InItemsToMove.Contains(*Iter))
+		const bool bIsDroppedOnItem = (NextItem == InDroppedOnItem);
+		FName ItemCategory = NAME_None;
+
+		if (NextItem.Type == ERCControllerPanelListItemType::Controller)
+		{
+			ItemCategory = *ControllerItems[NextItem.Index]->GetVirtualProperty()->GetMetadataValue(FRCControllerPropertyInfo::CategoryName);
+		}
+
+		// Make sure the category exists
+		TArray<FRCControllerPanelListItem>* CategoryArray = Tree.Find(ItemCategory);
+
+		if (!CategoryArray)
 		{
-			if (Iter.GetIndex() < InTargetIndex)
+			ItemCategory = NAME_None;
+			CategoryArray = Tree.Find(ItemCategory);
+		}
+
+		if (bIsDroppedOnItem && InDropZone == EItemDropZone::AboveItem)
+		{
+			AddDroppedItems(ItemCategory);
+		}
+
+		if (!InItemsToMove.Contains(NextItem))
+		{
+			CategoryArray->Add(NextItem);
+		}
+
+		if (bIsDroppedOnItem)
+		{
+			if (InDropZone == EItemDropZone::OntoItem)
+			{
+				if (NextItem.Type == ERCControllerPanelListItemType::Category)
+				{
+					AddDroppedItems(*CategoryItems[NextItem.Index]->GetCategory()->Id.ToString());
+				}
+				else
+				{
+					AddDroppedItems(ItemCategory);
+				}
+			}
+			else if (InDropZone == EItemDropZone::BelowItem)
 			{
-				--InTargetIndex;
+				AddDroppedItems(ItemCategory);
 			}
-			Iter.RemoveCurrent();
 		}
+
+		NextItem = GetNextItem();
 	}
 
-	// Update UI list. Remove all the items to move to then insert them at the desired index
-	ControllerItems.RemoveAll([&InItemsToMove](const TSharedPtr<FRCControllerModel>& InController)
+	auto GetDisplayIndex = [this](const FRCControllerPanelListItem& InItem) -> int32
 		{
-			return InItemsToMove.Contains(InController);
-		});
+			switch (InItem.Type)
+			{
+				default:
+					break;
+
+				case ERCControllerPanelListItemType::Category:
+					if (CategoryItems.IsValidIndex(InItem.Index) && CategoryItems[InItem.Index].IsValid())
+					{
+
+					}
+					break;
+
+				case ERCControllerPanelListItemType::Controller:
+					if (ControllerItems.IsValidIndex(InItem.Index) && ControllerItems[InItem.Index].IsValid())
+					{
+						if (URCVirtualPropertyBase* Controller = ControllerItems[InItem.Index]->GetVirtualProperty())
+						{
+							return Controller->DisplayIndex;
+						}
+					}
+					break;
+			}
 
-	ControllerItems.Insert(InItemsToMove.GetData(), InItemsToMove.Num(), InTargetIndex);
+			return INDEX_NONE;
+		};
 
-	// Update display indices
-	for (int32 Index = 0; Index < ControllerItems.Num(); Index++)
+	auto SetDisplayIndex = [this](const FRCControllerPanelListItem& InItem, int32 InDisplayIndex)
+		{
+			switch (InItem.Type)
+			{
+				default:
+					return;
+
+				case ERCControllerPanelListItemType::Category:
+					if (CategoryItems.IsValidIndex(InItem.Index) && CategoryItems[InItem.Index].IsValid())
+					{
+						if (TOptional<FRCVirtualPropertyCategory> Category = CategoryItems[InItem.Index]->GetCategory())
+						{
+							// Already has Modify() called
+							CategoryItems[InItem.Index]->SetDisplayIndex(InDisplayIndex);
+						}
+					}
+					return;
+
+				case ERCControllerPanelListItemType::Controller:
+					if (ControllerItems.IsValidIndex(InItem.Index) && ControllerItems[InItem.Index].IsValid())
+					{
+						if (URCVirtualPropertyBase* Controller = ControllerItems[InItem.Index]->GetVirtualProperty())
+						{
+							if (GUndo)
+							{
+								Controller->Modify();
+							}
+
+							Controller->DisplayIndex = InDisplayIndex;
+						}
+					}
+					return;
+			}
+		};
+
+	int32 DisplayIndex = 0;
+	RootItems = Tree[NAME_None];
+	
+	for (const FRCControllerPanelListItem& RootItem : RootItems)
 	{
-		if (ControllerItems[Index])
+		SetDisplayIndex(RootItem, DisplayIndex);
+		++DisplayIndex;
+
+		if (RootItem.Type == ERCControllerPanelListItemType::Category)
 		{
-			if (URCVirtualPropertyBase* Controller = ControllerItems[Index]->GetVirtualProperty())
+			if (TOptional<FRCVirtualPropertyCategory> Category = CategoryItems[RootItem.Index]->GetCategory())
 			{
-				Controller->Modify();
-				Controller->SetDisplayIndex(Index);
+				const FName CategoryName = *Category->Id.ToString();
+
+				for (const FRCControllerPanelListItem& CategoryItem : Tree[CategoryName])
+				{
+					SetDisplayIndex(CategoryItem, DisplayIndex);
+					++DisplayIndex;
+				}
 			}
 		}
 	}
 
-	ListView->RequestListRefresh();
+	Algo::Sort(
+		CategoryItems,
+		[](const TSharedPtr<FRCCategoryModel>& A, const TSharedPtr<FRCCategoryModel>& B)
+		{
+			int32 DisplayIndexA = INDEX_NONE;
+			int32 DisplayIndexB = INDEX_NONE;
+
+			if (A.IsValid())
+			{
+				if (TOptional<FRCVirtualPropertyCategory> Category = A->GetCategory())
+				{
+					DisplayIndexA = Category->DisplayIndex;
+				}
+			}
+
+			if (B.IsValid())
+			{
+				if (TOptional<FRCVirtualPropertyCategory> Category = B->GetCategory())
+				{
+					DisplayIndexB = Category->DisplayIndex;
+				}
+			}
+
+			return DisplayIndexA < DisplayIndexB;
+		}
+	);
+
+	Algo::Sort(
+		ControllerItems,
+		[](const TSharedPtr<FRCControllerModel>& A, const TSharedPtr<FRCControllerModel>& B)
+		{
+			int32 DisplayIndexA = INDEX_NONE;
+			int32 DisplayIndexB = INDEX_NONE;
+
+			if (A.IsValid())
+			{
+				if (URCVirtualPropertyBase* Controller = A->GetVirtualProperty())
+				{
+					DisplayIndexA = Controller->DisplayIndex;
+				}
+			}
+
+			if (B.IsValid())
+			{
+				if (URCVirtualPropertyBase* Controller = B->GetVirtualProperty())
+				{
+					DisplayIndexB = Controller->DisplayIndex;
+				}
+			}
+
+			return DisplayIndexA < DisplayIndexB;
+		}
+	);
+
+	Reset();
+	//TreeView->RequestListRefresh();
+
 	return true;
 }
 
@@ -660,7 +1370,7 @@ bool SRCControllerPanelList::IsEntitySupported(const FGuid ExposedEntityId)
 
 bool SRCControllerPanelList::OnAllowDrop(TSharedPtr<FDragDropOperation> DragDropOperation)
 {
-	if (IsListViewHovered())
+	if (IsTreeViewHovered())
 	{
 		bIsAnyControllerItemEligibleForDragDrop = false;
 	}
@@ -688,7 +1398,7 @@ bool SRCControllerPanelList::OnAllowDrop(TSharedPtr<FDragDropOperation> DragDrop
 	return false;
 }
 
-FReply SRCControllerPanelList::OnControllerListViewDragDrop(TSharedPtr<FDragDropOperation> InDragDropOperation)
+FReply SRCControllerPanelList::OnControllerTreeViewDragDrop(TSharedPtr<FDragDropOperation> InDragDropOperation)
 {
 	URemoteControlPreset* const Preset = GetPreset();
 	if (!Preset)
@@ -815,9 +1525,9 @@ void SRCControllerPanelList::CreateBindBehaviourAndAssignTo(URCController* Contr
 	SelectController(Controller);
 }
 
-bool SRCControllerPanelList::IsListViewHovered()
+bool SRCControllerPanelList::IsTreeViewHovered()
 {
-	return ListView->IsDirectlyHovered();
+	return TreeView->IsDirectlyHovered();
 }
 
 void SRCControllerPanelList::ShowValueTypeHeaderColumn(bool bInShowColumn)
diff --git a/Engine/Plugins/VirtualProduction/RemoteControl/Source/RemoteControlUI/Private/UI/Controller/SRCControllerPanelList.h b/Engine/Plugins/VirtualProduction/RemoteControl/Source/RemoteControlUI/Private/UI/Controller/SRCControllerPanelList.h
--- a/Engine/Plugins/VirtualProduction/RemoteControl/Source/RemoteControlUI/Private/UI/Controller/SRCControllerPanelList.h
+++ b/Engine/Plugins/VirtualProduction/RemoteControl/Source/RemoteControlUI/Private/UI/Controller/SRCControllerPanelList.h
@@ -1,21 +1,27 @@
 // Copyright Epic Games, Inc. All Rights Reserved.
 
 #pragma once
+
 #include "DragAndDrop/DecoratedDragDropOp.h"
 #include "Misc/NotifyHook.h"
+#include "Misc/Optional.h"
 #include "RCMultiController.h"
+#include "RCVirtualPropertyContainer.h"
 #include "UI/BaseLogicUI/SRCLogicPanelListBase.h"
+#include "UI/Controller/RCControllerPanelListItem.h"
 #include "UI/RemoteControlPanelStyle.h"
 #include "Widgets/Layout/SBorder.h"
 
 class FDragDropOperation;
 class FEditPropertyChain;
+class FRCCategoryModel;
 class FRCControllerModel;
 class FRCLogicModeBase;
 class IPropertyRowGenerator;
 class ITableBase;
 class ITableRow;
 class SDropTarget;
+class SRCControllerPanelControllerRow;
 class SRCControllerPanel;
 class SRemoteControlPanel;
 class STableViewBase;
@@ -23,31 +29,29 @@ class URCController;
 class URemoteControlPreset;
 enum class EItemDropZone;
 struct FRCPanelStyle;
-template <typename ItemType> class SListView;
+template <typename ItemType> class STreeView;
 
 /*
-* ~ FRCControllerDragDrop ~
-* Facilitates drag-drop operation for Controller row drag handles
+* ~ FRCControllerPanelDragDrop ~
+* Facilitates drag-drop operation
 */
-class FRCControllerDragDrop final : public FDecoratedDragDropOp
+class FRCControllerPanelDragDrop final : public FDecoratedDragDropOp
 {
 public:
 	DRAG_DROP_OPERATOR_TYPE(FRCControllerDragDropOp, FDecoratedDragDropOp)
 
-	explicit FRCControllerDragDrop(const TArray<TSharedPtr<FRCControllerModel>>& InControllers)
-		: ControllersWeak(InControllers)
+	explicit FRCControllerPanelDragDrop(const TArray<FRCControllerPanelListItem>& InItems)
+		: Items(InItems)
 	{
 	}
 
-	TArray<TSharedPtr<FRCControllerModel>> ResolveControllers() const;
-
-	virtual void OnDrop(bool bDropWasHandled, const FPointerEvent& MouseEvent) override
+	TArray<FRCControllerPanelListItem> GetItems() const
 	{
-		FDecoratedDragDropOp::OnDrop(bDropWasHandled, MouseEvent);
+		return Items;
 	}
 
 private:
-	TArray<TWeakPtr<FRCControllerModel>> ControllersWeak;
+	TArray<FRCControllerPanelListItem> Items;
 };
 
 /*
@@ -83,6 +87,9 @@ public:
 	/** Deletes currently selected items from the list view */
 	virtual void DeleteSelectedPanelItems() override;
 
+	/** Returns the list items selected by the user (if any). */
+	TArray<FRCControllerPanelListItem> GetSelectedItems() const;
+
 	/** Returns the controllers selected by the user (if any). */
 	TArray<TSharedPtr<FRCControllerModel>> GetSelectedControllers() const;
 
@@ -97,26 +104,38 @@ public:
 
 	int32 NumControllerItems() const
 	{
-		return ControllerItems.Num();
+		return ControllerItems.Num() + CategoryItems.Num();
 	}
 
+	/** Finds a category by list item */
+	TSharedPtr<FRCCategoryModel> FindCategoryItemByListItem(const FRCControllerPanelListItem& InItem) const;
+
+	/** Finds a Controller UI model by list item */
+	TSharedPtr<FRCControllerModel> FindControllerItemByListItem(const FRCControllerPanelListItem& InItem) const;
+
 	/** Finds a Controller UI model by unique Id */
 	TSharedPtr<FRCControllerModel> FindControllerItemById(const FGuid& InId) const;
 
+	/** Finds a Controller UI model by unique Id */
+	TSharedPtr<FRCControllerModel> FindControllerItemByPropertyName(const FName& InPropertyName) const;
+
 	/** Finds the Controller UI models by unique Id */
 	TArray<TSharedPtr<FRCControllerModel>> FindControllerItemsById(TConstArrayView<FGuid> InIds);
 
 	/** Finds the Controller UI models by their controllers */
 	TArray<TSharedPtr<FRCControllerModel>> FindControllerItemsByObject(TConstArrayView<URCController*> InControllers);
 
-	/** Finds the index of the controller item */
-	int32 GetDropIndex(const TSharedPtr<FRCControllerModel>& InItem, EItemDropZone InDropZone) const;
+	/** Finds the item by controller UI model */
+	FRCControllerPanelListItem FindListItemByControllerItem(const TSharedRef<FRCControllerModel>& InModel) const;
+
+	/** Finds the item by category id */
+	FRCControllerPanelListItem FindListItemByCategoryItem(const TSharedPtr<FRCCategoryModel>& InModel) const;
 
 	/**
 	 * Called to moves the given controller items to the given target index, pushing all other rows below
 	 * returns true if the operation was successful
 	 */
-	bool ReorderControllerItems(TConstArrayView<TSharedPtr<FRCControllerModel>> InItemsToMove, int32 InTargetIndex);
+	bool ReorderItems(TConstArrayView<FRCControllerPanelListItem> InItemsToMove, FRCControllerPanelListItem InDroppedOnItem, EItemDropZone InDropZone);
 
 	/** Requests the panel to refresh its contents from the latest list of Controllers */
 	void RequestRefresh()
@@ -128,7 +147,7 @@ public:
 	bool OnAllowDrop(TSharedPtr<FDragDropOperation> DragDropOperation);
 
 	/** Drag-Drop action delegate for the Controllers Panel List */
-	FReply OnControllerListViewDragDrop(TSharedPtr<FDragDropOperation> InDragDropOperation);
+	FReply OnControllerTreeViewDragDrop(TSharedPtr<FDragDropOperation> InDragDropOperation);
 
 	/** Fetches the Remote Control preset associated with the parent panel */
 	virtual URemoteControlPreset* GetPreset() override;
@@ -137,7 +156,7 @@ public:
 	void CreateBindBehaviourAndAssignTo(URCController* Controller, TSharedRef<const FRemoteControlProperty> InRemoteControlProperty, const bool bExecuteBind);
 
 	/** Whether the user's cursor is directly hovered over the List View */
-	bool IsListViewHovered();
+	bool IsTreeViewHovered();
 
 	/** Enable/disable MultiController Mode */
 	void SetMultiControllerMode(bool bIsUniqueModeOn);
@@ -154,10 +173,16 @@ public:
 	bool bIsAnyControllerItemEligibleForDragDrop = false;
 
 	/** OnGenerateRow delegate for the Actions List View */
-	TSharedRef<ITableRow> OnGenerateWidgetForList( TSharedPtr<FRCControllerModel> InItem, const TSharedRef<STableViewBase>& OwnerTable );
+	TSharedRef<ITableRow> OnGenerateWidgetForTree(FRCControllerPanelListItem InItem, const TSharedRef<STableViewBase>& OwnerTable);
+
+	/** Returns the row for a given item. */
+	bool IsItemCategoryExpanded(FRCControllerPanelListItem InItem) const;
+
+	/** Gets the children for categories. */
+	void OnGetChildren(FRCControllerPanelListItem InItem, TArray<FRCControllerPanelListItem>& OutChildren);
 	
 	/** OnSelectionChanged delegate for Actions List View */
-	void OnTreeSelectionChanged(TSharedPtr<FRCControllerModel> InItem , ESelectInfo::Type);
+	void OnTreeSelectionChanged(FRCControllerPanelListItem InItem , ESelectInfo::Type);
 
 	/** Selects the Controller UI item corresponding to a given Controller UObject */
 	void SelectController(URCController* InController);
@@ -192,7 +217,7 @@ public:
 	/** Creates a new Controller for the given Remote Control Property and also binds to it. Returns true if operation succeeded, false if nothing happened */
 	bool CreateAutoBindForProperty(TSharedPtr<const FRemoteControlProperty> InRemoteControlProperty);
 
-	/** The row generator used to represent each Controller as a row, when used with SListView */
+	/** The row generator used to represent each Controller as a row, when used with STreeView */
 	TSharedPtr<IPropertyRowGenerator> PropertyRowGenerator;
 
 	/** The currently selected Controller item (UI model) */
@@ -201,25 +226,37 @@ public:
 	/** The parent Controller Panel widget */
 	TWeakPtr<SRCControllerPanel> ControllerPanelWeakPtr;
 
+	/** List of Category (UI model) active in this widget */
+	TArray<TSharedPtr<FRCCategoryModel>> CategoryItems;
+
 	/** List of Controllers (UI model) active in this widget */
 	TArray<TSharedPtr<FRCControllerModel>> ControllerItems;
 
-	/** List View widget for representing our Controllers List */
-	TSharedPtr<SListView<TSharedPtr<FRCControllerModel>>> ListView;
+	/** List of Controllers (UI model) that are in the root category */
+	TArray<FRCControllerPanelListItem> RootItems;
+
+	/** Tree View widget for representing our Controllers List */
+	TSharedPtr<STreeView<FRCControllerPanelListItem>> TreeView;
 	
 	/** Refreshes the list from the latest state of the model */
 	virtual void Reset() override;
 
+	/** Updates the root item set from ControllerItems */
+	void UpdateRootItems();
+
+	/** Called to update metadata when things are removed. */
+	void OnItemRemoved(const FRCControllerPanelListItem& InRemovedItem);
+
 	/** Handles broadcasting of a successful remove item operation. */
 	virtual void BroadcastOnItemRemoved() override;
 
 	/** Removes the given Controller UI model item from the list of UI models */
 	virtual int32 RemoveModel(const TSharedPtr<FRCLogicModeBase> InModel) override;
 
-	/** Set the visibility of the List header the Value Type Column used for MultiControllers */
+	/** Set the visibility of the Tree header the Value Type Column used for MultiControllers */
 	void ShowValueTypeHeaderColumn(bool bInShowColumn);
 
-	/** Set the visibility of the List header the Field Id Column */
+	/** Set the visibility of the Tree header the Field Id Column */
 	void ShowFieldIdHeaderColumn(bool bInShowColumn);
 
 	/** Add a custom column */
diff --git a/Engine/Plugins/VirtualProduction/RemoteControl/Source/RemoteControlUI/Private/UI/RCUIHelpers.cpp b/Engine/Plugins/VirtualProduction/RemoteControl/Source/RemoteControlUI/Private/UI/RCUIHelpers.cpp
--- a/Engine/Plugins/VirtualProduction/RemoteControl/Source/RemoteControlUI/Private/UI/RCUIHelpers.cpp
+++ b/Engine/Plugins/VirtualProduction/RemoteControl/Source/RemoteControlUI/Private/UI/RCUIHelpers.cpp
@@ -203,21 +203,23 @@ TSharedRef<SWidget> UE::RCUIHelpers::GetTypeColorWidget(const FProperty* InPrope
 	const FLinearColor TypeColor = UE::RCUIHelpers::GetFieldClassTypeColor(InProperty);
 	const FText TooltipText = FText::FromName(UE::RCUIHelpers::GetFieldClassDisplayName(InProperty));
 
-	TSharedRef<SWidget> TypeColorWidget =
-			SNew(SBox)
-			.HeightOverride(5.f)
-			.HAlign(HAlign_Left)
-			.Padding(FMargin(5.f, 0.f))
-			.ToolTipText(TooltipText)
-			[
-				SNew(SBorder)
+	return GetTypeColorWidget(TypeColor, TooltipText);
+}
+
+TSharedRef<SWidget> UE::RCUIHelpers::GetTypeColorWidget(const FLinearColor& InTypeColor, const FText& InTypeDescription)
+{
+	return SNew(SBox)
+		.HeightOverride(5.f)
+		.HAlign(HAlign_Left)
+		.Padding(FMargin(5.f, 0.f))
+		.ToolTipText(InTypeDescription)
+		[
+			SNew(SBorder)
 				.Visibility(EVisibility::Visible)
 				.BorderImage(FAppStyle::Get().GetBrush("NumericEntrySpinBox.Decorator"))
-				.BorderBackgroundColor(TypeColor)
+				.BorderBackgroundColor(InTypeColor)
 				.Padding(FMargin(5.0f, 0.0f, 0.0f, 0.f))
-			];
-
-	return TypeColorWidget;
+		];
 }
 
 UObject* UE::RCUIHelpers::GetEntityOwner(const TSharedPtr<const FRemoteControlEntity>& InEntity)
diff --git a/Engine/Plugins/VirtualProduction/RemoteControl/Source/RemoteControlUI/Private/UI/RCUIHelpers.h b/Engine/Plugins/VirtualProduction/RemoteControl/Source/RemoteControlUI/Private/UI/RCUIHelpers.h
--- a/Engine/Plugins/VirtualProduction/RemoteControl/Source/RemoteControlUI/Private/UI/RCUIHelpers.h
+++ b/Engine/Plugins/VirtualProduction/RemoteControl/Source/RemoteControlUI/Private/UI/RCUIHelpers.h
@@ -63,6 +63,15 @@ namespace UE::RCUIHelpers
 	*/
 	TSharedRef<SWidget> GetTypeColorWidget(const FProperty* InProperty);
 
+	/**
+	* GetTypeColorWidget
+	*
+	* Generates a Type Color widget which provides a color coding guide for the user given a desired color
+	* Used to provide a consistent look for the type color widget across various Remote Control Logic panels
+	* This is the direct version which takes a color and description, instead of extracting them from a property.
+	*/
+	TSharedRef<SWidget> GetTypeColorWidget(const FLinearColor& InTypeColor, const FText& InTypeDescription);
+
 	/** Retrieves the entity owner for a given entity. If the entity is bound to an actor's subobject, the actor will be returned. */
 	UObject* GetEntityOwner(const TSharedPtr<const FRemoteControlEntity>& InEntity);
 
