Remote Control: Controller category fixes
- Categories can now be deleted.
- Things are now always added to the end of the bottom of the controller list.

diff --git a/Engine/Plugins/VirtualProduction/RemoteControl/Source/RemoteControl/Private/RemoteControlPreset.cpp b/Engine/Plugins/VirtualProduction/RemoteControl/Source/RemoteControl/Private/RemoteControlPreset.cpp
--- a/Engine/Plugins/VirtualProduction/RemoteControl/Source/RemoteControl/Private/RemoteControlPreset.cpp
+++ b/Engine/Plugins/VirtualProduction/RemoteControl/Source/RemoteControl/Private/RemoteControlPreset.cpp
@@ -2870,4 +2870,24 @@ int32 URemoteControlPreset::GetEntityBoundCount(const FGuid& InEntityGuid) const
 	return 0;
 }
 
+int32 URemoteControlPreset::GetNextControllerDisplayIndex() const
+{
+	int32 HighestDisplayIndex = INDEX_NONE;
+
+	if (ControllerContainer)
+	{
+		for (const FRCVirtualPropertyCategory& Category : ControllerContainer->Categories)
+		{
+			HighestDisplayIndex = FMath::Max(HighestDisplayIndex, Category.DisplayIndex);
+		}
+
+		for (URCVirtualPropertyBase* Controller : ControllerContainer->VirtualProperties)
+		{
+			HighestDisplayIndex = FMath::Max(HighestDisplayIndex, Controller->DisplayIndex);
+		}
+	}
+
+	return HighestDisplayIndex + 1;
+}
+
 #undef LOCTEXT_NAMESPACE /* RemoteControlPreset */ 
diff --git a/Engine/Plugins/VirtualProduction/RemoteControl/Source/RemoteControl/Public/RCVirtualPropertyContainer.h b/Engine/Plugins/VirtualProduction/RemoteControl/Source/RemoteControl/Public/RCVirtualPropertyContainer.h
--- a/Engine/Plugins/VirtualProduction/RemoteControl/Source/RemoteControl/Public/RCVirtualPropertyContainer.h
+++ b/Engine/Plugins/VirtualProduction/RemoteControl/Source/RemoteControl/Public/RCVirtualPropertyContainer.h
@@ -223,10 +223,8 @@ public:
 	/** Delegate that notifies changes to the virtual property container*/
 	FOnVirtualPropertyContainerModified OnVirtualPropertyContainerModifiedDelegate;
 
-#if WITH_EDITORONLY_DATA
 	UPROPERTY()
 	TArray<FRCVirtualPropertyCategory> Categories;
-#endif
 
 private:
 	/** Map of Controller Name to GUID. */
diff --git a/Engine/Plugins/VirtualProduction/RemoteControl/Source/RemoteControl/Public/RemoteControlPreset.h b/Engine/Plugins/VirtualProduction/RemoteControl/Source/RemoteControl/Public/RemoteControlPreset.h
--- a/Engine/Plugins/VirtualProduction/RemoteControl/Source/RemoteControl/Public/RemoteControlPreset.h
+++ b/Engine/Plugins/VirtualProduction/RemoteControl/Source/RemoteControl/Public/RemoteControlPreset.h
@@ -928,6 +928,9 @@ public:
 	/** Returns the number of behaviors an entity is bound to. */
 	int32 GetEntityBoundCount(const FGuid& InEntityGuid) const;
 
+	/** The display index to put it at the bottom of the list. */
+	int32 GetNextControllerDisplayIndex() const;
+
 private:
 
 	/** Find a binding that has the same boundobjectmap but that currently points to the object passed as argument. */
diff --git a/Engine/Plugins/VirtualProduction/RemoteControl/Source/RemoteControlLogic/Private/Signature/RCSignatureControllerAction.cpp b/Engine/Plugins/VirtualProduction/RemoteControl/Source/RemoteControlLogic/Private/Signature/RCSignatureControllerAction.cpp
--- a/Engine/Plugins/VirtualProduction/RemoteControl/Source/RemoteControlLogic/Private/Signature/RCSignatureControllerAction.cpp
+++ b/Engine/Plugins/VirtualProduction/RemoteControl/Source/RemoteControlLogic/Private/Signature/RCSignatureControllerAction.cpp
@@ -70,7 +70,7 @@ URCController* FRCSignatureControllerAction::FindOrAddController(const FRCSignat
 		if (MakeControllerDesc(InContext, Desc))
 		{
 			Controller = CastChecked<URCController>(InContext.Preset->AddController(URCController::StaticClass(), Desc.ValueType, const_cast<UObject*>(Desc.ValueTypeObject.Get()), ControllerNameToUse));
-			Controller->DisplayIndex = InContext.Preset->GetNumControllers() - 1;
+			Controller->DisplayIndex = InContext.Preset->GetNextControllerDisplayIndex();
 
 #if WITH_EDITOR
 			for (const FPropertyBagPropertyDescMetaData& MetaData : Desc.MetaData)
diff --git a/Engine/Plugins/VirtualProduction/RemoteControl/Source/RemoteControlUI/Private/UI/Action/RCActionModel.cpp b/Engine/Plugins/VirtualProduction/RemoteControl/Source/RemoteControlUI/Private/UI/Action/RCActionModel.cpp
--- a/Engine/Plugins/VirtualProduction/RemoteControl/Source/RemoteControlUI/Private/UI/Action/RCActionModel.cpp
+++ b/Engine/Plugins/VirtualProduction/RemoteControl/Source/RemoteControlUI/Private/UI/Action/RCActionModel.cpp
@@ -73,6 +73,7 @@ namespace UE::RCActionPanelList
 	};
 }
 
+const FLazyName FRCActionModel::ActionModelName = "ActionModel";
 
 FRCActionModel::FRCActionModel(URCAction* InAction, const TSharedPtr<class FRCBehaviourModel> InBehaviourItem, const TSharedPtr<SRemoteControlPanel> InRemoteControlPanel)
 	: FRCLogicModeBase(InRemoteControlPanel),
@@ -325,6 +326,32 @@ void FRCActionModel::OnSelectionExit()
 	}
 }
 
+int32 FRCActionModel::RemoveModel()
+{
+	URCAction* Action = GetAction();
+
+	if (!Action)
+	{
+		return 0;
+	}
+
+	TSharedPtr<FRCBehaviourModel> BehaviorModel = GetParentBehaviour();
+
+	if (!BehaviorModel.IsValid())
+	{
+		return 0;
+	}
+
+	URCBehaviour* Behavior = BehaviorModel->GetBehaviour();
+
+	if (!Behavior || !Behavior->ActionContainer)
+	{
+		return 0;
+	}
+
+	return Behavior->ActionContainer->RemoveAction(Action);
+}
+
 FRCPropertyIdActionType::FRCPropertyIdActionType(URCPropertyIdAction* InPropertyIdAction)
 	: PropertyIdActionWeakPtr(InPropertyIdAction)
 {
diff --git a/Engine/Plugins/VirtualProduction/RemoteControl/Source/RemoteControlUI/Private/UI/Action/RCActionModel.h b/Engine/Plugins/VirtualProduction/RemoteControl/Source/RemoteControlUI/Private/UI/Action/RCActionModel.h
--- a/Engine/Plugins/VirtualProduction/RemoteControl/Source/RemoteControlUI/Private/UI/Action/RCActionModel.h
+++ b/Engine/Plugins/VirtualProduction/RemoteControl/Source/RemoteControlUI/Private/UI/Action/RCActionModel.h
@@ -34,7 +34,14 @@ class URCFunctionAction;
 class FRCActionModel : public FRCLogicModeBase
 {
 public:
+	static const FLazyName ActionModelName;
+
 	FRCActionModel(URCAction* InAction, const TSharedPtr<class FRCBehaviourModel> InBehaviourItem, const TSharedPtr<SRemoteControlPanel> InRemoteControlPanel);
+
+	virtual FName GetModelType() const override
+	{
+		return ActionModelName;
+	}
 	
 	/** Fetch the Action (data model) associated with us*/
 	URCAction* GetAction() const;
@@ -45,6 +52,8 @@ public:
 	/** Widget representing the Value field */
 	virtual TSharedRef<SWidget> GetWidget() const override;
 
+	virtual int32 RemoveModel() override;
+
 	/** Widget representing Color Coding for the Action Type*/
 	virtual TSharedRef<SWidget> GetTypeColorTagWidget() const;
 
diff --git a/Engine/Plugins/VirtualProduction/RemoteControl/Source/RemoteControlUI/Private/UI/Action/SRCActionPanelList.h b/Engine/Plugins/VirtualProduction/RemoteControl/Source/RemoteControlUI/Private/UI/Action/SRCActionPanelList.h
--- a/Engine/Plugins/VirtualProduction/RemoteControl/Source/RemoteControlUI/Private/UI/Action/SRCActionPanelList.h
+++ b/Engine/Plugins/VirtualProduction/RemoteControl/Source/RemoteControlUI/Private/UI/Action/SRCActionPanelList.h
@@ -24,7 +24,6 @@
 #include "UI/RCUIHelpers.h"
 #include "UI/RemoteControlPanelStyle.h"
 #include "UI/SRCPanelExposedEntitiesGroup.h"
-#include "UI/SRCPanelExposedEntitiesList.h"
 #include "UI/SRCPanelExposedField.h"
 #include "UI/SRCPanelFieldGroup.h"
 #include "UI/SRemoteControlPanel.h"
@@ -374,29 +373,6 @@ private:
 		return nullptr;
 	}
 
-	/** Removes the given Action UI model item from the list of UI models */
-	virtual int32 RemoveModel(const TSharedPtr<FRCLogicModeBase> InModel) override
-	{
-		if (const TSharedPtr<FRCBehaviourModel> BehaviourModel = BehaviourItemWeakPtr.Pin())
-		{
-			if (const URCBehaviour* Behaviour = BehaviourModel->GetBehaviour())
-			{
-				if (const TSharedPtr<ActionType> SelectedAction = StaticCastSharedPtr<ActionType>(InModel))
-				{
-					// Remove Model from Data Container
-					FScopedTransaction Transaction(LOCTEXT("RemoveAction", "Remove Action"));
-					Behaviour->ActionContainer->Modify();
-
-					const int32 RemoveCount = Behaviour->ActionContainer->RemoveAction(SelectedAction->GetAction());
-
-					return RemoveCount;
-				}
-			}
-		}
-
-		return 0;
-	}
-
 	/*Drag and drop event for creating an Action from an exposed field */
 	FReply OnExposedFieldDrop(TSharedPtr<FDragDropOperation> DragDropOperation)
 	{
@@ -621,20 +597,6 @@ private:
 		SelectedActionItem = InItem;
 	}
 
-	void UpdateEntityUsage()
-	{
-		if (TSharedPtr<SRCActionPanel> ActionPanel = ActionPanelWeakPtr.Pin())
-		{
-			if (TSharedPtr<SRemoteControlPanel> RCPanel = ActionPanel->GetRemoteControlPanel())
-			{
-				if (TSharedPtr<SRCPanelExposedEntitiesList> EntitiesList = RCPanel->GetEntityList())
-				{
-					EntitiesList->UpdateEntityUsage();
-				}
-			}
-		}
-	}
-
 private:
 
 	/** The currently selected Action item */
diff --git a/Engine/Plugins/VirtualProduction/RemoteControl/Source/RemoteControlUI/Private/UI/BaseLogicUI/RCLogicModeBase.h b/Engine/Plugins/VirtualProduction/RemoteControl/Source/RemoteControlUI/Private/UI/BaseLogicUI/RCLogicModeBase.h
--- a/Engine/Plugins/VirtualProduction/RemoteControl/Source/RemoteControlUI/Private/UI/BaseLogicUI/RCLogicModeBase.h
+++ b/Engine/Plugins/VirtualProduction/RemoteControl/Source/RemoteControlUI/Private/UI/BaseLogicUI/RCLogicModeBase.h
@@ -23,6 +23,9 @@ public:
 
 	virtual ~FRCLogicModeBase() = default;
 
+	/** Returns a name to identify the type of logic model*/
+	virtual FName GetModelType() const = 0;
+
 	/** Represents the widget to be rendered for associated UI Model.
 	* Implemented by child classes to supply the row widget for various panel lists
 	*/
@@ -34,6 +37,9 @@ public:
 	/** Returns the Remote Control Preset object associated with us*/
 	TSharedPtr<SRemoteControlPanel> GetRemoteControlPanel() const;
 
+	/** Removes this model from the owning UI panel. Returns the number of items removed. */
+	virtual int32 RemoveModel() = 0;
+
 protected:
 	/** The parent Remote Control Panel widget*/
 	TWeakPtr<SRemoteControlPanel> PanelWeakPtr;
diff --git a/Engine/Plugins/VirtualProduction/RemoteControl/Source/RemoteControlUI/Private/UI/BaseLogicUI/SRCLogicPanelListBase.cpp b/Engine/Plugins/VirtualProduction/RemoteControl/Source/RemoteControlUI/Private/UI/BaseLogicUI/SRCLogicPanelListBase.cpp
--- a/Engine/Plugins/VirtualProduction/RemoteControl/Source/RemoteControlUI/Private/UI/BaseLogicUI/SRCLogicPanelListBase.cpp
+++ b/Engine/Plugins/VirtualProduction/RemoteControl/Source/RemoteControlUI/Private/UI/BaseLogicUI/SRCLogicPanelListBase.cpp
@@ -3,7 +3,9 @@
 #include "SRCLogicPanelListBase.h"
 
 #include "Commands/RemoteControlCommands.h"
+#include "ScopedTransaction.h"
 #include "SRCLogicPanelBase.h"
+#include "UI/SRCPanelExposedEntitiesList.h"
 #include "UI/SRemoteControlPanel.h"
 
 #define LOCTEXT_NAMESPACE "SRCLogicPanelListBase"
@@ -80,4 +82,37 @@ void SRCLogicPanelListBase::RequestDeleteAllItems()
 	}
 }
 
+int32 SRCLogicPanelListBase::RemoveModel(const TSharedPtr<FRCLogicModeBase> InModel)
+{
+	if (!InModel.IsValid())
+	{
+		return 0;
+	}
+
+	FScopedTransaction Transaction(LOCTEXT("RemoveModel", "Remove Model"));
+
+	const int32 RemovedCount = InModel->RemoveModel();
+
+	if (RemovedCount == 0)
+	{
+		Transaction.Cancel();
+		return 0;
+	}
+
+	UpdateEntityUsage();
+
+	return RemovedCount;
+}
+
+void SRCLogicPanelListBase::UpdateEntityUsage()
+{
+	if (TSharedPtr<SRemoteControlPanel> RCPanel = GetRemoteControlPanel())
+	{
+		if (TSharedPtr<SRCPanelExposedEntitiesList> EntitiesList = RCPanel->GetEntityList())
+		{
+			EntitiesList->UpdateEntityUsage();
+		}
+	}
+}
+
 #undef LOCTEXT_NAMESPACE
diff --git a/Engine/Plugins/VirtualProduction/RemoteControl/Source/RemoteControlUI/Private/UI/BaseLogicUI/SRCLogicPanelListBase.h b/Engine/Plugins/VirtualProduction/RemoteControl/Source/RemoteControlUI/Private/UI/BaseLogicUI/SRCLogicPanelListBase.h
--- a/Engine/Plugins/VirtualProduction/RemoteControl/Source/RemoteControlUI/Private/UI/BaseLogicUI/SRCLogicPanelListBase.h
+++ b/Engine/Plugins/VirtualProduction/RemoteControl/Source/RemoteControlUI/Private/UI/BaseLogicUI/SRCLogicPanelListBase.h
@@ -81,7 +81,7 @@ private:
 	virtual URemoteControlPreset* GetPreset() = 0;
 
 	/** Removes the given UI model item from the list of UI models for this panel list*/
-	virtual int32 RemoveModel(const TSharedPtr<FRCLogicModeBase> InModel) = 0;
+	virtual int32 RemoveModel(const TSharedPtr<FRCLogicModeBase> InModel);
 
 	/** Returns true when we can delete all items. */
 	bool CanDeleteAllItems() const;
@@ -141,4 +141,6 @@ protected:
 
 	/** The parent Logic Panel*/
 	TWeakPtr<SRCLogicPanelBase> LogicPanelWeakPtr;
+
+	virtual void UpdateEntityUsage();
 };
\ No newline at end of file
diff --git a/Engine/Plugins/VirtualProduction/RemoteControl/Source/RemoteControlUI/Private/UI/Behaviour/Builtin/Path/RCBehaviourSetAssetByPathModel.cpp b/Engine/Plugins/VirtualProduction/RemoteControl/Source/RemoteControlUI/Private/UI/Behaviour/Builtin/Path/RCBehaviourSetAssetByPathModel.cpp
--- a/Engine/Plugins/VirtualProduction/RemoteControl/Source/RemoteControlUI/Private/UI/Behaviour/Builtin/Path/RCBehaviourSetAssetByPathModel.cpp
+++ b/Engine/Plugins/VirtualProduction/RemoteControl/Source/RemoteControlUI/Private/UI/Behaviour/Builtin/Path/RCBehaviourSetAssetByPathModel.cpp
@@ -391,7 +391,7 @@ void FRCSetAssetByPathBehaviourModel::OnAssetPathFinishedChangingProperties(cons
 					URCVirtualPropertyInContainer* NewController = Preset->AddController(URCController::StaticClass(), EPropertyBagPropertyType::String);
 					if (NewController)
 					{
-						NewController->DisplayIndex = RCPanel->NumControllerItems();
+						NewController->DisplayIndex = Preset->GetNextControllerDisplayIndex();
 						FName NewControllerName = TEXT("DefaultInput");
 						if (!CurrentElement.Path.IsEmpty())
 						{
diff --git a/Engine/Plugins/VirtualProduction/RemoteControl/Source/RemoteControlUI/Private/UI/Behaviour/RCBehaviourModel.cpp b/Engine/Plugins/VirtualProduction/RemoteControl/Source/RemoteControlUI/Private/UI/Behaviour/RCBehaviourModel.cpp
--- a/Engine/Plugins/VirtualProduction/RemoteControl/Source/RemoteControlUI/Private/UI/Behaviour/RCBehaviourModel.cpp
+++ b/Engine/Plugins/VirtualProduction/RemoteControl/Source/RemoteControlUI/Private/UI/Behaviour/RCBehaviourModel.cpp
@@ -3,6 +3,7 @@
 #include "RCBehaviourModel.h"
 
 #include "Behaviour/RCBehaviourNode.h"
+#include "Controller/RCController.h"
 #include "Engine/Blueprint.h"
 #include "Engine/BlueprintGeneratedClass.h"
 #include "Modules/ModuleManager.h"
@@ -20,6 +21,8 @@
 
 #define LOCTEXT_NAMESPACE "FRCBehaviourModel"
 
+const FLazyName FRCBehaviourModel::BehaviorModelName = "BehaviorModel";
+
 FRCBehaviourModel::FRCBehaviourModel(URCBehaviour* InBehaviour
 	, const TSharedPtr<SRemoteControlPanel> InRemoteControlPanel /*= nullptr*/)
 	: FRCLogicModeBase(InRemoteControlPanel)
@@ -216,4 +219,28 @@ URCBehaviour* FRCBehaviourModel::GetBehaviour() const
 	return BehaviourWeakPtr.Get();
 }
 
+int32 FRCBehaviourModel::RemoveModel()
+{
+	URCBehaviour* Behavior = GetBehaviour();
+
+	if (!Behavior)
+	{
+		return 0;
+	}
+
+	URCController* Controller = Behavior->ControllerWeakPtr.Get();
+
+	if (!Controller)
+	{
+		return 0;
+	}
+
+	if (GUndo)
+	{
+		Controller->Modify();
+	}
+
+	return Controller->RemoveBehaviour(Behavior) ? 1 : 0;
+}
+
 #undef LOCTEXT_NAMESPACE
\ No newline at end of file
diff --git a/Engine/Plugins/VirtualProduction/RemoteControl/Source/RemoteControlUI/Private/UI/Behaviour/RCBehaviourModel.h b/Engine/Plugins/VirtualProduction/RemoteControl/Source/RemoteControlUI/Private/UI/Behaviour/RCBehaviourModel.h
--- a/Engine/Plugins/VirtualProduction/RemoteControl/Source/RemoteControlUI/Private/UI/Behaviour/RCBehaviourModel.h
+++ b/Engine/Plugins/VirtualProduction/RemoteControl/Source/RemoteControlUI/Private/UI/Behaviour/RCBehaviourModel.h
@@ -22,8 +22,17 @@ struct FRCPanelStyle;
 class FRCBehaviourModel : public FRCLogicModeBase
 {
 public:
+	static const FLazyName BehaviorModelName;
+
 	FRCBehaviourModel(URCBehaviour* InBehaviour, const TSharedPtr<SRemoteControlPanel> InRemoteControlPanel = nullptr);
 
+	virtual FName GetModelType() const override
+	{
+		return BehaviorModelName;
+	}
+
+	virtual int32 RemoveModel() override;
+
 	/** Add a Logic Action using as an identity action. */
 	virtual URCAction* AddAction();
 
diff --git a/Engine/Plugins/VirtualProduction/RemoteControl/Source/RemoteControlUI/Private/UI/Behaviour/SRCBehaviourPanelList.cpp b/Engine/Plugins/VirtualProduction/RemoteControl/Source/RemoteControlUI/Private/UI/Behaviour/SRCBehaviourPanelList.cpp
--- a/Engine/Plugins/VirtualProduction/RemoteControl/Source/RemoteControlUI/Private/UI/Behaviour/SRCBehaviourPanelList.cpp
+++ b/Engine/Plugins/VirtualProduction/RemoteControl/Source/RemoteControlUI/Private/UI/Behaviour/SRCBehaviourPanelList.cpp
@@ -286,8 +286,6 @@ void SRCBehaviourPanelList::AddBehaviourToList(URCBehaviour* InBehaviour)
 	}
 }
 
-
-
 void SRCBehaviourPanelList::Reset()
 {
 	const TArray<TSharedPtr<FRCBehaviourModel>> SelectedBehaviors = ListView->GetSelectedItems();
@@ -402,30 +400,6 @@ URemoteControlPreset* SRCBehaviourPanelList::GetPreset()
 	return nullptr;
 }
 
-int32 SRCBehaviourPanelList::RemoveModel(const TSharedPtr<FRCLogicModeBase> InModel)
-{
-	if (const TSharedPtr<FRCControllerModel> ControllerModel = ControllerItemWeakPtr.Pin())
-	{
-		if (URCController* Controller = Cast<URCController>(ControllerModel->GetVirtualProperty()))
-		{
-			if(const TSharedPtr<FRCBehaviourModel> SelectedBehaviour = StaticCastSharedPtr<FRCBehaviourModel>(InModel))
-			{
-				FScopedTransaction Transaction(LOCTEXT("RemoveBehaviour", "Remove Behaviour"));
-				Controller->Modify();
-
-				// Remove Model from Data Container
-				const int32 RemoveCount = Controller->RemoveBehaviour(SelectedBehaviour->GetBehaviour());
-
-				UpdateEntityUsage();
-
-				return RemoveCount;
-			}
-		}
-	}
-
-	return 0;
-}
-
 void SRCBehaviourPanelList::OnBehaviourListModified()
 {
 	Reset();
@@ -474,18 +448,4 @@ void SRCBehaviourPanelList::RequestRefresh()
 	ListView->RequestListRefresh();
 }
 
-void SRCBehaviourPanelList::UpdateEntityUsage()
-{
-	if (TSharedPtr<SRCBehaviourPanel> BehaviourPanel = BehaviourPanelWeakPtr.Pin())
-	{
-		if (TSharedPtr<SRemoteControlPanel> RCPanel = BehaviourPanel->GetRemoteControlPanel())
-		{
-			if (TSharedPtr<SRCPanelExposedEntitiesList> EntitiesList = RCPanel->GetEntityList())
-			{
-				EntitiesList->UpdateEntityUsage();
-			}
-		}
-	}
-}
-
 #undef LOCTEXT_NAMESPACE
\ No newline at end of file
diff --git a/Engine/Plugins/VirtualProduction/RemoteControl/Source/RemoteControlUI/Private/UI/Behaviour/SRCBehaviourPanelList.h b/Engine/Plugins/VirtualProduction/RemoteControl/Source/RemoteControlUI/Private/UI/Behaviour/SRCBehaviourPanelList.h
--- a/Engine/Plugins/VirtualProduction/RemoteControl/Source/RemoteControlUI/Private/UI/Behaviour/SRCBehaviourPanelList.h
+++ b/Engine/Plugins/VirtualProduction/RemoteControl/Source/RemoteControlUI/Private/UI/Behaviour/SRCBehaviourPanelList.h
@@ -121,13 +121,8 @@ private:
 	/** Fetches the Remote Control preset associated with the parent panel */
 	virtual URemoteControlPreset* GetPreset() override;
 
-	/** Removes the given Behaviour UI model item from the list of UI models */
-	virtual int32 RemoveModel(const TSharedPtr<FRCLogicModeBase> InModel) override;
-
 	void OnBehaviourListModified();
 
-	void UpdateEntityUsage();
-
 private:
 
 	/** The parent Behaviour Panel widget */
diff --git a/Engine/Plugins/VirtualProduction/RemoteControl/Source/RemoteControlUI/Private/UI/Controller/RCCategoryModel.cpp b/Engine/Plugins/VirtualProduction/RemoteControl/Source/RemoteControlUI/Private/UI/Controller/RCCategoryModel.cpp
--- a/Engine/Plugins/VirtualProduction/RemoteControl/Source/RemoteControlUI/Private/UI/Controller/RCCategoryModel.cpp
+++ b/Engine/Plugins/VirtualProduction/RemoteControl/Source/RemoteControlUI/Private/UI/Controller/RCCategoryModel.cpp
@@ -9,6 +9,8 @@
 
 #define LOCTEXT_NAMESPACE "FRCCategoryModel"
 
+const FLazyName FRCCategoryModel::CategoryModelName = "CategoryModel";
+
 FRCCategoryModel::FRCCategoryModel(const FGuid& InCateogryId, const TSharedPtr<SRemoteControlPanel> InRemoteControlPanel)
 	: FRCLogicModeBase(InRemoteControlPanel)
 	, Id(InCateogryId)
@@ -47,6 +49,46 @@ TSharedRef<SWidget> FRCCategoryModel::GetWidget() const
 		];
 }
 
+int32 FRCCategoryModel::RemoveModel()
+{
+	TOptional<FRCVirtualPropertyCategory> Category = GetCategory();
+
+	if (!Category.IsSet())
+	{
+		return 0;
+	}
+
+	URemoteControlPreset* Preset = GetPreset();
+
+	if (!Preset)
+	{
+		return 0;
+	}
+
+	URCVirtualPropertyContainerBase* ControllerContainer = Preset->GetControllerContainer();
+
+	if (!ControllerContainer)
+	{
+		return 0;
+	}
+
+	if (GUndo)
+	{
+		ControllerContainer->Modify();
+	}
+
+	for (int32 Index = 0; Index < ControllerContainer->Categories.Num(); ++Index)
+	{
+		if (ControllerContainer->Categories[Index].Id == Category->Id)
+		{
+			ControllerContainer->Categories.RemoveAt(Index);
+			return 1;
+		}
+	}
+
+	return 0;
+}
+
 TSharedRef<SWidget> FRCCategoryModel::GetNameWidget() const
 {
 	return SNew(SBox)
diff --git a/Engine/Plugins/VirtualProduction/RemoteControl/Source/RemoteControlUI/Private/UI/Controller/RCCategoryModel.h b/Engine/Plugins/VirtualProduction/RemoteControl/Source/RemoteControlUI/Private/UI/Controller/RCCategoryModel.h
--- a/Engine/Plugins/VirtualProduction/RemoteControl/Source/RemoteControlUI/Private/UI/Controller/RCCategoryModel.h
+++ b/Engine/Plugins/VirtualProduction/RemoteControl/Source/RemoteControlUI/Private/UI/Controller/RCCategoryModel.h
@@ -21,10 +21,17 @@ struct FRCVirtualPropertyCategory;
 class FRCCategoryModel : public FRCLogicModeBase
 {
 public:
+	static const FLazyName CategoryModelName;
+
 	DECLARE_MULTICAST_DELEGATE_OneParam(FOnValueChanged, TSharedPtr<FRCCategoryModel> /* InCategoryModel */);
 
 	FRCCategoryModel(const FGuid& InCateogryId, const TSharedPtr<SRemoteControlPanel> InRemoteControlPanel);
 
+	virtual FName GetModelType() const override
+	{
+		return CategoryModelName;
+	}
+
 	TOptional<FRCVirtualPropertyCategory> GetCategory() const;
 
 	/** Called right after object construction to initialize the model (and allow using SharedThis, etc) */
@@ -36,6 +43,8 @@ public:
 	 */
 	virtual TSharedRef<SWidget> GetWidget() const override;
 
+	virtual int32 RemoveModel() override;
+
 	/** The widget showing the Name of this Category */
 	TSharedRef<SWidget> GetNameWidget() const;
 
@@ -53,6 +62,7 @@ public:
 
 	/** Triggered when the Value this controller changes */
 	FOnValueChanged OnValueChanged;
+
 private:
 	/** Text commit event for Category Name text box */
 	void OnCategoryNameCommitted(const FText& InNewCategoryName, ETextCommit::Type InCommitInfo);
diff --git a/Engine/Plugins/VirtualProduction/RemoteControl/Source/RemoteControlUI/Private/UI/Controller/RCControllerModel.cpp b/Engine/Plugins/VirtualProduction/RemoteControl/Source/RemoteControlUI/Private/UI/Controller/RCControllerModel.cpp
--- a/Engine/Plugins/VirtualProduction/RemoteControl/Source/RemoteControlUI/Private/UI/Controller/RCControllerModel.cpp
+++ b/Engine/Plugins/VirtualProduction/RemoteControl/Source/RemoteControlUI/Private/UI/Controller/RCControllerModel.cpp
@@ -15,6 +15,8 @@
 
 #define LOCTEXT_NAMESPACE "FRCControllerModel"
 
+const FLazyName FRCControllerModel::ControllerModelName = "ControllerModel";
+
 FRCControllerModel::FRCControllerModel(URCVirtualPropertyBase* InVirtualProperty, const TSharedRef<IDetailTreeNode>& InTreeNode, const TSharedPtr<SRemoteControlPanel> InRemoteControlPanel)
 	: FRCLogicModeBase(InRemoteControlPanel)
 	, VirtualPropertyWeakPtr(InVirtualProperty)
@@ -366,4 +368,33 @@ FText FRCControllerModel::GetControllerDescription() const
 	return FText::GetEmpty();
 }
 
+int32 FRCControllerModel::RemoveModel()
+{
+	URCVirtualPropertyBase* Controller = GetVirtualProperty();
+
+	if (!Controller)
+	{
+		return 0;
+	}
+
+	URemoteControlPreset* Preset = GetPreset();
+
+	if (!Preset)
+	{
+		return 0;
+	}
+
+	if (GUndo)
+	{
+		Preset->Modify();
+	}
+
+	if (Preset->RemoveController(GetPropertyName()))
+	{
+		return 1;
+	}
+
+	return 0;
+}
+
 #undef LOCTEXT_NAMESPACE
diff --git a/Engine/Plugins/VirtualProduction/RemoteControl/Source/RemoteControlUI/Private/UI/Controller/RCControllerModel.h b/Engine/Plugins/VirtualProduction/RemoteControl/Source/RemoteControlUI/Private/UI/Controller/RCControllerModel.h
--- a/Engine/Plugins/VirtualProduction/RemoteControl/Source/RemoteControlUI/Private/UI/Controller/RCControllerModel.h
+++ b/Engine/Plugins/VirtualProduction/RemoteControl/Source/RemoteControlUI/Private/UI/Controller/RCControllerModel.h
@@ -30,11 +30,18 @@ class URCVirtualPropertyBase;
 class FRCControllerModel : public FRCLogicModeBase, public FSelfRegisteringEditorUndoClient
 {
 public:
+	static const FLazyName ControllerModelName;
+
 	DECLARE_MULTICAST_DELEGATE_TwoParams(FOnValueTypeChanged, URCVirtualPropertyBase* /* InController */, EPropertyBagPropertyType /* InValueType */);
 	DECLARE_MULTICAST_DELEGATE_OneParam(FOnValueChanged, TSharedPtr<FRCControllerModel> /* InControllerModel */);
 	
 	FRCControllerModel(URCVirtualPropertyBase* InVirtualProperty, const TSharedRef<IDetailTreeNode>& InTreeNode, const TSharedPtr<SRemoteControlPanel> InRemoteControlPanel);
 
+	virtual FName GetModelType() const override
+	{
+		return ControllerModelName;
+	}
+
 	/** Called right after object construction to initialize the model (and allow using SharedThis, etc) */
 	void Initialize();
 
@@ -44,6 +51,8 @@ public:
 	 */
 	virtual TSharedRef<SWidget> GetWidget() const override;
 
+	virtual int32 RemoveModel() override;
+
 	/** The widget showing the Name of this Controller */
 	TSharedRef<SWidget> GetNameWidget() const;
 
diff --git a/Engine/Plugins/VirtualProduction/RemoteControl/Source/RemoteControlUI/Private/UI/Controller/SRCControllerPanel.cpp b/Engine/Plugins/VirtualProduction/RemoteControl/Source/RemoteControlUI/Private/UI/Controller/SRCControllerPanel.cpp
--- a/Engine/Plugins/VirtualProduction/RemoteControl/Source/RemoteControlUI/Private/UI/Controller/SRCControllerPanel.cpp
+++ b/Engine/Plugins/VirtualProduction/RemoteControl/Source/RemoteControlUI/Private/UI/Controller/SRCControllerPanel.cpp
@@ -273,7 +273,7 @@ void SRCControllerPanel::DuplicateController(URCController* InController)
 	{
 		if (URCController* NewController = Cast<URCController>(Preset->DuplicateController(InController)))
 		{
-			NewController->SetDisplayIndex(ControllerPanelList->NumControllerItems());
+			NewController->SetDisplayIndex(Preset->GetNextControllerDisplayIndex());
 
 			ControllerPanelList->RequestRefresh();
 		}
@@ -760,10 +760,7 @@ void SRCControllerPanel::OnAddControllerClicked(const EPropertyBagPropertyType I
 		}
 	}
 
-	if (ControllerPanelList.IsValid())
-	{
-		NewVirtualProperty->DisplayIndex = ControllerPanelList->NumControllerItems();
-	}
+	NewVirtualProperty->DisplayIndex = Preset->GetNextControllerDisplayIndex();
 
 	// Refresh list
 	const TSharedPtr<SRemoteControlPanel> RemoteControlPanel = GetRemoteControlPanel();
@@ -795,10 +792,7 @@ void SRCControllerPanel::OnAddCategoryClicked() const
 	NewCategory.Id = FGuid::NewGuid();
 	NewCategory.Name = LOCTEXT("NewCategory", "New Category");
 	
-	if (ControllerPanelList.IsValid())
-	{
-		NewCategory.DisplayIndex = ControllerPanelList->NumControllerItems();
-	}
+	NewCategory.DisplayIndex = Preset->GetNextControllerDisplayIndex();
 
 	// Refresh list
 	const TSharedPtr<SRemoteControlPanel> RemoteControlPanel = GetRemoteControlPanel();
@@ -827,6 +821,11 @@ FReply SRCControllerPanel::OnClickEmptyButton()
 		Preset->Modify();
 
 		Preset->ResetControllers();
+
+		if (URCVirtualPropertyContainerBase* ControllerContainer = Preset->GetControllerContainer())
+		{
+			ControllerContainer->Categories.Empty();
+		}
 	}
 
 	if (const TSharedPtr<SRemoteControlPanel> RemoteControlPanel = GetRemoteControlPanel())
diff --git a/Engine/Plugins/VirtualProduction/RemoteControl/Source/RemoteControlUI/Private/UI/Controller/SRCControllerPanelControllerRow.cpp b/Engine/Plugins/VirtualProduction/RemoteControl/Source/RemoteControlUI/Private/UI/Controller/SRCControllerPanelControllerRow.cpp
--- a/Engine/Plugins/VirtualProduction/RemoteControl/Source/RemoteControlUI/Private/UI/Controller/SRCControllerPanelControllerRow.cpp
+++ b/Engine/Plugins/VirtualProduction/RemoteControl/Source/RemoteControlUI/Private/UI/Controller/SRCControllerPanelControllerRow.cpp
@@ -371,38 +371,12 @@ TSharedRef<SWidget> SRCControllerPanelControllerRow::GetTypeColorWidget()
 
 	if (const FProperty* Property = Controller->GetProperty())
 	{
-		return SNew(SBox)
-			.Padding(this, &SRCControllerPanelControllerRow::GetTypeColorPadding)
-			[
-				UE::RCUIHelpers::GetTypeColorWidget(Property)
-			];
+		return UE::RCUIHelpers::GetTypeColorWidget(Property);
 	}
 
 	return SNullWidget::NullWidget;
 }
 
-FMargin SRCControllerPanelControllerRow::GetTypeColorPadding() const
-{
-	TSharedPtr<FRCControllerModel> ControllerModel = GetControllerModel();
-
-	if (!ControllerModel.IsValid())
-	{
-		return FMargin(0.f);
-	}
-
-	using namespace UE::RemoteControl::UI::Private;
-
-	if (URCController* Controller = Cast<URCController>(ControllerModel->GetVirtualProperty()))
-	{
-		if (!Controller->GetMetadataValue(FRCControllerPropertyInfo::CategoryName).IsEmpty())
-		{ 
-			return FMargin(10.f, 0.f, 0.f, 0.f);
-		}
-	}
-
-	return FMargin(0.f);
-}
-
 FReply SRCControllerPanelControllerRow::OnDragDetected(const FGeometry& InMyGeometry, const FPointerEvent& InMouseEvent)
 {
 	if (FSuperRowType::OnDragDetected(InMyGeometry, InMouseEvent).IsEventHandled())
diff --git a/Engine/Plugins/VirtualProduction/RemoteControl/Source/RemoteControlUI/Private/UI/Controller/SRCControllerPanelList.cpp b/Engine/Plugins/VirtualProduction/RemoteControl/Source/RemoteControlUI/Private/UI/Controller/SRCControllerPanelList.cpp
--- a/Engine/Plugins/VirtualProduction/RemoteControl/Source/RemoteControlUI/Private/UI/Controller/SRCControllerPanelList.cpp
+++ b/Engine/Plugins/VirtualProduction/RemoteControl/Source/RemoteControlUI/Private/UI/Controller/SRCControllerPanelList.cpp
@@ -706,81 +706,6 @@ URemoteControlPreset* SRCControllerPanelList::GetPreset()
 	return nullptr;	
 }
 
-int32 SRCControllerPanelList::RemoveModel(const TSharedPtr<FRCLogicModeBase> InModel)
-{
-	if(ControllerPanelWeakPtr.IsValid())
-	{
-		if (URemoteControlPreset* Preset = ControllerPanelWeakPtr.Pin()->GetPreset())
-		{
-			if(const TSharedPtr<FRCControllerModel> SelectedController = StaticCastSharedPtr<FRCControllerModel>(InModel))
-			{
-				if(URCVirtualPropertyBase* ControllerToRemove = SelectedController->GetVirtualProperty())
-				{
-					const int32 DisplayIndexToRemove = ControllerToRemove->DisplayIndex;
-
-					FScopedTransaction Transaction(LOCTEXT("RemoveController", "Remove Controller"));
-					Preset->Modify();
-
-					// Remove Model from Data Container
-					const FString ControllerPropertyName = SelectedController->GetPropertyName().ToString();
-					const bool bHasPropertyName = !ControllerPropertyName.IsEmpty();
-
-					const bool bRemoved = Preset->RemoveController(SelectedController->GetPropertyName());
-
-					if (bRemoved)
-					{
-						int32 FoundIndex = INDEX_NONE;
-
-						// Shift all display indices higher than the deleted index down by 1
-						for (int32 ControllerIndex = 0; ControllerIndex < ControllerItems.Num(); ControllerIndex++)
-						{
-							if (ensure(ControllerItems[ControllerIndex]))
-							{
-								if (URCVirtualPropertyBase* Controller = ControllerItems[ControllerIndex]->GetVirtualProperty())
-								{
-									if (Controller->DisplayIndex == DisplayIndexToRemove)
-									{
-										FoundIndex = Controller->DisplayIndex;
-									}
-									else if (Controller->DisplayIndex > DisplayIndexToRemove)
-									{
-										Controller->Modify();
-										Controller->SetDisplayIndex(Controller->DisplayIndex - 1);
-										FoundIndex = ControllerIndex;
-									}
-
-									if (bHasPropertyName)
-									{
-										using namespace UE::RemoteControl::UI::Private;
-
-										if (Controller->GetMetadataValue(FRCControllerPropertyInfo::CategoryName) == ControllerPropertyName)
-										{
-											Controller->Modify();
-											Controller->RemoveMetadataValue(FRCControllerPropertyInfo::CategoryName);
-										}
-									}
-								}
-							}
-						}
-
-						if (FoundIndex != INDEX_NONE)
-						{
-							OnItemRemoved({
-								ERCControllerPanelListItemType::Controller,
-								FoundIndex
-							});
-						}
-
-						return 1; // Remove Count
-					}
-				}
-			}
-		}
-	}
-
-	return 0;
-}
-
 bool SRCControllerPanelList::IsListFocused() const
 {
 	return TreeView->HasAnyUserFocus().IsSet() || ContextMenuWidgetCached.IsValid();
@@ -791,6 +716,12 @@ void SRCControllerPanelList::DeleteSelectedPanelItems()
 	FScopedTransaction Transaction(LOCTEXT("DeleteSelectedItems", "Delete Selected Items"));
 	TArray<FRCControllerPanelListItem> SelectedItems = TreeView->GetSelectedItems();
 
+	TArray<TSharedPtr<FRCCategoryModel>> SelectedCategories;
+	SelectedCategories.Reserve(SelectedItems.Num());
+
+	TArray<int32> SelectedCategoryIndices;
+	SelectedCategoryIndices.Reserve(SelectedItems.Num());
+
 	TArray<TSharedPtr<FRCControllerModel>> SelectedControllers;
 	SelectedControllers.Reserve(SelectedItems.Num());
 
@@ -799,21 +730,42 @@ void SRCControllerPanelList::DeleteSelectedPanelItems()
 
 	for (const FRCControllerPanelListItem& Item : SelectedItems)
 	{
-		if (Item.Type != ERCControllerPanelListItemType::Controller)
+		switch (Item.Type)
 		{
-			continue;
+			case ERCControllerPanelListItemType::Category:
+				if (CategoryItems.IsValidIndex(Item.Index) && CategoryItems[Item.Index].IsValid())
+				{
+					SelectedCategories.Add(CategoryItems[Item.Index]);
+					SelectedCategoryIndices.Add(Item.Index);
+				}
+				break;
+
+			case ERCControllerPanelListItemType::Controller:
+				if (ControllerItems.IsValidIndex(Item.Index) && ControllerItems[Item.Index].IsValid())
+				{
+					SelectedControllers.Add(ControllerItems[Item.Index]);
+					SelectedControllerIndices.Add(Item.Index);
+				}
+				break;
 		}
+	}
 
-		if (!ControllerItems.IsValidIndex(Item.Index) || !ControllerItems[Item.Index].IsValid())
+	const bool bDeletedCategories = DeleteItemsFromLogicPanel<FRCCategoryModel>(CategoryItems, SelectedCategories);
+
+	if (bDeletedCategories)
+	{
+		for (int32 Index = 0; Index < SelectedCategoryIndices.Num(); ++Index)
 		{
-			continue;
+			OnItemRemoved({
+				ERCControllerPanelListItemType::Category,
+				SelectedCategoryIndices[Index]
+			});
 		}
-
-		SelectedControllers.Add(ControllerItems[Item.Index]);
-		SelectedControllerIndices.Add(Item.Index);
 	}
-	
-	if (DeleteItemsFromLogicPanel<FRCControllerModel>(ControllerItems, SelectedControllers))
+
+	const bool bDeletedControllers = DeleteItemsFromLogicPanel<FRCControllerModel>(ControllerItems, SelectedControllers);
+
+	if (bDeletedControllers)
 	{
 		for (int32 Index = 0; Index < SelectedControllerIndices.Num(); ++Index)
 		{
@@ -825,7 +777,8 @@ void SRCControllerPanelList::DeleteSelectedPanelItems()
 
 		UpdateEntityUsage();
 	}
-	else
+
+	if (!bDeletedCategories && !bDeletedControllers)
 	{
 		Transaction.Cancel();
 	}
@@ -861,9 +814,38 @@ TArray<TSharedPtr<FRCControllerModel>> SRCControllerPanelList::GetSelectedContro
 	return SelectedControllers;
 }
 
+TArray<TSharedPtr<FRCCategoryModel>> SRCControllerPanelList::GetSelectedCategories() const
+{
+	TArray<FRCControllerPanelListItem> SelectedItems = TreeView->GetSelectedItems();
+
+	TArray<TSharedPtr<FRCCategoryModel>> SelectedCategories;
+	SelectedCategories.Reserve(SelectedItems.Num());
+
+	for (const FRCControllerPanelListItem& Item : SelectedItems)
+	{
+		if (Item.Type != ERCControllerPanelListItemType::Category)
+		{
+			continue;
+		}
+
+		if (!CategoryItems.IsValidIndex(Item.Index) || !CategoryItems[Item.Index].IsValid())
+		{
+			continue;
+		}
+
+		SelectedCategories.Add(CategoryItems[Item.Index]);
+	}
+
+	return SelectedCategories;
+}
+
 TArray<TSharedPtr<FRCLogicModeBase>> SRCControllerPanelList::GetSelectedLogicItems()
 {
-	return TArray<TSharedPtr<FRCLogicModeBase>>(GetSelectedControllers());
+	TArray<TSharedPtr<FRCLogicModeBase>> Items;
+	Items.Append(GetSelectedCategories());
+	Items.Append(GetSelectedControllers());
+
+	return Items;
 }
 
 void SRCControllerPanelList::NotifyPreChange(FEditPropertyChain* PropertyAboutToChange)
@@ -986,7 +968,7 @@ FRCControllerPanelListItem SRCControllerPanelList::FindListItemByControllerItem(
 	return UE::RemoteControl::UI::Private::ControllerPanelListItem::None;
 }
 
-FRCControllerPanelListItem SRCControllerPanelList::FindListItemByCategoryItem(const TSharedPtr<FRCCategoryModel>& InModel) const
+FRCControllerPanelListItem SRCControllerPanelList::FindListItemByCategoryItem(const TSharedRef<FRCCategoryModel>& InModel) const
 {
 	for (int32 Index = 0; Index < CategoryItems.Num(); ++Index)
 	{
@@ -1337,7 +1319,6 @@ bool SRCControllerPanelList::ReorderItems(TConstArrayView<FRCControllerPanelList
 	);
 
 	Reset();
-	//TreeView->RequestListRefresh();
 
 	return true;
 }
@@ -1591,18 +1572,4 @@ void SRCControllerPanelList::SetMultiControllerMode(bool bIsUniqueModeOn)
 	Reset();
 }
 
-void SRCControllerPanelList::UpdateEntityUsage()
-{
-	if (TSharedPtr<SRCControllerPanel> ControllerPanel = ControllerPanelWeakPtr.Pin())
-	{
-		if (TSharedPtr<SRemoteControlPanel> RCPanel = ControllerPanel->GetRemoteControlPanel())
-		{
-			if (TSharedPtr<SRCPanelExposedEntitiesList> EntitiesList = RCPanel->GetEntityList())
-			{
-				EntitiesList->UpdateEntityUsage();
-			}
-		}
-	}
-}
-
 #undef LOCTEXT_NAMESPACE
diff --git a/Engine/Plugins/VirtualProduction/RemoteControl/Source/RemoteControlUI/Private/UI/Controller/SRCControllerPanelList.h b/Engine/Plugins/VirtualProduction/RemoteControl/Source/RemoteControlUI/Private/UI/Controller/SRCControllerPanelList.h
--- a/Engine/Plugins/VirtualProduction/RemoteControl/Source/RemoteControlUI/Private/UI/Controller/SRCControllerPanelList.h
+++ b/Engine/Plugins/VirtualProduction/RemoteControl/Source/RemoteControlUI/Private/UI/Controller/SRCControllerPanelList.h
@@ -93,6 +93,9 @@ public:
 	/** Returns the controllers selected by the user (if any). */
 	TArray<TSharedPtr<FRCControllerModel>> GetSelectedControllers() const;
 
+	/** Returns the categories selected by the user (if any). */
+	TArray<TSharedPtr<FRCCategoryModel>> GetSelectedCategories() const;
+
 	/** Returns the UI items currently selected by the user (if any). */
 	virtual TArray<TSharedPtr<FRCLogicModeBase>> GetSelectedLogicItems() override;
 
@@ -129,7 +132,7 @@ public:
 	FRCControllerPanelListItem FindListItemByControllerItem(const TSharedRef<FRCControllerModel>& InModel) const;
 
 	/** Finds the item by category id */
-	FRCControllerPanelListItem FindListItemByCategoryItem(const TSharedPtr<FRCCategoryModel>& InModel) const;
+	FRCControllerPanelListItem FindListItemByCategoryItem(const TSharedRef<FRCCategoryModel>& InModel) const;
 
 	/**
 	 * Called to moves the given controller items to the given target index, pushing all other rows below
@@ -250,9 +253,6 @@ public:
 	/** Handles broadcasting of a successful remove item operation. */
 	virtual void BroadcastOnItemRemoved() override;
 
-	/** Removes the given Controller UI model item from the list of UI models */
-	virtual int32 RemoveModel(const TSharedPtr<FRCLogicModeBase> InModel) override;
-
 	/** Set the visibility of the Tree header the Value Type Column used for MultiControllers */
 	void ShowValueTypeHeaderColumn(bool bInShowColumn);
 
@@ -278,7 +278,4 @@ public:
 
 	/** List of custom columns names */
 	TArray<FName> CustomColumns;
-
-private:
-	void UpdateEntityUsage();
 };
diff --git a/Engine/Plugins/VirtualProduction/RemoteControl/Source/RemoteControlUI/Private/UI/RCUIHelpers.cpp b/Engine/Plugins/VirtualProduction/RemoteControl/Source/RemoteControlUI/Private/UI/RCUIHelpers.cpp
--- a/Engine/Plugins/VirtualProduction/RemoteControl/Source/RemoteControlUI/Private/UI/RCUIHelpers.cpp
+++ b/Engine/Plugins/VirtualProduction/RemoteControl/Source/RemoteControlUI/Private/UI/RCUIHelpers.cpp
@@ -279,7 +279,7 @@ URCController* UE::RCUIHelpers::CreateControllerFromEntity(URemoteControlPreset*
 
 	// Create a Controller of matching type
 	URCController* const NewController = CastChecked<URCController>(InPreset->AddController(URCController::StaticClass(), PropertyBagType, StructObject));
-	NewController->DisplayIndex = InPreset->GetNumControllers() - 1;
+	NewController->DisplayIndex = InPreset->GetNextControllerDisplayIndex();
 	NewController->Description = GenerateControllerDescriptionFromEntity(InRemoteControlProperty);
 
 	// Add metadata to this controller, if this is a custom controller
diff --git a/Engine/Plugins/VirtualProduction/RemoteControl/Source/RemoteControlUI/Private/UI/Signature/Items/RCSignatureTreeItemBase.cpp b/Engine/Plugins/VirtualProduction/RemoteControl/Source/RemoteControlUI/Private/UI/Signature/Items/RCSignatureTreeItemBase.cpp
--- a/Engine/Plugins/VirtualProduction/RemoteControl/Source/RemoteControlUI/Private/UI/Signature/Items/RCSignatureTreeItemBase.cpp
+++ b/Engine/Plugins/VirtualProduction/RemoteControl/Source/RemoteControlUI/Private/UI/Signature/Items/RCSignatureTreeItemBase.cpp
@@ -5,6 +5,8 @@
 #include "UI/Signature/RCSignatureTreeItemSelection.h"
 #include "UI/Signature/SRCSignatureTree.h"
 
+const FLazyName FRCSignatureTreeItemBase::SignatureTreeItemName = "SignatureTreeItem";
+
 FRCSignatureTreeItemBase::FRCSignatureTreeItemBase(const TSharedPtr<SRCSignatureTree>& InSignatureTree)
 	: FRCLogicModeBase(InSignatureTree ? InSignatureTree->GetRemoteControlPanel() : nullptr)
 	, SignatureTreeWeak(InSignatureTree)
@@ -181,3 +183,8 @@ FName FRCSignatureTreeItemBase::BuildPath() const
 
 	return PathBuilder.ToString();
 }
+
+int32 FRCSignatureTreeItemBase::RemoveModel()
+{
+	return RemoveFromRegistry();
+}
diff --git a/Engine/Plugins/VirtualProduction/RemoteControl/Source/RemoteControlUI/Private/UI/Signature/Items/RCSignatureTreeItemBase.h b/Engine/Plugins/VirtualProduction/RemoteControl/Source/RemoteControlUI/Private/UI/Signature/Items/RCSignatureTreeItemBase.h
--- a/Engine/Plugins/VirtualProduction/RemoteControl/Source/RemoteControlUI/Private/UI/Signature/Items/RCSignatureTreeItemBase.h
+++ b/Engine/Plugins/VirtualProduction/RemoteControl/Source/RemoteControlUI/Private/UI/Signature/Items/RCSignatureTreeItemBase.h
@@ -33,6 +33,8 @@ class FRCSignatureTreeItemBase : public FRCLogicModeBase
 	friend class FRCSignatureTreeRootItem;
 
 public:
+	static const FLazyName SignatureTreeItemName;
+
 	/**
 	 * Creates an Item and Initializes it with the given Parent
 	 * NOTE: this does not build/make child items. RebuildChildren would need to be called separately if necessary.
@@ -47,6 +49,13 @@ public:
 
 	explicit FRCSignatureTreeItemBase(const TSharedPtr<SRCSignatureTree>& InSignatureTree);
 
+	virtual FName GetModelType() const override
+	{
+		return SignatureTreeItemName;
+	}
+
+	virtual int32 RemoveModel() override;
+
 	virtual TOptional<bool> IsEnabled() const
 	{
 		return TOptional<bool>();
diff --git a/Engine/Plugins/VirtualProduction/RemoteControl/Source/RemoteControlUI/Private/UI/Signature/SRCSignatureTree.cpp b/Engine/Plugins/VirtualProduction/RemoteControl/Source/RemoteControlUI/Private/UI/Signature/SRCSignatureTree.cpp
--- a/Engine/Plugins/VirtualProduction/RemoteControl/Source/RemoteControlUI/Private/UI/Signature/SRCSignatureTree.cpp
+++ b/Engine/Plugins/VirtualProduction/RemoteControl/Source/RemoteControlUI/Private/UI/Signature/SRCSignatureTree.cpp
@@ -188,15 +188,6 @@ int32 SRCSignatureTree::NumSelectedLogicItems() const
 	return SignatureTreeView->GetNumItemsSelected();
 }
 
-int32 SRCSignatureTree::RemoveModel(const TSharedPtr<FRCLogicModeBase> InItem)
-{
-	if (InItem.IsValid())
-	{
-		return StaticCastSharedPtr<FRCSignatureTreeItemBase>(InItem)->RemoveFromRegistry();
-	}
-	return 0;
-}
-
 void SRCSignatureTree::DeleteSelectedPanelItems()
 {
 	const TArray<TSharedPtr<FRCSignatureTreeItemBase>> SelectedItems = RootItem->GetSelection()->GetSelectedItems();
diff --git a/Engine/Plugins/VirtualProduction/RemoteControl/Source/RemoteControlUI/Private/UI/Signature/SRCSignatureTree.h b/Engine/Plugins/VirtualProduction/RemoteControl/Source/RemoteControlUI/Private/UI/Signature/SRCSignatureTree.h
--- a/Engine/Plugins/VirtualProduction/RemoteControl/Source/RemoteControlUI/Private/UI/Signature/SRCSignatureTree.h
+++ b/Engine/Plugins/VirtualProduction/RemoteControl/Source/RemoteControlUI/Private/UI/Signature/SRCSignatureTree.h
@@ -42,7 +42,6 @@ public:
 	virtual bool IsListFocused() const override;
 	virtual int32 Num() const override;
 	virtual int32 NumSelectedLogicItems() const override;
-	virtual int32 RemoveModel(const TSharedPtr<FRCLogicModeBase> InItem) override;
 	virtual void DeleteSelectedPanelItems() override;
 	virtual void BroadcastOnItemRemoved() override {}
 	virtual void Reset() override;
