Remote Comtrol: Controller tree view wires

diff --git a/Engine/Plugins/VirtualProduction/RemoteControl/Source/RemoteControlUI/Private/UI/Controller/SRCControllerPanelCategoryRow.cpp b/Engine/Plugins/VirtualProduction/RemoteControl/Source/RemoteControlUI/Private/UI/Controller/SRCControllerPanelCategoryRow.cpp
--- a/Engine/Plugins/VirtualProduction/RemoteControl/Source/RemoteControlUI/Private/UI/Controller/SRCControllerPanelCategoryRow.cpp
+++ b/Engine/Plugins/VirtualProduction/RemoteControl/Source/RemoteControlUI/Private/UI/Controller/SRCControllerPanelCategoryRow.cpp
@@ -27,14 +27,19 @@ void SRCControllerPanelCategoryRow::Construct(const FArguments& InArgs, const TS
 	STableRow::Construct(
 		STableRow::FArguments()
 		.Style(InArgs._Style)
-		.Padding(4.5f)
+		.Padding(0.f)
 		.OnCanAcceptDrop(this, &SRCControllerPanelCategoryRow::CanAcceptDrop)
 		.OnAcceptDrop(this, &SRCControllerPanelCategoryRow::OnAcceptDrop)
+		.ShowWires(true)
 		.Content()
 		[
-			CategoryModel.IsValid() 
-				? CategoryModel->GetNameWidget()
-				: SNullWidget::NullWidget
+			SNew(SBox)
+			.Padding(4.5f)
+			[
+				CategoryModel.IsValid() 
+					? CategoryModel->GetNameWidget()
+					: SNullWidget::NullWidget
+			]
 		],
 		InOwnerTableView
 	);
diff --git a/Engine/Plugins/VirtualProduction/RemoteControl/Source/RemoteControlUI/Private/UI/Controller/SRCControllerPanelControllerRow.cpp b/Engine/Plugins/VirtualProduction/RemoteControl/Source/RemoteControlUI/Private/UI/Controller/SRCControllerPanelControllerRow.cpp
--- a/Engine/Plugins/VirtualProduction/RemoteControl/Source/RemoteControlUI/Private/UI/Controller/SRCControllerPanelControllerRow.cpp
+++ b/Engine/Plugins/VirtualProduction/RemoteControl/Source/RemoteControlUI/Private/UI/Controller/SRCControllerPanelControllerRow.cpp
@@ -14,6 +14,7 @@
 
 #define LOCTEXT_NAMESPACE "SRCControllerPanelControllerRow"
 
+const FLazyName FRCControllerColumns::Wires = TEXT("Wires");
 const FLazyName FRCControllerColumns::TypeColor = TEXT("TypeColor");
 const FLazyName FRCControllerColumns::ControllerId = TEXT("Controller Id");
 const FLazyName FRCControllerColumns::Description = TEXT("Controller Description");
@@ -29,9 +30,10 @@ void SRCControllerPanelControllerRow::Construct(const FArguments& InArgs, const
 
 	FSuperRowType::Construct(FSuperRowType::FArguments()
 		.Style(InArgs._Style)
-		.Padding(4.5f)
+		.Padding(0.f)
 		.OnCanAcceptDrop(this, &SRCControllerPanelControllerRow::CanAcceptDrop)
 		.OnAcceptDrop(this, &SRCControllerPanelControllerRow::OnAcceptDrop)
+		.ShowWires(true)
 		, InOwnerTableView);
 }
 
@@ -44,42 +46,56 @@ TSharedRef<SWidget> SRCControllerPanelControllerRow::GenerateWidgetForColumn(con
 		return SNullWidget::NullWidget;
 	}
 
+	if (InColumnName == FRCControllerColumns::Wires)
+	{
+		return SNew(SExpanderArrow, SharedThis(this))
+			.ShouldDrawWires(true);
+	}
+
+	TSharedPtr<SWidget> RowWidget;
+
 	if (InColumnName == FRCControllerColumns::TypeColor)
 	{
-		return SNew(SBox)
-			.HAlign(HAlign_Left)
-			[
-				GetTypeColorWidget()
-			];
+		RowWidget = GetTypeColorWidget();
 	}
 	else if (InColumnName == FRCControllerColumns::FieldId)
 	{
-		return ControllerModel->GetFieldIdWidget();
+		RowWidget = ControllerModel->GetFieldIdWidget();
 	}
 	else if (InColumnName == FRCControllerColumns::ValueTypeSelection)
 	{
-		return ControllerModel->GetTypeSelectionWidget();
+		RowWidget = ControllerModel->GetTypeSelectionWidget();
 	}
 	else if (InColumnName == FRCControllerColumns::ControllerId)
 	{
-		return ControllerModel->GetNameWidget();
+		RowWidget = ControllerModel->GetNameWidget();
 	}
 	else if (InColumnName == FRCControllerColumns::Description)
 	{
-		return ControllerModel->GetDescriptionWidget();
+		RowWidget = ControllerModel->GetDescriptionWidget();
 	}
 	else if (InColumnName == FRCControllerColumns::Value)
 	{
-		return ControllerModel->GetWidget();
+		RowWidget = ControllerModel->GetWidget();
 	}
 	else if (const TSharedPtr<SRCControllerPanelList> ControlPanelList = ControllerPanelListWeak.Pin())
 	{
 		if (ControlPanelList->GetCustomColumns().Contains(InColumnName))
 		{
-			return ControllerModel->GetControllerExtensionWidget(InColumnName);
+			RowWidget = ControllerModel->GetControllerExtensionWidget(InColumnName);
 		}
 	}
-	return SNullWidget::NullWidget;
+
+	if (!RowWidget.IsValid() || RowWidget == SNullWidget::NullWidget)
+	{
+		return SNullWidget::NullWidget;
+	}
+
+	return SNew(SBox)
+		.Padding(0.f, 4.5f)
+		[
+			RowWidget.ToSharedRef()
+		];
 }
 
 TSharedPtr<FRCControllerModel> SRCControllerPanelControllerRow::GetControllerModel() const
diff --git a/Engine/Plugins/VirtualProduction/RemoteControl/Source/RemoteControlUI/Private/UI/Controller/SRCControllerPanelControllerRow.h b/Engine/Plugins/VirtualProduction/RemoteControl/Source/RemoteControlUI/Private/UI/Controller/SRCControllerPanelControllerRow.h
--- a/Engine/Plugins/VirtualProduction/RemoteControl/Source/RemoteControlUI/Private/UI/Controller/SRCControllerPanelControllerRow.h
+++ b/Engine/Plugins/VirtualProduction/RemoteControl/Source/RemoteControlUI/Private/UI/Controller/SRCControllerPanelControllerRow.h
@@ -13,6 +13,7 @@ class URemoteControlPreset;
 
 struct FRCControllerColumns
 {
+	static const FLazyName Wires;
 	static const FLazyName TypeColor;
 	static const FLazyName ControllerId;
 	static const FLazyName Description;
@@ -58,8 +59,6 @@ private:
 
 	TSharedRef<SWidget> GetTypeColorWidget();
 
-	FMargin GetTypeColorPadding() const;
-
 	//~ Begin SWidget
 	virtual FReply OnDragDetected(const FGeometry& InMyGeometry, const FPointerEvent& InMouseEvent) override;
 	virtual FReply OnMouseButtonDoubleClick(const FGeometry& InMyGeometry, const FPointerEvent& InMouseEvent) override;
diff --git a/Engine/Plugins/VirtualProduction/RemoteControl/Source/RemoteControlUI/Private/UI/Controller/SRCControllerPanelList.cpp b/Engine/Plugins/VirtualProduction/RemoteControl/Source/RemoteControlUI/Private/UI/Controller/SRCControllerPanelList.cpp
--- a/Engine/Plugins/VirtualProduction/RemoteControl/Source/RemoteControlUI/Private/UI/Controller/SRCControllerPanelList.cpp
+++ b/Engine/Plugins/VirtualProduction/RemoteControl/Source/RemoteControlUI/Private/UI/Controller/SRCControllerPanelList.cpp
@@ -62,9 +62,14 @@ void SRCControllerPanelList::Construct(const FArguments& InArgs, const TSharedRe
 			SAssignNew(ControllersHeaderRow, SHeaderRow)
 			.Style(&RCPanelStyle->HeaderRowStyle)
 
+			+ SHeaderRow::Column(FRCControllerColumns::Wires)
+			.DefaultLabel(FText())
+			.FixedWidth(20)
+			.HeaderContentPadding(RCPanelStyle->HeaderRowPadding)
+
 			+ SHeaderRow::Column(FRCControllerColumns::TypeColor)
 			.DefaultLabel(FText())
-			.FixedWidth(25)
+			.FixedWidth(20)
 			.HeaderContentPadding(RCPanelStyle->HeaderRowPadding)
 
 			+ SHeaderRow::Column(FRCControllerColumns::ControllerId)
