Remove Control: Category improvements
- Removing a category removes all contained controllers.
- Moved "Select Source" to action specific menu.
- Categories are now in a separate part of the add controller menu.
- Entities dropped above/below controllers that are in categories now add their new controllers to that category.
- Entities dropped on categories now create controllers in that category, at the end.
- Entities dropped above/below categories now create controllers above/below that category.

diff --git a/Engine/Plugins/VirtualProduction/RemoteControl/Source/RemoteControl/Private/RemoteControlPreset.cpp b/Engine/Plugins/VirtualProduction/RemoteControl/Source/RemoteControl/Private/RemoteControlPreset.cpp
--- a/Engine/Plugins/VirtualProduction/RemoteControl/Source/RemoteControl/Private/RemoteControlPreset.cpp
+++ b/Engine/Plugins/VirtualProduction/RemoteControl/Source/RemoteControl/Private/RemoteControlPreset.cpp
@@ -672,6 +672,7 @@ URemoteControlPreset::URemoteControlPreset()
 	: Layout(FRemoteControlPresetLayout{ this })
 	, PresetId(FGuid::NewGuid())
 	, RebindingManager(MakePimpl<FRemoteControlPresetRebindingManager>())
+	, NewControllerCategory(FString())
 {
 	Registry = CreateDefaultSubobject<URemoteControlExposeRegistry>(FName("ExposeRegistry"));
 	Signatures = CreateDefaultSubobject<URCSignatureRegistry>(TEXT("Signatures"));
diff --git a/Engine/Plugins/VirtualProduction/RemoteControl/Source/RemoteControl/Public/RemoteControlPreset.h b/Engine/Plugins/VirtualProduction/RemoteControl/Source/RemoteControl/Public/RemoteControlPreset.h
--- a/Engine/Plugins/VirtualProduction/RemoteControl/Source/RemoteControl/Public/RemoteControlPreset.h
+++ b/Engine/Plugins/VirtualProduction/RemoteControl/Source/RemoteControl/Public/RemoteControlPreset.h
@@ -1118,6 +1118,12 @@ public:
 	void SetDetailsTabIdentifierOverrides(const TArray<FName>& NewOverrides) { DetailsTabIdentifierOverrides = NewOverrides; }
 #endif
 
+	/**
+	 * Transient value for the default category for new properties.
+	 * Not automatically set by the preset. Used when adding controllers.
+	 */
+	FString NewControllerCategory;
+
 private:
 #if WITH_EDITORONLY_DATA
 	UPROPERTY()
diff --git a/Engine/Plugins/VirtualProduction/RemoteControl/Source/RemoteControlUI/Private/UI/Action/RCActionModel.cpp b/Engine/Plugins/VirtualProduction/RemoteControl/Source/RemoteControlUI/Private/UI/Action/RCActionModel.cpp
--- a/Engine/Plugins/VirtualProduction/RemoteControl/Source/RemoteControlUI/Private/UI/Action/RCActionModel.cpp
+++ b/Engine/Plugins/VirtualProduction/RemoteControl/Source/RemoteControlUI/Private/UI/Action/RCActionModel.cpp
@@ -315,6 +315,11 @@ void FRCActionModel::AddSpecialContextMenuOptions(FMenuBuilder& MenuBuilder)
 		{
 			EditableVirtualPropertyWidget->AddEditContextMenuOption(MenuBuilder);
 		}
+
+		const FRemoteControlCommands& Commands = FRemoteControlCommands::Get();
+
+		// 2. Select Source
+		MenuBuilder.AddMenuEntry(Commands.SelectSource);
 	}
 }
 
diff --git a/Engine/Plugins/VirtualProduction/RemoteControl/Source/RemoteControlUI/Private/UI/BaseLogicUI/SRCLogicPanelListBase.cpp b/Engine/Plugins/VirtualProduction/RemoteControl/Source/RemoteControlUI/Private/UI/BaseLogicUI/SRCLogicPanelListBase.cpp
--- a/Engine/Plugins/VirtualProduction/RemoteControl/Source/RemoteControlUI/Private/UI/BaseLogicUI/SRCLogicPanelListBase.cpp
+++ b/Engine/Plugins/VirtualProduction/RemoteControl/Source/RemoteControlUI/Private/UI/BaseLogicUI/SRCLogicPanelListBase.cpp
@@ -31,7 +31,6 @@ TSharedPtr<SWidget> SRCLogicPanelListBase::GetContextMenuWidget()
 
 	const FRemoteControlCommands& Commands = FRemoteControlCommands::Get();
 
-	MenuBuilder.AddMenuEntry(Commands.SelectSource);
 	MenuBuilder.AddMenuEntry(Commands.CopyItem);
 
 	FText PasteItemLabel = LOCTEXT("Paste", "Paste");
diff --git a/Engine/Plugins/VirtualProduction/RemoteControl/Source/RemoteControlUI/Private/UI/Behaviour/Builtin/Path/RCBehaviourSetAssetByPathModel.cpp b/Engine/Plugins/VirtualProduction/RemoteControl/Source/RemoteControlUI/Private/UI/Behaviour/Builtin/Path/RCBehaviourSetAssetByPathModel.cpp
--- a/Engine/Plugins/VirtualProduction/RemoteControl/Source/RemoteControlUI/Private/UI/Behaviour/Builtin/Path/RCBehaviourSetAssetByPathModel.cpp
+++ b/Engine/Plugins/VirtualProduction/RemoteControl/Source/RemoteControlUI/Private/UI/Behaviour/Builtin/Path/RCBehaviourSetAssetByPathModel.cpp
@@ -373,36 +373,50 @@ void FRCSetAssetByPathBehaviourModel::OnAssetPathFinishedChangingProperties(cons
 		if (InEvent.GetPropertyName() == GET_MEMBER_NAME_CHECKED(FRCSetAssetPath, AssetPath))
 		{
 			const int32 Index = InEvent.GetArrayIndex(InEvent.GetMemberPropertyName().ToString());
+
 			URCSetAssetByPathBehaviour* SetAssetByPathBehaviour = SetAssetByPathBehaviourWeakPtr.Get();
 			if (!SetAssetByPathBehaviour || Index == INDEX_NONE)
 			{
 				return;
 			}
+
 			TArray<FRCAssetPathElement>& AssetPath = SetAssetByPathBehaviour->PathStruct.AssetPath;
 			if (!AssetPath.IsValidIndex(Index))
 			{
 				return;
 			}
+
 			FRCAssetPathElement& CurrentElement = AssetPath[Index];
 			if (const TSharedPtr<SRemoteControlPanel>& RCPanel = GetRemoteControlPanel())
 			{
 				if (URemoteControlPreset* Preset = RCPanel->GetPreset())
 				{
 					URCVirtualPropertyInContainer* NewController = Preset->AddController(URCController::StaticClass(), EPropertyBagPropertyType::String);
+					
 					if (NewController)
 					{
 						NewController->DisplayIndex = Preset->GetNextControllerDisplayIndex();
+
 						FName NewControllerName = TEXT("DefaultInput");
 						if (!CurrentElement.Path.IsEmpty())
 						{
 							NewControllerName = FName(CurrentElement.Path);
 						}
+
 						const FName NewUniqueName = Preset->SetControllerDisplayName(NewController->Id, NewControllerName);
 						if (CurrentElement.Path != NewUniqueName.ToString())
 						{
 							// Assign to the path the actual new name of the controller
 							CurrentElement.Path = NewUniqueName.ToString();
 						}
+
+						if (!Preset->NewControllerCategory.IsEmpty())
+						{
+							using namespace UE::RemoteControl::UI::Private;
+
+							NewController->SetMetadataValue(FRCControllerPropertyInfo::CategoryName, Preset->NewControllerCategory);
+						}
+
 						RCPanel->OnControllerAdded.Broadcast(NewController->PropertyName);
 					}
 				}
diff --git a/Engine/Plugins/VirtualProduction/RemoteControl/Source/RemoteControlUI/Private/UI/Controller/SRCControllerPanel.cpp b/Engine/Plugins/VirtualProduction/RemoteControl/Source/RemoteControlUI/Private/UI/Controller/SRCControllerPanel.cpp
--- a/Engine/Plugins/VirtualProduction/RemoteControl/Source/RemoteControlUI/Private/UI/Controller/SRCControllerPanel.cpp
+++ b/Engine/Plugins/VirtualProduction/RemoteControl/Source/RemoteControlUI/Private/UI/Controller/SRCControllerPanel.cpp
@@ -23,6 +23,7 @@
 #include "Styling/RemoteControlStyles.h"
 #include "TypeSelector/SRCControllerTypeSelector.h"
 #include "UI/BaseLogicUI/RCLogicModeBase.h"
+#include "UI/Controller/SRCControllerPanel.h"
 #include "UI/Panels/SRCDockPanel.h"
 #include "UI/RCUIHelpers.h"
 #include "UI/RemoteControlPanelStyle.h"
@@ -365,6 +366,8 @@ TSharedRef<SWidget> SRCControllerPanel::GetControllerMenuContentWidget() const
 	constexpr  bool bShouldCloseWindowAfterMenuSelection = true;
 	FMenuBuilder MenuBuilder(bShouldCloseWindowAfterMenuSelection, nullptr);
 
+	MenuBuilder.BeginSection("Types", LOCTEXT("Types", "Types"));
+
 	TArray<TPair<EPropertyBagPropertyType, FRCControllerPropertyInfo>> VirtualPropertyFieldClassNames;
 
 	// See Config file: BaseRemoteControl.ini
@@ -534,6 +537,10 @@ TSharedRef<SWidget> SRCControllerPanel::GetControllerMenuContentWidget() const
 		MenuBuilder.AddMenuEntry(Action, MenuItemWidget);
 	}
 
+	MenuBuilder.EndSection();
+
+	MenuBuilder.BeginSection("Other", LOCTEXT("Other", "Other"));
+
 	// Category
 	{
 		const FText CategoryName = LOCTEXT("Category", "Category");
@@ -559,6 +566,8 @@ TSharedRef<SWidget> SRCControllerPanel::GetControllerMenuContentWidget() const
 		MenuBuilder.AddMenuEntry(Action, MenuItemWidget);
 	}
 
+	MenuBuilder.EndSection();
+
 	return MenuBuilder.MakeWidget();
 }
 
@@ -735,6 +744,11 @@ void SRCControllerPanel::OnAddControllerClicked(const EPropertyBagPropertyType I
 		return;
 	}
 
+	if (!Preset->NewControllerCategory.IsEmpty())
+	{
+		NewVirtualProperty->SetMetadataValue(FRCControllerPropertyInfo::CategoryName, Preset->NewControllerCategory);
+	}
+
 	// Copy over the custom metadata of this property
 	// Initial description should not be copied
 	for (const TPair<FName, FString>& Pair : InControllerPropertyInfo.CustomMetaData)
diff --git a/Engine/Plugins/VirtualProduction/RemoteControl/Source/RemoteControlUI/Private/UI/Controller/SRCControllerPanelCategoryRow.cpp b/Engine/Plugins/VirtualProduction/RemoteControl/Source/RemoteControlUI/Private/UI/Controller/SRCControllerPanelCategoryRow.cpp
--- a/Engine/Plugins/VirtualProduction/RemoteControl/Source/RemoteControlUI/Private/UI/Controller/SRCControllerPanelCategoryRow.cpp
+++ b/Engine/Plugins/VirtualProduction/RemoteControl/Source/RemoteControlUI/Private/UI/Controller/SRCControllerPanelCategoryRow.cpp
@@ -85,6 +85,11 @@ TOptional<EItemDropZone> SRCControllerPanelCategoryRow::CanAcceptDrop(const FDra
 		}
 	}
 
+	if (TSharedPtr<FExposedEntityDragDrop> ExposedEntityDragDrop = InDragDropEvent.GetOperationAs<FExposedEntityDragDrop>())
+	{
+		return InDropZone;
+	}
+
 	return TOptional<EItemDropZone>();
 }
 
@@ -122,9 +127,108 @@ FReply SRCControllerPanelCategoryRow::OnAcceptDrop(const FDragDropEvent& InDragD
 		return FReply::Handled();
 	}
 
+	if (TSharedPtr<FExposedEntityDragDrop> ExposedEntityDragDrop = InDragDropEvent.GetOperationAs<FExposedEntityDragDrop>())
+	{
+		AddControllerFromEntities(*ExposedEntityDragDrop, InDropZone);
+
+		return FReply::Handled();
+	}
+
 	return FReply::Unhandled();
 }
 
+void SRCControllerPanelCategoryRow::AddControllerFromEntities(const FExposedEntityDragDrop& InExposedEntityDragDrop, EItemDropZone InDropZone)
+{
+	TSharedPtr<FRCCategoryModel> CategoryModel = GetCategory();
+	if (!CategoryModel.IsValid())
+	{
+		return;
+	}
+
+	const TSharedPtr<SRCControllerPanelList> ControllerPanelList = ControllerPanelListWeak.Pin();
+	if (!ControllerPanelList.IsValid())
+	{
+		return;
+	}
+
+	URemoteControlPreset* Preset = ControllerPanelList->GetPreset();
+	if (!Preset)
+	{
+		return;
+	}
+
+	const FString OldDefault = Preset->NewControllerCategory;
+
+	if (InDropZone == EItemDropZone::OntoItem)
+	{
+		Preset->NewControllerCategory = CategoryModel->GetId().ToString();
+	}
+
+	FScopedTransaction Transaction(LOCTEXT("AutoBindEntities", "Auto bind entities to controllers"));
+
+	const TArray<FGuid>& SelectedFieldIds = InExposedEntityDragDrop.GetSelectedFieldsId();
+
+	TArray<URCController*, TInlineAllocator<1>> CreatedControllers;
+	TArray<TSharedRef<const FRemoteControlProperty>, TInlineAllocator<1>> SourcePropertyEntities;
+
+	CreatedControllers.Reserve(SelectedFieldIds.Num());
+	SourcePropertyEntities.Reserve(SelectedFieldIds.Num());
+
+	for (const FGuid& ExposedEntityId : SelectedFieldIds)
+	{
+		if (TSharedPtr<const FRemoteControlProperty> RemoteControlProperty = Preset->GetExposedEntity<FRemoteControlProperty>(ExposedEntityId).Pin())
+		{
+			if (URCController* NewController = UE::RCUIHelpers::CreateControllerFromEntity(Preset, RemoteControlProperty))
+			{
+				CreatedControllers.Emplace(NewController);
+				SourcePropertyEntities.Add(RemoteControlProperty.ToSharedRef());
+			}
+		}
+	}
+
+	Preset->NewControllerCategory = OldDefault;
+
+	if (CreatedControllers.IsEmpty())
+	{
+		Transaction.Cancel();
+		return;
+	}
+
+	// Refresh panel list so that the controller models are created
+	ControllerPanelList->Reset();
+
+	// Reorder the controller items to match the target drop zone
+	{
+		const FRCControllerPanelListItem NewControllerItem = ControllerPanelList->FindListItemByCategoryItem(CategoryModel.ToSharedRef());
+		const TArray<TSharedPtr<FRCControllerModel>> ControllersToMove = ControllerPanelList->FindControllerItemsByObject(CreatedControllers);
+
+		TArray<FRCControllerPanelListItem> ControllerItemsToMove;
+		ControllerItemsToMove.Reserve(CreatedControllers.Num());
+
+		for (const TSharedPtr<FRCControllerModel>& ControllerToMove : ControllersToMove)
+		{
+			if (ControllerToMove.IsValid())
+			{
+				if (const FRCControllerPanelListItem ControllerModelItem = ControllerPanelList->FindListItemByControllerItem(ControllerToMove.ToSharedRef()))
+				{
+					ControllerItemsToMove.Add(ControllerModelItem);
+				}
+			}
+		}
+
+		ControllerPanelList->ReorderItems(ControllerItemsToMove, NewControllerItem, InDropZone);
+	}
+
+	check(CreatedControllers.Num() == SourcePropertyEntities.Num());
+
+	// Create Bind Behaviour and Bind to the property
+	for (int32 Index = 0; Index < CreatedControllers.Num(); ++Index)
+	{
+		constexpr bool bExecuteBind = true;
+		ControllerPanelList->CreateBindBehaviourAndAssignTo(CreatedControllers[Index], SourcePropertyEntities[Index], bExecuteBind);
+	}
+}
+
 FReply SRCControllerPanelCategoryRow::OnDragDetected(const FGeometry& InMyGeometry, const FPointerEvent& InMouseEvent)
 {
 	if (FSuperRowType::OnDragDetected(InMyGeometry, InMouseEvent).IsEventHandled())
diff --git a/Engine/Plugins/VirtualProduction/RemoteControl/Source/RemoteControlUI/Private/UI/Controller/SRCControllerPanelCategoryRow.h b/Engine/Plugins/VirtualProduction/RemoteControl/Source/RemoteControlUI/Private/UI/Controller/SRCControllerPanelCategoryRow.h
--- a/Engine/Plugins/VirtualProduction/RemoteControl/Source/RemoteControlUI/Private/UI/Controller/SRCControllerPanelCategoryRow.h
+++ b/Engine/Plugins/VirtualProduction/RemoteControl/Source/RemoteControlUI/Private/UI/Controller/SRCControllerPanelCategoryRow.h
@@ -34,13 +34,7 @@ private:
 	/** Called when a drop happens on this controller item row */
 	FReply OnAcceptDrop(const FDragDropEvent& InDragDropEvent, EItemDropZone InDropZone, FRCControllerPanelListItem InItem);
 
-	struct FDragDropContext
-	{
-		URemoteControlPreset* Preset;
-		const TSharedPtr<SRCControllerPanelList>& ControllerPanelList;
-		const FRCControllerPanelListItem& Item;
-		EItemDropZone DropZone;
-	};
+	void AddControllerFromEntities(const FExposedEntityDragDrop& InExposedEntityDragDrop, EItemDropZone InDropZone);
 
 	//~ Begin SWidget
 	virtual FReply OnDragDetected(const FGeometry& InMyGeometry, const FPointerEvent& InMouseEvent) override;
diff --git a/Engine/Plugins/VirtualProduction/RemoteControl/Source/RemoteControlUI/Private/UI/Controller/SRCControllerPanelControllerRow.cpp b/Engine/Plugins/VirtualProduction/RemoteControl/Source/RemoteControlUI/Private/UI/Controller/SRCControllerPanelControllerRow.cpp
--- a/Engine/Plugins/VirtualProduction/RemoteControl/Source/RemoteControlUI/Private/UI/Controller/SRCControllerPanelControllerRow.cpp
+++ b/Engine/Plugins/VirtualProduction/RemoteControl/Source/RemoteControlUI/Private/UI/Controller/SRCControllerPanelControllerRow.cpp
@@ -302,6 +302,14 @@ void SRCControllerPanelControllerRow::CreateControllersFromEntities(const FExpos
 	{
 		return;
 	}
+	
+	const FString OldDefault = InContext.Preset->NewControllerCategory;
+
+	if (URCVirtualPropertyBase* Controller = ControllerModel->GetVirtualProperty())
+	{
+		using namespace UE::RemoteControl::UI::Private;
+		InContext.Preset->NewControllerCategory = Controller->GetMetadataValue(FRCControllerPropertyInfo::CategoryName);
+	}
 
 	FScopedTransaction Transaction(LOCTEXT("AutoBindEntities", "Auto bind entities to controllers"));
 
@@ -325,6 +333,8 @@ void SRCControllerPanelControllerRow::CreateControllersFromEntities(const FExpos
 		}
 	}
 
+	InContext.Preset->NewControllerCategory = OldDefault;
+
 	if (CreatedControllers.IsEmpty())
 	{
 		Transaction.Cancel();
@@ -353,7 +363,7 @@ void SRCControllerPanelControllerRow::CreateControllersFromEntities(const FExpos
 			}
 		}
 
-		InContext.ControllerPanelList->ReorderItems(ControllerItemsToMove, NewControllerItem, EItemDropZone::BelowItem);
+		InContext.ControllerPanelList->ReorderItems(ControllerItemsToMove, NewControllerItem, InContext.DropZone);
 	}
 
 	check(CreatedControllers.Num() == SourcePropertyEntities.Num());
diff --git a/Engine/Plugins/VirtualProduction/RemoteControl/Source/RemoteControlUI/Private/UI/Controller/SRCControllerPanelList.cpp b/Engine/Plugins/VirtualProduction/RemoteControl/Source/RemoteControlUI/Private/UI/Controller/SRCControllerPanelList.cpp
--- a/Engine/Plugins/VirtualProduction/RemoteControl/Source/RemoteControlUI/Private/UI/Controller/SRCControllerPanelList.cpp
+++ b/Engine/Plugins/VirtualProduction/RemoteControl/Source/RemoteControlUI/Private/UI/Controller/SRCControllerPanelList.cpp
@@ -118,6 +118,11 @@ void SRCControllerPanelList::Construct(const FArguments& InArgs, const TSharedRe
 	PropertyRowGenerator = FModuleManager::GetModuleChecked<FPropertyEditorModule>("PropertyEditor").CreatePropertyRowGenerator(Args);
 
 	Reset();
+
+	if (URemoteControlPreset* Preset = GetPreset())
+	{
+		Preset->NewControllerCategory = FString();
+	}
 }
 
 bool SRCControllerPanelList::IsEmpty() const
@@ -598,6 +603,70 @@ void SRCControllerPanelList::OnTreeSelectionChanged(FRCControllerPanelListItem I
 		RemoteControlPanel->OnControllerSelectionChanged.Broadcast(ControllerModel, InSelectInfo);
 		RemoteControlPanel->OnBehaviourSelectionChanged.Broadcast(ControllerModel ? ControllerModel->GetSelectedBehaviourModel() : nullptr);
 	}
+
+	URemoteControlPreset* Preset = RemoteControlPanel->GetPreset();
+
+	if (!Preset)
+	{
+		return;
+	}
+
+	TArray<TSharedPtr<FRCLogicModeBase>> SelectedItems = GetSelectedLogicItems();
+	FString NewCategory = FString();
+
+	for (const TSharedPtr<FRCLogicModeBase>& SelectedItem : SelectedItems)
+	{
+		if (!SelectedItem.IsValid())
+		{
+			continue;
+		}
+		
+		if (SelectedItem->GetModelType() == FRCCategoryModel::CategoryModelName)
+		{
+			if (TOptional<FRCVirtualPropertyCategory> Category = StaticCastSharedPtr<FRCCategoryModel>(SelectedItem)->GetCategory())
+			{
+				// We have encountered our first category, set it.
+				if (NewCategory.IsEmpty())
+				{
+					NewCategory = Category->Id.ToString();
+				}
+				// We have more than 1 category selected, reset value.
+				else
+				{
+					NewCategory = FString();
+					break;
+				}
+			}
+		}
+		else if (SelectedItem->GetModelType() == FRCControllerModel::ControllerModelName)
+		{
+			if (URCVirtualPropertyBase* Controller = StaticCastSharedPtr<FRCControllerModel>(SelectedItem)->GetVirtualProperty())
+			{
+				using namespace UE::RemoteControl::UI::Private;
+				const FString ControllerCategory = Controller->GetMetadataValue(FRCControllerPropertyInfo::CategoryName);
+
+				// Root item selected, no default category.
+				if (ControllerCategory.IsEmpty())
+				{
+					NewCategory = FString();
+					break;
+				}
+				// A controller in a category sets the new default.
+				else if (NewCategory.IsEmpty())
+				{
+					NewCategory = ControllerCategory;
+				}
+				// If it's the same, we do nothing. However a multiple selected categories means no default.
+				else if (NewCategory != ControllerCategory)
+				{
+					NewCategory = FString();
+					break;
+				}
+			}
+		}
+	}
+
+	Preset->NewControllerCategory = NewCategory;
 }
 
 void SRCControllerPanelList::SelectController(URCController* InController)
@@ -733,6 +802,8 @@ void SRCControllerPanelList::DeleteSelectedPanelItems()
 	TArray<int32> SelectedControllerIndices;
 	SelectedControllerIndices.Reserve(SelectedItems.Num());
 
+	TSet<FString> CategoryNames;
+
 	for (const FRCControllerPanelListItem& Item : SelectedItems)
 	{
 		switch (Item.Type)
@@ -742,6 +813,7 @@ void SRCControllerPanelList::DeleteSelectedPanelItems()
 				{
 					SelectedCategories.Add(CategoryItems[Item.Index]);
 					SelectedCategoryIndices.Add(Item.Index);
+					CategoryNames.Add(CategoryItems[Item.Index]->GetId().ToString());
 				}
 				break;
 
@@ -755,15 +827,39 @@ void SRCControllerPanelList::DeleteSelectedPanelItems()
 		}
 	}
 
+	// Also remove controllers in deleted categories.
+	for (int32 ControllerIndex = 0; ControllerIndex < ControllerItems.Num(); ++ControllerIndex)
+	{
+		if (!SelectedControllerIndices.Contains(ControllerIndex))
+		{
+			if (ControllerItems[ControllerIndex].IsValid())
+			{
+				if (URCVirtualPropertyBase* Controller = ControllerItems[ControllerIndex]->GetVirtualProperty())
+				{
+					using namespace UE::RemoteControl::UI::Private;
+					const FString Category = Controller->GetMetadataValue(FRCControllerPropertyInfo::CategoryName);
+
+					if (!Category.IsEmpty() && CategoryNames.Contains(Category))
+					{
+						SelectedControllers.Add(ControllerItems[ControllerIndex]);
+						SelectedControllerIndices.Add(ControllerIndex);
+					}
+				}
+			}
+		}
+	}
+
 	const bool bDeletedCategories = DeleteItemsFromLogicPanel<FRCCategoryModel>(CategoryItems, SelectedCategories);
 
 	if (bDeletedCategories)
 	{
-		for (int32 Index = 0; Index < SelectedCategoryIndices.Num(); ++Index)
+		 CategoryNames.Reserve(SelectedCategoryIndices.Num());
+
+		for (int32 SelectedIndex = 0; SelectedIndex < SelectedCategoryIndices.Num(); ++SelectedIndex)
 		{
 			OnItemRemoved({
 				ERCControllerPanelListItemType::Category,
-				SelectedCategoryIndices[Index]
+				SelectedCategoryIndices[SelectedIndex]
 			});
 		}
 	}
@@ -962,9 +1058,11 @@ TArray<TSharedPtr<FRCControllerModel>> SRCControllerPanelList::FindControllerIte
 
 FRCControllerPanelListItem SRCControllerPanelList::FindListItemByControllerItem(const TSharedRef<FRCControllerModel>& InModel) const
 {
+	URCVirtualPropertyBase* InController = InModel->GetVirtualProperty();
+
 	for (int32 Index = 0; Index < ControllerItems.Num(); ++Index)
 	{
-		if (ControllerItems[Index] == InModel)
+		if (ControllerItems[Index].IsValid() && ControllerItems[Index]->GetVirtualProperty() == InController)
 		{
 			return {ERCControllerPanelListItemType::Controller, Index};
 		}
@@ -975,9 +1073,11 @@ FRCControllerPanelListItem SRCControllerPanelList::FindListItemByControllerItem(
 
 FRCControllerPanelListItem SRCControllerPanelList::FindListItemByCategoryItem(const TSharedRef<FRCCategoryModel>& InModel) const
 {
+	const FGuid InCategoryId = InModel->GetId();
+
 	for (int32 Index = 0; Index < CategoryItems.Num(); ++Index)
 	{
-		if (CategoryItems[Index] == InModel)
+		if (CategoryItems[Index].IsValid() && CategoryItems[Index]->GetId() == InCategoryId)
 		{
 			return {ERCControllerPanelListItemType::Category, Index};
 		}
@@ -1014,6 +1114,43 @@ bool SRCControllerPanelList::ReorderItems(TConstArrayView<FRCControllerPanelList
 			break;
 	}
 
+	using namespace UE::RemoteControl::UI::Private;
+
+	// If we reorder onto a category, find the last item of the category and drop it below that.
+	if (InDroppedOnItem.Type == ERCControllerPanelListItemType::Category && InDropZone == EItemDropZone::OntoItem)
+	{
+		if (TSharedPtr<FRCCategoryModel> DroppedOnCategory = FindCategoryItemByListItem(InDroppedOnItem))
+		{
+			if (TOptional<FRCVirtualPropertyCategory> Category = DroppedOnCategory->GetCategory())
+			{
+				const FString CategoryId = Category->Id.ToString();
+				int32 MaxControllerIndex = INDEX_NONE;
+
+				for (int32 Index = 0; Index < ControllerItems.Num(); ++Index)
+				{
+					if (ControllerItems[Index].IsValid())
+					{
+						if (URCVirtualPropertyBase* Controller = ControllerItems[Index]->GetVirtualProperty())
+						{
+							const FString ControllerCategory = Controller->GetMetadataValue(FRCControllerPropertyInfo::CategoryName);
+
+							if (ControllerCategory == CategoryId)
+							{
+								MaxControllerIndex = FMath::Max(MaxControllerIndex, Index);
+							}
+						}
+					}
+				}
+
+				if (MaxControllerIndex != INDEX_NONE)
+				{
+					InDroppedOnItem = {ERCControllerPanelListItemType::Controller, MaxControllerIndex};
+					InDropZone = EItemDropZone::BelowItem;
+				}
+			}
+		}
+	}
+
 	// Rebuild list assuming that everything is already in the correct display order
 	TMap<FName, TArray<FRCControllerPanelListItem>> Tree;
 	Tree.Add(NAME_None);
@@ -1128,8 +1265,6 @@ bool SRCControllerPanelList::ReorderItems(TConstArrayView<FRCControllerPanelList
 			}
 		};
 
-	using namespace UE::RemoteControl::UI::Private;
-
 	FRCControllerPanelListItem NextItem = GetNextItem();
 
 	while (NextItem.Type != ERCControllerPanelListItemType::None)
@@ -1395,6 +1530,8 @@ FReply SRCControllerPanelList::OnControllerTreeViewDragDrop(TSharedPtr<FDragDrop
 	if (TSharedPtr<FExposedEntityDragDrop> DragDropOp = GetExposedEntityDragDrop(InDragDropOperation))
 	{
 		FScopedTransaction Transaction(LOCTEXT("AutoBindEntities", "Auto bind entities to controllers"));
+		const FString OldDefault = Preset->NewControllerCategory;
+		Preset->NewControllerCategory = FString();
 
 		bool bModified = false;
 		for (const FGuid& ExposedEntityId : DragDropOp->GetSelectedFieldsId())
@@ -1405,6 +1542,8 @@ FReply SRCControllerPanelList::OnControllerTreeViewDragDrop(TSharedPtr<FDragDrop
 			}
 		}
 
+		Preset->NewControllerCategory = OldDefault;
+
 		if (!bModified)
 		{
 			Transaction.Cancel();
diff --git a/Engine/Plugins/VirtualProduction/RemoteControl/Source/RemoteControlUI/Private/UI/RCUIHelpers.cpp b/Engine/Plugins/VirtualProduction/RemoteControl/Source/RemoteControlUI/Private/UI/RCUIHelpers.cpp
--- a/Engine/Plugins/VirtualProduction/RemoteControl/Source/RemoteControlUI/Private/UI/RCUIHelpers.cpp
+++ b/Engine/Plugins/VirtualProduction/RemoteControl/Source/RemoteControlUI/Private/UI/RCUIHelpers.cpp
@@ -21,6 +21,7 @@
 #include "RemoteControlField.h"
 #include "RemoteControlPreset.h"
 #include "TimerManager.h"
+#include "UI/Controller/SRCControllerPanel.h"
 #include "Widgets/Layout/SBorder.h"
 #include "Widgets/Layout/SBox.h"
 
@@ -282,6 +283,12 @@ URCController* UE::RCUIHelpers::CreateControllerFromEntity(URemoteControlPreset*
 	NewController->DisplayIndex = InPreset->GetNextControllerDisplayIndex();
 	NewController->Description = GenerateControllerDescriptionFromEntity(InRemoteControlProperty);
 
+	if (!InPreset->NewControllerCategory.IsEmpty())
+	{
+		using namespace UE::RemoteControl::UI::Private;
+		NewController->SetMetadataValue(FRCControllerPropertyInfo::CategoryName, InPreset->NewControllerCategory);
+	}
+
 	// Add metadata to this controller, if this is a custom controller
 	if (!CustomControllerName.IsEmpty())
 	{
