Remote Control: Path Behavior 2: Back in Action

diff --git a/Engine/Plugins/VirtualProduction/RemoteControl/Source/RemoteControlLogic/Private/Action/Path/RCSetAssetByPathActionNew.cpp b/Engine/Plugins/VirtualProduction/RemoteControl/Source/RemoteControlLogic/Private/Action/Path/RCSetAssetByPathActionNew.cpp
new file mode 100644
--- /dev/null
+++ b/Engine/Plugins/VirtualProduction/RemoteControl/Source/RemoteControlLogic/Private/Action/Path/RCSetAssetByPathActionNew.cpp
@@ -0,0 +1,354 @@
+// Copyright Epic Games, Inc. All Rights Reserved.
+
+#include "Action/Path/RCSetAssetByPathActionNew.h"
+
+#include "Action/Bind/RCCustomBindActionUtilitiesPrivate.h"
+#include "Backends/CborStructDeserializerBackend.h"
+#include "Backends/CborStructSerializerBackend.h"
+#include "Behaviour/Builtin/Path/RCSetAssetByPathBehaviorNew.h"
+#include "Components/MeshComponent.h"
+#include "Components/StaticMeshComponent.h"
+#include "Controller/RCController.h"
+#include "Engine/Blueprint.h"
+#include "Engine/StaticMesh.h"
+#include "Engine/Texture2D.h"
+#include "Engine/Texture2DDynamic.h"
+#include "Engine/World.h"
+#include "IRemoteControlModule.h"
+#include "IRemoteControlPropertyHandle.h"
+#include "Kismet/KismetRenderingLibrary.h"
+#include "Materials/Material.h"
+#include "Materials/MaterialInstance.h"
+#include "Misc/FileHelper.h"
+#include "Misc/Paths.h"
+#include "RCVirtualProperty.h"
+#include "RCVirtualPropertyContainer.h"
+#include "RemoteControlLogger.h"
+#include "RemoteControlPreset.h"
+
+void URCSetAssetByPathActionNew::Execute() const
+{
+	if (IsInternal())
+	{
+		const_cast<URCSetAssetByPathActionNew*>(this)->UpdateInternal();
+	}
+	else
+	{
+		const_cast<URCSetAssetByPathActionNew*>(this)->UpdateExternal();
+	}
+}
+
+URCSetAssetByPathBehaviorNew* URCSetAssetByPathActionNew::GetSetAssetByPathBehavior() const
+{
+	return Cast<URCSetAssetByPathBehaviorNew>(GetParentBehaviour());
+}
+
+bool URCSetAssetByPathActionNew::IsInternal() const
+{
+	if (URCSetAssetByPathBehaviorNew* SetAssetByPathBehavior = GetSetAssetByPathBehavior())
+	{
+		return SetAssetByPathBehavior->bInternal;
+	}
+
+	return true;
+}
+
+FString URCSetAssetByPathActionNew::GetPath() const
+{
+	URCSetAssetByPathBehaviorNew* SetAssetByPath = GetSetAssetByPathBehavior();
+	if (!SetAssetByPath)
+	{
+		return FString();
+	}
+
+	FString BaseAssetPath = SetAssetByPath->GetCurrentPath();
+	if (BaseAssetPath.IsEmpty())
+	{
+		return FString();
+	}
+
+	return BaseAssetPath.Replace(TEXT("//"), TEXT("/"));;
+}
+
+bool URCSetAssetByPathActionNew::UpdateInternal()
+{
+	const FString AssetPath = GetPath();
+
+	TSharedPtr<FRemoteControlProperty> RemoteControlProperty = GetRemoteControlProperty();
+	if (!RemoteControlProperty.IsValid())
+	{
+		return false;
+	}
+
+	FProperty* Property = RemoteControlProperty->GetProperty();
+	if (!Property)
+	{
+		return false;
+	}
+
+	FSoftObjectPath MainObjectRef(AssetPath);
+
+	UObject* Object = MainObjectRef.TryLoad();
+
+	if (!Object)
+	{
+		return false;
+	}
+
+	using namespace UE::RemoteControl::Logic;
+
+	FRemoteControlLogger::Get().Log(SetAssetByPathNew::SetAssetByPathBehaviour, [MainObjectRef]
+		{
+			return FText::FromString(FString("Path Behavior attempts to set Asset to %s") + *MainObjectRef.GetAssetPathString());
+		});
+
+#if !WITH_EDITOR
+	return SetInternalAsset(Object);
+#else
+	FEditPropertyChain PropertyChain;
+	PropertyChain.AddHead(Property);
+	RemoteControlProperty->GetBoundObject()->PreEditChange(PropertyChain);
+
+	if (SetInternalAsset(Object))
+	{
+		FPropertyChangedEvent ChangedEvent(Property, EPropertyChangeType::ValueSet);
+		RemoteControlProperty->GetBoundObject()->PostEditChangeProperty(ChangedEvent);
+
+		return true;
+	}
+
+	return false;
+#endif
+}
+
+bool URCSetAssetByPathActionNew::UpdateExternal()
+{
+	const FString AssetPath = GetPath();
+
+	using namespace UE::RemoteControl::Logic;
+
+	FRemoteControlLogger::Get().Log(SetAssetByPathNew::SetAssetByPathBehaviour, [AssetPath]
+		{
+			return FText::FromString(FString("Path Behavior attempts to set Asset to %s") + *AssetPath);
+		});
+
+#if !WITH_EDITOR
+	return SetExternalAsset(AssetPath);
+#else
+	TSharedPtr<FRemoteControlProperty> RemoteControlProperty = GetRemoteControlProperty();
+	if (!RemoteControlProperty.IsValid())
+	{
+		return false;
+	}
+
+	FProperty* Property = RemoteControlProperty->GetProperty();
+	if (!Property)
+	{
+		return false;
+	}
+
+	FEditPropertyChain PropertyChain;
+	PropertyChain.AddHead(Property);
+	RemoteControlProperty->GetBoundObject()->PreEditChange(PropertyChain);
+
+	if (SetExternalAsset(AssetPath))
+	{
+		FPropertyChangedEvent ChangedEvent(Property, EPropertyChangeType::ValueSet);
+		RemoteControlProperty->GetBoundObject()->PostEditChangeProperty(ChangedEvent);
+
+		return true;
+	}
+
+	return false;
+#endif
+}
+
+bool URCSetAssetByPathActionNew::SetTextureFromPath(TSharedPtr<FRemoteControlProperty> InRCPropertyToSet, const FString& InFileName)
+{
+	if (InRCPropertyToSet.IsValid())
+	{
+		if (URCBehaviour* Behavior = GetParentBehaviour())
+		{
+			if (URCController* Controller = Behavior->ControllerWeakPtr.Get())
+			{
+				UE::RCCustomBindActionUtilities::SetTexturePropertyFromPath(InRCPropertyToSet.ToSharedRef(), Controller, InFileName);
+			}
+		}
+
+		return true;
+	}
+
+	return false;
+}
+
+bool URCSetAssetByPathActionNew::SetInternalAsset(UObject* InSetterObject)
+{
+	using namespace UE::RemoteControl::Logic;
+
+	if (!InSetterObject)
+	{
+		FRemoteControlLogger::Get().Log(SetAssetByPathNew::SetAssetByPathBehaviour, [InSetterObject]
+			{
+				return FText::FromString(FString("Path Behavior fails to set Asset: %s") + *InSetterObject->GetName());
+			}, EVerbosityLevel::Error);
+
+		return false;
+	}
+
+	URCSetAssetByPathBehaviorNew* SetAssetByPath = GetSetAssetByPathBehavior();
+	if (!SetAssetByPath)
+	{
+		return false;
+	}
+
+	URCController* Controller = SetAssetByPath->ControllerWeakPtr.Get();
+	if (!Controller)
+	{
+		return false;
+	}
+
+	URemoteControlPreset* Preset = Controller->PresetWeakPtr.Get();
+	if (!Preset)
+	{
+		return false;
+	}
+
+	TSharedPtr<FRemoteControlProperty> RemoteControlProperty = GetRemoteControlProperty();
+
+	if (!RemoteControlProperty.IsValid())
+	{
+		return false;
+	}
+
+	if (InSetterObject->IsA(UBlueprint::StaticClass()))
+	{
+		if (AActor* OldActor = Cast<AActor>(RemoteControlProperty->GetBoundObject()))
+		{
+			UWorld* World = OldActor->GetWorld();
+
+			OldActor->UnregisterAllComponents();
+
+			const FVector OldLocation = OldActor->GetActorLocation();
+			const FRotator OldRotation = OldActor->GetActorRotation();
+			const FName OldActorName = OldActor->GetFName();
+
+			FActorSpawnParameters SpawnParams;
+			SpawnParams.Name = OldActorName;
+
+			const FName OldActorReplacedNamed = MakeUniqueObjectName(OldActor->GetOuter(), OldActor->GetClass(), *FString::Printf(TEXT("%s_REPLACED"), *OldActor->GetFName().ToString()));
+			OldActor->Rename(*OldActorReplacedNamed.ToString(), OldActor->GetOuter(), REN_DoNotDirty | REN_DontCreateRedirectors);
+
+			if (AActor* NewActor = World->SpawnActor(Cast<UBlueprint>(InSetterObject)->GeneratedClass, &OldLocation, &OldRotation, SpawnParams))
+			{
+				World->DestroyActor(OldActor);
+
+				UE_LOG(LogRemoteControl, Display, TEXT("Path Behavior sets Blueprint Asset %s"), *InSetterObject->GetName());
+				return true;
+			}
+
+			UE_LOG(LogRemoteControl, Error, TEXT("Path Behavior unable to delete old Actor %s"), *OldActor->GetName());
+			OldActor->UObject::Rename(*OldActorName.ToString(), OldActor->GetOuter(), REN_DoNotDirty | REN_DontCreateRedirectors);
+			OldActor->RegisterAllComponents();
+
+			return false;
+		}
+	}
+
+	if (TSharedPtr<IRemoteControlPropertyHandle> PropertyHandle = RemoteControlProperty->GetPropertyHandle())
+	{
+		UClass* SetterObjectClass = InSetterObject->GetClass();
+
+		if (TSharedPtr<IRemoteControlPropertyHandleArray> ArrayHandle = PropertyHandle->AsArray())
+		{
+			const int32 Index = PropertyHandle->GetIndexInArray();
+
+			if (Index >= 0)
+			{
+				if (ArrayHandle->GetNumElements() > 0)
+				{
+					if (TSharedPtr<IRemoteControlPropertyHandle> SubPropertyHandle = ArrayHandle->GetElement(Index))
+					{
+						if (FObjectPropertyBase* ObjectProperty = CastField<FObjectPropertyBase>(SubPropertyHandle->GetProperty()))
+						{
+							if (ObjectProperty->PropertyClass && (SetterObjectClass == ObjectProperty->PropertyClass || SetterObjectClass->IsChildOf(ObjectProperty->PropertyClass)))
+							{
+								UE_LOG(LogRemoteControl, Display, TEXT("Path Behavior sets Array Property %s"), *InSetterObject->GetName());
+								SubPropertyHandle->SetValue(InSetterObject);
+								return true;
+							}
+						}
+					}
+				}
+			}
+		}
+		else if (FObjectPropertyBase* ObjectProperty = CastField<FObjectPropertyBase>(PropertyHandle->GetProperty()))
+		{
+			if (ObjectProperty->PropertyClass && (SetterObjectClass == ObjectProperty->PropertyClass || SetterObjectClass->IsChildOf(ObjectProperty->PropertyClass)))
+			{
+				UE_LOG(LogRemoteControl, Display, TEXT("Path Behavior sets Property %s"), *InSetterObject->GetName());
+				PropertyHandle->SetValue(InSetterObject);
+				return true;
+			}
+		}
+	}
+
+	UE_LOG(LogRemoteControl, Error, TEXT("Path Behavior fails to set Asset %s"), *InSetterObject->GetName());
+	return false;
+}
+
+bool URCSetAssetByPathActionNew::SetExternalAsset(const FString& InExternalPath)
+{
+	using namespace UE::RemoteControl::Logic;
+
+	if (InExternalPath.IsEmpty())
+	{
+		FRemoteControlLogger::Get().Log(SetAssetByPathNew::SetAssetByPathBehaviour, [InExternalPath]
+			{
+				return FText::FromString(FString("Path Behavior fails to set Asset: %s") + *InExternalPath);
+			}, EVerbosityLevel::Error);
+
+		return false;
+	}
+
+	URCSetAssetByPathBehaviorNew* SetAssetByPath = GetSetAssetByPathBehavior();
+	if (!SetAssetByPath)
+	{
+		return false;
+	}
+
+	URCController* Controller = SetAssetByPath->ControllerWeakPtr.Get();
+	if (!Controller)
+	{
+		return false;
+	}
+
+	URemoteControlPreset* Preset = Controller->PresetWeakPtr.Get();
+	if (!Preset)
+	{
+		return false;
+	}
+
+	TSharedPtr<FRemoteControlProperty> RemoteControlProperty = GetRemoteControlProperty();
+
+	if (!RemoteControlProperty.IsValid())
+	{
+		return false;
+	}
+
+	if (SetTextureFromPath(RemoteControlProperty, InExternalPath))
+	{
+		FRemoteControlLogger::Get().Log(SetAssetByPathNew::SetAssetByPathBehaviour, [InExternalPath]
+		{
+			return FText::FromString(FString("Path Behavior Set external Asset: ") + *InExternalPath);
+		});
+
+		return true;
+	}
+
+	FRemoteControlLogger::Get().Log(SetAssetByPathNew::SetAssetByPathBehaviour, []
+	{
+		return FText::FromString(TEXT("Path Behavior Set external Asset failed"));
+	});
+
+	return false;
+}
diff --git a/Engine/Plugins/VirtualProduction/RemoteControl/Source/RemoteControlLogic/Private/Behaviour/Builtin/Path/RCBehaviorSetAssetByPathNodeNew.cpp b/Engine/Plugins/VirtualProduction/RemoteControl/Source/RemoteControlLogic/Private/Behaviour/Builtin/Path/RCBehaviorSetAssetByPathNodeNew.cpp
new file mode 100644
--- /dev/null
+++ b/Engine/Plugins/VirtualProduction/RemoteControl/Source/RemoteControlLogic/Private/Behaviour/Builtin/Path/RCBehaviorSetAssetByPathNodeNew.cpp
@@ -0,0 +1,35 @@
+// Copyright Epic Games, Inc. All Rights Reserved.
+
+#include "Behaviour/Builtin/Path/RCBehaviorSetAssetByPathNodeNew.h"
+
+#include "Behaviour/Builtin/Path/RCSetAssetByPathBehaviorNew.h"
+
+#define LOCTEXT_NAMESPACE "RCBehaviorSetAssetByPathNodeNew"
+
+URCBehaviorSetAssetByPathNodeNew::URCBehaviorSetAssetByPathNodeNew()
+{
+	DisplayName = LOCTEXT("Behaviour Name - Set Asset By Path", "Path");
+	BehaviorDescription = LOCTEXT("Behaviour Desc - Set Asset By Path", "Triggers an event which sets objects by building a path to an asset.");
+}
+
+bool URCBehaviorSetAssetByPathNodeNew::Execute_Implementation(URCBehaviour* InBehavior) const
+{
+	return true;
+}
+
+bool URCBehaviorSetAssetByPathNodeNew::IsSupported_Implementation(URCBehaviour* InBehavior) const
+{
+	static const TArray<TPair<EPropertyBagPropertyType, UObject*>> SupportedPropertyBagTypes =
+	{
+		{EPropertyBagPropertyType::String, nullptr}
+	};
+	
+	return SupportedPropertyBagTypes.ContainsByPredicate(GetIsSupportedCallback(InBehavior));
+}
+
+UClass* URCBehaviorSetAssetByPathNodeNew::GetBehaviourClass() const
+{
+	return URCSetAssetByPathBehaviorNew::StaticClass();
+}
+
+#undef LOCTEXT_NAMESPACE
diff --git a/Engine/Plugins/VirtualProduction/RemoteControl/Source/RemoteControlLogic/Private/Behaviour/Builtin/Path/RCSetAssetByPathBehaviorNew.cpp b/Engine/Plugins/VirtualProduction/RemoteControl/Source/RemoteControlLogic/Private/Behaviour/Builtin/Path/RCSetAssetByPathBehaviorNew.cpp
new file mode 100644
--- /dev/null
+++ b/Engine/Plugins/VirtualProduction/RemoteControl/Source/RemoteControlLogic/Private/Behaviour/Builtin/Path/RCSetAssetByPathBehaviorNew.cpp
@@ -0,0 +1,185 @@
+// Copyright Epic Games, Inc. All Rights Reserved.
+
+#include "Behaviour/Builtin/Path/RCSetAssetByPathBehaviorNew.h"
+
+#include "Action/Path/RCSetAssetByPathActionNew.h"
+#include "Controller/RCController.h"
+#include "RemoteControlField.h"
+#include "RemoteControlPreset.h"
+#include "RCVirtualProperty.h"
+
+#if WITH_EDITOR
+#include "PropertyPath.h"
+#endif
+
+URCSetAssetByPathActionNew* URCSetAssetByPathBehaviorNew::AddPropertyBindAction(const TSharedRef<const FRemoteControlProperty> InRemoteControlProperty)
+{
+	if (!ensure(ControllerWeakPtr.IsValid() && ActionContainer))
+	{
+		return nullptr;
+	}
+
+	URCSetAssetByPathActionNew* BindAction = NewObject<URCSetAssetByPathActionNew>(ActionContainer);
+	BindAction->PresetWeakPtr = ActionContainer->PresetWeakPtr;
+	BindAction->ExposedFieldId = InRemoteControlProperty->GetId();
+	BindAction->Id = FGuid::NewGuid();
+
+	// Add action to array
+	ActionContainer->AddAction(BindAction);
+
+	return BindAction;
+}
+
+URCAction* URCSetAssetByPathBehaviorNew::AddAction(const TSharedRef<const FRemoteControlField> InRemoteControlField)
+{
+	const TSharedRef<const FRemoteControlProperty> RemoteControlProperty = StaticCastSharedRef<const FRemoteControlProperty>(InRemoteControlField);
+
+	URCAction* SetAssetByPathAction = AddPropertyBindAction(RemoteControlProperty);
+
+	if (ensure(SetAssetByPathAction))
+	{
+		SetAssetByPathAction->Execute();
+	}
+
+	return SetAssetByPathAction;
+}
+
+#if WITH_EDITOR
+bool URCSetAssetByPathBehaviorNew::CanHaveActionForField(const TSharedPtr<FRemoteControlField> InRemoteControlField) const
+{
+	if (!InRemoteControlField.IsValid())
+	{
+		return false;
+	}
+
+	if (!Super::CanHaveActionForField(InRemoteControlField))
+	{
+		return false;
+	}
+
+	if (!InRemoteControlField->FieldPathInfo.IsResolved())
+	{
+		UObject* BoundObject = InRemoteControlField->GetBoundObject();
+
+		if (!BoundObject)
+		{
+			return false;
+		}
+
+		if (!InRemoteControlField->FieldPathInfo.Resolve(BoundObject))
+		{
+			return false;
+		}		
+	}
+
+	TSharedPtr<FPropertyPath> PropertyPath = InRemoteControlField->FieldPathInfo.ToPropertyPath();
+
+	if (!PropertyPath.IsValid() || PropertyPath->GetNumProperties() <= 0)
+	{
+		return false;
+	}
+
+	FProperty* Property = PropertyPath->GetPropertyInfo(0).Property.Get();
+
+	if (!Property)
+	{
+		return false;
+	}
+
+	if (FArrayProperty* ArrayProperty = CastField<FArrayProperty>(Property))
+	{
+		if (ArrayProperty->Inner && ArrayProperty->IsA<FObjectPropertyBase>())
+		{
+			return true;
+		}
+	}
+	else if (Property->IsA<FObjectPropertyBase>())
+	{
+		return true;
+	}
+
+	return false;
+}
+#endif
+
+FString URCSetAssetByPathBehaviorNew::GetCurrentPath() const
+{
+	using namespace UE::RemoteControl::Logic;
+
+	FString CurrentPath;
+	URCController* RCController = ControllerWeakPtr.Get();
+	if (!RCController)
+	{
+		ensureMsgf(false, TEXT("Controller is invalid."));
+		return TEXT("Path is invalid!");
+	}
+
+	CurrentPath.Empty();
+	FString ControllerString;
+	RCController->GetValueString(ControllerString);
+
+	// Add Path String Concat
+	CurrentPath = bInternal ? SetAssetByPathNew::ContentFolder : FString("");
+	
+	for (FRCAssetPathElementNew PathPart : PathStruct.AssetPath)
+	{
+		if (PathPart.bIsInput)
+		{
+			if (PathPart.Path.IsEmpty())
+			{
+				URCController* Controller = ControllerWeakPtr.Get();
+				if (!Controller)
+				{
+					ensureMsgf(false, TEXT("Cannot find parent controller"));
+					continue;
+				}
+
+				FString ControllerValue;
+				if (!Controller->GetValueString(ControllerValue))
+				{
+					ensureMsgf(false, TEXT("Unable to get controller value as string"));
+					continue;
+				}
+
+				PathPart.Path = ControllerValue;
+			}
+			else
+			{
+				const URemoteControlPreset* RemoteControlPreset = RCController->PresetWeakPtr.Get();
+				if (!RemoteControlPreset)
+				{
+					ensureMsgf(false, TEXT("Remote Control Preset is invalid"));
+					continue;
+				}
+
+				const URCVirtualPropertyBase* TokenController = RemoteControlPreset->GetControllerByDisplayName(FName(PathPart.Path));
+				if (!TokenController)
+				{
+					PathPart.Path = TEXT("InvalidControllerName");
+				}
+				else
+				{
+					PathPart.Path = TokenController->GetDisplayValueAsString();
+				}
+			}
+		}
+		else if (PathPart.Path.IsEmpty())
+		{
+			continue;
+		}
+
+		PathPart.Path.ReplaceCharInline(TCHAR('\\'), TCHAR('/'), ESearchCase::IgnoreCase);
+		
+		CurrentPath += PathPart.Path;
+	}
+
+	return CurrentPath.Replace(TEXT("//"), TEXT("/"));
+}
+
+void URCSetAssetByPathBehaviorNew::RefreshPathArray()
+{
+	if (PathStruct.AssetPath.Num() < 1)
+	{
+		PathStruct.AssetPath.AddDefaulted();
+	}
+}
diff --git a/Engine/Plugins/VirtualProduction/RemoteControl/Source/RemoteControlLogic/Public/Action/Path/RCSetAssetByPathActionNew.h b/Engine/Plugins/VirtualProduction/RemoteControl/Source/RemoteControlLogic/Public/Action/Path/RCSetAssetByPathActionNew.h
new file mode 100644
--- /dev/null
+++ b/Engine/Plugins/VirtualProduction/RemoteControl/Source/RemoteControlLogic/Public/Action/Path/RCSetAssetByPathActionNew.h
@@ -0,0 +1,49 @@
+// Copyright Epic Games, Inc. All Rights Reserved.
+
+#pragma once
+
+#include "Action/RCPropertyAction.h"
+#include "RCSetAssetByPathActionNew.generated.h"
+
+class FString;
+class URCSetAssetByPathBehaviorNew;
+class UTexture;
+struct FRemoteControlProperty;
+
+/**
+ * Property Bind Action specifically for Bind Behaviour
+ */
+UCLASS(MinimalAPI)
+class URCSetAssetByPathActionNew : public URCPropertyAction
+{
+	GENERATED_BODY()
+
+public:	
+	//~ Begin URCAction interface
+	REMOTECONTROLLOGIC_API virtual void Execute() const override;
+	//~ End URCAction interface
+
+	REMOTECONTROLLOGIC_API URCSetAssetByPathBehaviorNew* GetSetAssetByPathBehavior() const;
+
+	/** Internal -> Asset. External -> File. External is only supported by textures. */
+	REMOTECONTROLLOGIC_API bool IsInternal() const;
+
+private:
+	/** Gets the path to the asset. */
+	FString GetPath() const;
+
+	/** Updates property with a existing asset. */
+	bool UpdateInternal();
+
+	/** Updatse property with an imported, external texture. */
+	bool UpdateExternal();
+
+	/** Auxiliary Function which sets the given SetterObject Asset onto the Exposed Asset with the given PropertyString name. */
+	bool SetInternalAsset(UObject* InSetterObject);
+
+	/** Given an External Path towards a external location, loads the asset associated and places it onto an object */
+	bool SetExternalAsset(const FString& InExternalPath);
+
+	/** Auxiliary Function to apply a Texture onto a given Property */
+	bool SetTextureFromPath(TSharedPtr<FRemoteControlProperty> InRCPropertyToSet, const FString& InFileName);
+};
diff --git a/Engine/Plugins/VirtualProduction/RemoteControl/Source/RemoteControlLogic/Public/Behaviour/Builtin/Path/RCBehaviorSetAssetByPathNodeNew.h b/Engine/Plugins/VirtualProduction/RemoteControl/Source/RemoteControlLogic/Public/Behaviour/Builtin/Path/RCBehaviorSetAssetByPathNodeNew.h
new file mode 100644
--- /dev/null
+++ b/Engine/Plugins/VirtualProduction/RemoteControl/Source/RemoteControlLogic/Public/Behaviour/Builtin/Path/RCBehaviorSetAssetByPathNodeNew.h
@@ -0,0 +1,28 @@
+// Copyright Epic Games, Inc. All Rights Reserved.
+
+#pragma once
+
+#include "Behaviour/RCBehaviourNode.h"
+#include "RCBehaviorSetAssetByPathNodeNew.generated.h"
+
+/**
+ * Takes the string as a path and goes on to search for the Asset it is connected to, setting it on the list of bound entities.
+ */
+UCLASS(MinimalAPI)
+class URCBehaviorSetAssetByPathNodeNew : public URCBehaviourNode
+{
+	GENERATED_BODY()
+	
+public:
+	REMOTECONTROLLOGIC_API URCBehaviorSetAssetByPathNodeNew();
+	
+	//~ Begin URCBehaviourNode interface
+	UFUNCTION(BlueprintNativeEvent, BlueprintCallable, Category = "Remote Control Behaviour")
+	REMOTECONTROLLOGIC_API bool Execute(URCBehaviour* InBehavior) const override;
+
+	UFUNCTION(BlueprintNativeEvent, BlueprintCallable, Category = "Remote Control Behaviour")
+	REMOTECONTROLLOGIC_API bool IsSupported(URCBehaviour* InBehavior) const override;
+
+	REMOTECONTROLLOGIC_API UClass* GetBehaviourClass() const override;
+	//~ End URCBehaviourNode interface
+};
diff --git a/Engine/Plugins/VirtualProduction/RemoteControl/Source/RemoteControlLogic/Public/Behaviour/Builtin/Path/RCSetAssetByPathBehaviorNew.h b/Engine/Plugins/VirtualProduction/RemoteControl/Source/RemoteControlLogic/Public/Behaviour/Builtin/Path/RCSetAssetByPathBehaviorNew.h
new file mode 100644
--- /dev/null
+++ b/Engine/Plugins/VirtualProduction/RemoteControl/Source/RemoteControlLogic/Public/Behaviour/Builtin/Path/RCSetAssetByPathBehaviorNew.h
@@ -0,0 +1,108 @@
+// Copyright Epic Games, Inc. All Rights Reserved.
+
+#pragma once
+
+#include "Behaviour/RCBehaviour.h"
+#include "Containers/UnrealString.h"
+#include "HttpModule.h"
+#include "RCVirtualProperty.h"
+#include "RemoteControlEntity.h"
+
+#include "RCSetAssetByPathBehaviorNew.generated.h"
+
+class URCSetAssetByPathActionNew;
+class URCUserDefinedStruct;
+class URCVirtualPropertyContainerBase;
+class UTexture;
+struct FRemoteControlProperty;
+
+namespace UE::RemoteControl::Logic
+{
+	namespace SetAssetByPathNew
+	{
+		const FString ContentFolder = FString(TEXT("/Game/"));
+		const FName SetAssetByPathBehaviour = FName(TEXT("Set Asset By Path"));
+	}
+}
+
+USTRUCT()
+struct FRCAssetPathElementNew
+{
+	GENERATED_BODY()
+
+	FRCAssetPathElementNew()
+		: bIsInput(false)
+		, Path(TEXT(""))
+	{}
+
+	FRCAssetPathElementNew(bool bInIsInput, const FString& InPath)
+		: bIsInput(bInIsInput)
+		, Path(InPath)
+	{}
+
+	UPROPERTY(EditAnywhere, Category="Path")
+	bool bIsInput;
+
+	UPROPERTY(EditAnywhere, Category="Path")
+	FString Path;
+};
+
+/** Struct to help generate Widgets for the DetailsPanel of the Behaviour */
+USTRUCT()
+struct FRCSetAssetPathNew
+{
+	GENERATED_BODY()
+
+	FRCSetAssetPathNew()
+	{
+		// Add Empty index 0 when creating
+		AssetPath.AddDefaulted();
+	}
+	~FRCSetAssetPathNew() = default;
+	FRCSetAssetPathNew(const FRCSetAssetPathNew&) = default;
+	FRCSetAssetPathNew(FRCSetAssetPathNew&&) = default;
+	FRCSetAssetPathNew& operator=(const FRCSetAssetPathNew&) = default;
+	FRCSetAssetPathNew& operator=(FRCSetAssetPathNew&&) = default;
+
+	/** An Array of Strings holding the Path of an Asset, seperated in several String. Will concatenated back together later. */
+	UPROPERTY(EditAnywhere, Category="Path")
+	TArray<FRCAssetPathElementNew> AssetPath;
+};
+
+/**
+ * Custom behaviour for Set Asset By Path
+ */
+UCLASS(MinimalAPI)
+class URCSetAssetByPathBehaviorNew : public URCBehaviour
+{
+	GENERATED_BODY()
+
+public:
+	//~ Begin URCBehaviour interface
+	/** Add a Logic action using a remote control field as input */
+	REMOTECONTROLLOGIC_API virtual URCAction* AddAction(const TSharedRef<const FRemoteControlField> InRemoteControlField) override;
+	REMOTECONTROLLOGIC_API virtual bool SupportPropertyId() const override { return false; }
+
+#if WITH_EDITOR
+	REMOTECONTROLLOGIC_API virtual bool CanHaveActionForField(const TSharedPtr<FRemoteControlField> InRemoteControlField) const override;
+#endif
+	//~ End URCBehaviour interface
+
+	/** Returns the Path as the completed version, including Token Inputs. */
+	REMOTECONTROLLOGIC_API FString GetCurrentPath() const;
+
+	/** Called whenever a change has occured in the Slate */
+	REMOTECONTROLLOGIC_API void RefreshPathArray();
+
+	/** Add an action specifically for Set Asset By Path Behaviour */
+	REMOTECONTROLLOGIC_API URCSetAssetByPathActionNew* AddPropertyBindAction(const TSharedRef<const FRemoteControlProperty> InRemoteControlProperty);
+	
+public:
+	/** Struct holding the Path Information to help load Assets. */
+	UPROPERTY()
+	FRCSetAssetPathNew PathStruct;
+
+	/** Bool used to help tell if the path given is an internal or external one. */
+	UPROPERTY() 
+	bool bInternal = true;
+};
diff --git a/Engine/Plugins/VirtualProduction/RemoteControl/Source/RemoteControlUI/Private/RemoteControlUIModule.cpp b/Engine/Plugins/VirtualProduction/RemoteControl/Source/RemoteControlUI/Private/RemoteControlUIModule.cpp
--- a/Engine/Plugins/VirtualProduction/RemoteControl/Source/RemoteControlUI/Private/RemoteControlUIModule.cpp
+++ b/Engine/Plugins/VirtualProduction/RemoteControl/Source/RemoteControlUI/Private/RemoteControlUIModule.cpp
@@ -5,6 +5,7 @@
 #include "AssetEditor/RemoteControlPresetEditorToolkit.h"
 #include "AssetTools/RemoteControlPresetActions.h"
 #include "AssetToolsModule.h"
+#include "Behaviour/Builtin/Path/RCSetAssetByPathBehaviorNew.h"
 #include "Behaviour/Builtin/Path/RCSetAssetByPathBehaviour.h"
 #include "Commands/RemoteControlCommands.h"
 #include "Components/SkeletalMeshComponent.h"
@@ -1123,6 +1124,7 @@ void FRemoteControlUIModule::RegisterStructCustomizations()
 	}
 
 	PropertyEditorModule.RegisterCustomPropertyTypeLayout(FRCAssetPathElement::StaticStruct()->GetFName(), FOnGetPropertyTypeCustomizationInstance::CreateStatic(&FRCAssetPathElementCustomization::MakeInstance));
+	PropertyEditorModule.RegisterCustomPropertyTypeLayout(FRCAssetPathElementNew::StaticStruct()->GetFName(), FOnGetPropertyTypeCustomizationInstance::CreateStatic(&FRCAssetPathElementCustomization::MakeInstance));
 	PropertyEditorModule.RegisterCustomPropertyTypeLayout(FRCNetworkAddress::StaticStruct()->GetFName(), FOnGetPropertyTypeCustomizationInstance::CreateStatic(&FNetworkAddressCustomization::MakeInstance));
 	PropertyEditorModule.RegisterCustomPropertyTypeLayout(FRCPassphrase::StaticStruct()->GetFName(), FOnGetPropertyTypeCustomizationInstance::CreateStatic(&FPassphraseCustomization::MakeInstance));
 }
diff --git a/Engine/Plugins/VirtualProduction/RemoteControl/Source/RemoteControlUI/Private/UI/Action/Path/RCActionSetAssetByPathModelNew.cpp b/Engine/Plugins/VirtualProduction/RemoteControl/Source/RemoteControlUI/Private/UI/Action/Path/RCActionSetAssetByPathModelNew.cpp
new file mode 100644
--- /dev/null
+++ b/Engine/Plugins/VirtualProduction/RemoteControl/Source/RemoteControlUI/Private/UI/Action/Path/RCActionSetAssetByPathModelNew.cpp
@@ -0,0 +1,115 @@
+// Copyright Epic Games, Inc. All Rights Reserved.
+
+#include "UI/Action/Path/RCActionSetAssetByPathModelNew.h"
+
+#include "Action/Path/RCSetAssetByPathActionNew.h"
+#include "RemoteControlField.h"
+#include "ScopedTransaction.h"
+#include "Styling/RemoteControlStyles.h"
+#include "UI/RCUIHelpers.h"
+#include "UI/RemoteControlPanelStyle.h"
+#include "Widgets/Input/SEditableTextBox.h"
+#include "Widgets/Views/SListView.h"
+#include "Widgets/Views/STableRow.h"
+
+#define LOCTEXT_NAMESPACE "RCActionSetAssetByPathModelNew"
+
+namespace UE::RemoteControl::UI::Private::ActionSetAssetByPathModelNew
+{
+	namespace Columns
+	{
+		const FName VariableColor = TEXT("VariableColor");
+		const FName Description = TEXT("Description");
+	}
+
+	class SActionSetAssetByPathModelNewRow : public SMultiColumnTableRow<TSharedRef<FRCActionModel>>
+	{
+	public:
+		void Construct(const FTableRowArgs& InArgs, const TSharedRef<STableViewBase>& InOwnerTableView, TSharedRef<FRCActionSetAssetByPathModelNew> InActionItem)
+		{
+			ActionItem = InActionItem;
+			FSuperRowType::Construct(InArgs, InOwnerTableView);
+		}
+
+		TSharedRef<SWidget> GenerateWidgetForColumn(const FName& ColumnName) override
+		{
+			if (!ensure(ActionItem.IsValid()))
+			{
+				return SNullWidget::NullWidget;
+			}
+
+			using namespace UE::RemoteControl::UI::Private::ActionSetAssetByPathModelNew;
+		
+			if (ColumnName == Columns::VariableColor)
+			{
+				return ActionItem->GetTypeColorTagWidget();
+			}
+			else if (ColumnName == Columns::Description)
+			{
+				return ActionItem->GetNameWidget();
+			}
+
+			return SNullWidget::NullWidget;
+		}
+
+	private:
+		TSharedPtr<FRCActionSetAssetByPathModelNew> ActionItem;
+	};
+}
+
+TSharedRef<ITableRow> FRCActionSetAssetByPathModelNew::OnGenerateWidgetForList(TSharedPtr<FRCActionSetAssetByPathModelNew> InItem, const TSharedRef<STableViewBase>& OwnerTable)
+{
+	using SActionRow = UE::RemoteControl::UI::Private::ActionSetAssetByPathModelNew::SActionSetAssetByPathModelNewRow;
+
+	return SNew(SActionRow, OwnerTable, InItem.ToSharedRef())
+		.Style(&RCPanelStyle->TableRowStyle)
+		.Padding(FMargin(3.f));
+}
+
+TSharedPtr<SHeaderRow> FRCActionSetAssetByPathModelNew::GetHeaderRow()
+{
+	using namespace UE::RemoteControl::UI::Private;
+
+	const FRCPanelStyle* RCPanelStyle = &FRemoteControlPanelStyle::Get()->GetWidgetStyle<FRCPanelStyle>("RemoteControlPanel.MinorPanel");
+
+	return SNew(SHeaderRow)
+		.Style(&RCPanelStyle->HeaderRowStyle)
+
+		+ SHeaderRow::Column(ActionSetAssetByPathModelNew::Columns::VariableColor)
+		.DefaultLabel(FText())
+		.FixedWidth(5.f)
+		.HeaderContentPadding(RCPanelStyle->HeaderRowPadding)
+
+		+ SHeaderRow::Column(ActionSetAssetByPathModelNew::Columns::Description)
+		.DefaultLabel(LOCTEXT("RCActionDescColumnHeader", "Description"))
+		.FillWidth(0.5f)
+		.HeaderContentPadding(RCPanelStyle->HeaderRowPadding);
+}
+
+TSharedPtr<FRCActionSetAssetByPathModelNew> FRCActionSetAssetByPathModelNew::GetModelByActionType(URCAction* InAction, const TSharedPtr<class FRCBehaviourModel> InBehaviourItem, const TSharedPtr<SRemoteControlPanel> InRemoteControlPanel)
+{
+	if (URCSetAssetByPathActionNew* PropertyBindAction = Cast<URCSetAssetByPathActionNew>(InAction))
+	{
+		return MakeShared<FRCPropertyActionSetAssetByPathModelNew>(PropertyBindAction, InBehaviourItem, InRemoteControlPanel);
+	}
+
+	return nullptr;
+}
+
+FLinearColor FRCPropertyActionSetAssetByPathModelNew::GetActionTypeColor() const
+{
+	if (const URCSetAssetByPathActionNew* PropertyAction = Cast<URCSetAssetByPathActionNew>(ActionWeakPtr.Get()))
+	{
+		if (TSharedPtr<FRemoteControlProperty> RemoteControlProperty = PropertyAction->GetRemoteControlProperty())
+		{
+			if (RemoteControlProperty->IsBound())
+			{
+				return UE::RCUIHelpers::GetFieldClassTypeColor(RemoteControlProperty->GetProperty());
+			}
+		}
+	}
+
+	return FLinearColor::White;
+}
+
+#undef LOCTEXT_NAMESPACE
diff --git a/Engine/Plugins/VirtualProduction/RemoteControl/Source/RemoteControlUI/Private/UI/Action/Path/RCActionSetAssetByPathModelNew.h b/Engine/Plugins/VirtualProduction/RemoteControl/Source/RemoteControlUI/Private/UI/Action/Path/RCActionSetAssetByPathModelNew.h
new file mode 100644
--- /dev/null
+++ b/Engine/Plugins/VirtualProduction/RemoteControl/Source/RemoteControlUI/Private/UI/Action/Path/RCActionSetAssetByPathModelNew.h
@@ -0,0 +1,51 @@
+// Copyright Epic Games, Inc. All Rights Reserved.
+
+#pragma once
+
+#include "Action/Path/RCSetAssetByPathActionNew.h"
+#include "Action/RCPropertyIdAction.h"
+#include "UI/Action/RCActionModel.h"
+
+class SActionItemListRow;
+class SBox;
+class URCSetAssetByPathActionNew;
+
+/* 
+* ~ FRCActionSetAssetByPathModelNew ~
+*
+* UI model for representing an Action of the Set Asset (new)
+*/
+class FRCActionSetAssetByPathModelNew : public FRCActionModel
+{
+public:
+	FRCActionSetAssetByPathModelNew(URCAction* InAction, const TSharedPtr<class FRCBehaviourModel> InBehaviourItem, const TSharedPtr<SRemoteControlPanel> InRemoteControlPanel)
+		: FRCActionModel(InAction, InBehaviourItem, InRemoteControlPanel)
+	{
+	}
+
+	/** OnGenerateRow delegate for the Actions List View*/
+	TSharedRef<ITableRow> OnGenerateWidgetForList(TSharedPtr<FRCActionSetAssetByPathModelNew> InItem, const TSharedRef<STableViewBase>& InOwnerTable);
+
+	/** Get the Header Row appropriate for this particular Action model */
+	static TSharedPtr<SHeaderRow> GetHeaderRow();
+
+	/** Chooses the appropriate Action model for the current class and field type*/
+	static TSharedPtr<FRCActionSetAssetByPathModelNew> GetModelByActionType(URCAction* InAction, const TSharedPtr<class FRCBehaviourModel> InBehaviourItem, const TSharedPtr<SRemoteControlPanel> InRemoteControlPanel);
+};
+
+/*
+* ~ FRCPropertyActionSetAssetByPathModelNew ~
+*
+* UI model for representing a Property Action of the Set Asset (new)
+*/
+class FRCPropertyActionSetAssetByPathModelNew : public FRCActionSetAssetByPathModelNew
+{
+public:
+	FRCPropertyActionSetAssetByPathModelNew(URCSetAssetByPathActionNew* InPropertyAction, const TSharedPtr<class FRCBehaviourModel> InBehaviourItem, const TSharedPtr<SRemoteControlPanel> InRemoteControlPanel)
+		: FRCActionSetAssetByPathModelNew(InPropertyAction, InBehaviourItem, InRemoteControlPanel)
+	{
+	}
+
+	/** Color code for this Action*/
+	virtual FLinearColor GetActionTypeColor() const override;
+};
diff --git a/Engine/Plugins/VirtualProduction/RemoteControl/Source/RemoteControlUI/Private/UI/Behaviour/Builtin/Path/RCBehaviorSetAssetByPathModelNew.cpp b/Engine/Plugins/VirtualProduction/RemoteControl/Source/RemoteControlUI/Private/UI/Behaviour/Builtin/Path/RCBehaviorSetAssetByPathModelNew.cpp
new file mode 100644
--- /dev/null
+++ b/Engine/Plugins/VirtualProduction/RemoteControl/Source/RemoteControlUI/Private/UI/Behaviour/Builtin/Path/RCBehaviorSetAssetByPathModelNew.cpp
@@ -0,0 +1,310 @@
+// Copyright Epic Games, Inc. All Rights Reserved.
+
+#include "UI/Behaviour/Builtin/Path/RCBehaviorSetAssetByPathModelNew.h"
+
+#include "Behaviour/Builtin/Path/RCSetAssetByPathBehaviorNew.h"
+#include "Controller/RCController.h"
+#include "IDetailTreeNode.h"
+#include "IPropertyRowGenerator.h"
+#include "Modules/ModuleManager.h"
+#include "PropertyEditorModule.h"
+#include "RemoteControlPreset.h"
+#include "UI/Action/Path/RCActionSetAssetByPathModelNew.h"
+#include "UI/Action/SRCActionPanelList.h"
+#include "UI/Behaviour/Builtin/Path/SRCBehaviorSetAssetByPathNew.h"
+#include "UI/Controller/SRCControllerPanel.h"
+#include "UI/SRemoteControlPanel.h"
+#include "Widgets/SNullWidget.h"
+
+#define LOCTEXT_NAMESPACE "RCSetAssetByPathBehaviorModelNew"
+
+FRCSetAssetByPathBehaviorModelNew::FRCSetAssetByPathBehaviorModelNew(URCSetAssetByPathBehaviorNew* InSetAssetByPathBehavior, const TSharedPtr<SRemoteControlPanel> InRemoteControlPanel)
+	: FRCBehaviourModel(InSetAssetByPathBehavior, InRemoteControlPanel)
+{
+	SetAssetByPathBehaviorWeakPtr = InSetAssetByPathBehavior;
+
+	PreviewPathWidget = SNew(STextBlock)
+		.AutoWrapText(true);
+	PathArrayWidget = SNew(SBox);
+
+	FPropertyRowGeneratorArgs ArgsRowGeneratorArray;
+	ArgsRowGeneratorArray.NotifyHook = this;
+	PropertyRowGeneratorArray = FModuleManager::GetModuleChecked<FPropertyEditorModule>("PropertyEditor").CreatePropertyRowGenerator(ArgsRowGeneratorArray);
+
+	if (InSetAssetByPathBehavior)
+	{
+		PropertyRowGeneratorArray->SetStructure(MakeShareable(new FStructOnScope(FRCSetAssetPathNew::StaticStruct(), (uint8*)&InSetAssetByPathBehavior->PathStruct)));
+
+		PropertyRowGeneratorArray->OnFinishedChangingProperties().AddLambda([this](const FPropertyChangedEvent& InEvent)
+			{
+				OnAssetPathFinishedChangingProperties(InEvent);
+			});
+
+		PropertyRowGeneratorArray->OnRowsRefreshed().AddLambda([this]()
+			{
+				RefreshPathAndPreview();
+			});
+
+		DetailTreeNodeWeakPtrArray.Empty();
+
+		for (const TSharedRef<IDetailTreeNode>& CategoryNode : PropertyRowGeneratorArray->GetRootTreeNodes())
+		{
+			TArray<TSharedRef<IDetailTreeNode>> Children;
+			CategoryNode->GetChildren(Children);
+
+			for (const TSharedRef<IDetailTreeNode>& Child : Children)
+			{
+				DetailTreeNodeWeakPtrArray.Add(Child);
+			}
+		}
+	}
+
+	RefreshPathAndPreview();
+}
+
+bool FRCSetAssetByPathBehaviorModelNew::HasBehaviourDetailsWidget()
+{
+	return true;
+}
+
+TSharedRef<SWidget> FRCSetAssetByPathBehaviorModelNew::GetBehaviourDetailsWidget()
+{
+	return SNew(SRCBehaviorSetAssetByPathNew, SharedThis(this));
+}
+
+TSharedRef<SWidget> FRCSetAssetByPathBehaviorModelNew::GetPropertyWidget()
+{
+	TSharedRef<SVerticalBox> FieldWidget = SNew(SVerticalBox);
+
+	FieldWidget->AddSlot()
+		.Padding(FMargin(3.f, 3.f))
+		.AutoHeight()
+		[
+			SNew(SBox)
+			.MaxDesiredHeight(150.f)
+			[
+				SNew(SScrollBox)
+				+ SScrollBox::Slot()
+				[
+					PathArrayWidget->AsShared()
+				]
+			]
+		];
+
+	FieldWidget->AddSlot()
+		.Padding(FMargin(3.f, 3.f))
+		.AutoHeight()
+		[
+			SNew(SHorizontalBox)
+			+ SHorizontalBox::Slot()
+			.AutoWidth()
+			[
+				SNew(STextBlock)
+				.Text(FText::FromString("Preview Path: "))
+			]
+			+ SHorizontalBox::Slot()
+			[
+				PreviewPathWidget->AsShared()
+			]
+		];
+
+	
+	return 	SNew(SHorizontalBox)
+		.Clipping(EWidgetClipping::OnDemand)
+		+ SHorizontalBox::Slot()
+		.VAlign(VAlign_Center)
+		.AutoWidth()
+		.FillWidth(1.0f)
+		[
+			FieldWidget
+		];
+}
+
+void FRCSetAssetByPathBehaviorModelNew::NotifyControllerValueChanged(TSharedPtr<FRCControllerModel> InControllerModel)
+{
+	// Update the preview when a controller change to always get the correct path
+	RefreshPreview();
+}
+
+TSharedPtr<SRCLogicPanelListBase> FRCSetAssetByPathBehaviorModelNew::GetActionsListWidget(TSharedRef<SRCActionPanel> InActionPanel)
+{
+	return SNew(SRCActionPanelList<FRCActionSetAssetByPathModelNew>, InActionPanel, SharedThis(this));
+}
+
+void FRCSetAssetByPathBehaviorModelNew::RefreshPreview() const
+{
+	FString PreviewText = TEXT("Path is invalid!");
+
+	if (URCSetAssetByPathBehaviorNew* SetAssetByPathBehavior = SetAssetByPathBehaviorWeakPtr.Get())
+	{
+		PreviewText = SetAssetByPathBehavior->GetCurrentPath();
+	}
+
+	PreviewPathWidget->SetText(FText::FromString(PreviewText));
+}
+
+void FRCSetAssetByPathBehaviorModelNew::NotifyPostChange(const FPropertyChangedEvent& InPropertyChangedEvent, FEditPropertyChain* InPropertyThatChanged)
+{
+	if (URCSetAssetByPathBehaviorNew* SetAssetByPath = SetAssetByPathBehaviorWeakPtr.Get())
+	{
+		SetAssetByPath->Execute();
+	}
+}
+
+void FRCSetAssetByPathBehaviorModelNew::OnAssetPathFinishedChangingProperties(const FPropertyChangedEvent& InEvent)
+{
+	if (InEvent.ChangeType == EPropertyChangeType::ValueSet)
+	{
+		// if the actual PropertyName is the AssetPath than it means that we want to add a controller since we force this in the Customization
+		if (InEvent.GetPropertyName() == GET_MEMBER_NAME_CHECKED(FRCSetAssetPathNew, AssetPath))
+		{
+			const int32 Index = InEvent.GetArrayIndex(InEvent.GetMemberPropertyName().ToString());
+
+			if (Index == INDEX_NONE)
+			{
+				return;
+			}
+
+			URCSetAssetByPathBehaviorNew* SetAssetByPathBehavior = SetAssetByPathBehaviorWeakPtr.Get();
+
+			if (!SetAssetByPathBehavior)
+			{
+				return;
+			}
+
+			TArray<FRCAssetPathElementNew>& AssetPath = SetAssetByPathBehavior->PathStruct.AssetPath;
+
+			if (!AssetPath.IsValidIndex(Index))
+			{
+				return;
+			}
+
+			FRCAssetPathElementNew& CurrentElement = AssetPath[Index];
+
+			if (const TSharedPtr<SRemoteControlPanel>& RCPanel = GetRemoteControlPanel())
+			{
+				if (URemoteControlPreset* Preset = RCPanel->GetPreset())
+				{
+					URCVirtualPropertyInContainer* NewController = Preset->AddController(URCController::StaticClass(), EPropertyBagPropertyType::String);
+					
+					if (NewController)
+					{
+						NewController->DisplayIndex = Preset->GetNextControllerDisplayIndex();
+
+						FName NewControllerName = TEXT("DefaultInput");
+						if (!CurrentElement.Path.IsEmpty())
+						{
+							NewControllerName = FName(CurrentElement.Path);
+						}
+
+						const FName NewUniqueName = Preset->SetControllerDisplayName(NewController->Id, NewControllerName);
+						if (CurrentElement.Path != NewUniqueName.ToString())
+						{
+							// Assign to the path the actual new name of the controller
+							CurrentElement.Path = NewUniqueName.ToString();
+						}
+
+						if (!Preset->NewControllerCategory.IsEmpty())
+						{
+							using namespace UE::RemoteControl::UI::Private;
+
+							NewController->SetMetadataValue(FRCControllerPropertyInfo::CategoryName, Preset->NewControllerCategory);
+						}
+
+						RCPanel->OnControllerAdded.Broadcast(NewController->PropertyName);
+					}
+				}
+			}
+		}
+		// Just refresh the preview without reconstructing the widget if we just set values
+		RefreshPreview();
+	}
+	else
+	{
+		RefreshPathAndPreview();
+	}
+}
+
+void FRCSetAssetByPathBehaviorModelNew::RegenerateWeakPtrInternal()
+{
+	DetailTreeNodeWeakPtrArray.Empty();
+
+	for (const TSharedRef<IDetailTreeNode>& CategoryNode : PropertyRowGeneratorArray->GetRootTreeNodes())
+	{
+		TArray<TSharedRef<IDetailTreeNode>> Children;
+		CategoryNode->GetChildren(Children);
+
+		for (const TSharedRef<IDetailTreeNode>& Child : Children)
+		{
+			DetailTreeNodeWeakPtrArray.Add(Child);
+		}
+	}
+}
+
+void FRCSetAssetByPathBehaviorModelNew::RefreshPathAndPreview()
+{
+	const TSharedPtr<SVerticalBox> FieldArrayWidget = SNew(SVerticalBox);
+
+	RegenerateWeakPtrInternal();
+
+	if (URCSetAssetByPathBehaviorNew* Behavior = Cast<URCSetAssetByPathBehaviorNew>(GetBehaviour()))
+	{
+		Behavior->RefreshPathArray();
+	}
+	
+	for (const TSharedPtr<IDetailTreeNode>& DetailTreeNodeArray : DetailTreeNodeWeakPtrArray)
+	{
+		const TSharedPtr<IDetailTreeNode> PinnedNode = DetailTreeNodeArray;
+
+		TArray<TSharedRef<IDetailTreeNode>> Children;
+		PinnedNode->GetChildren(Children);
+		Children.Insert(PinnedNode.ToSharedRef(), 0);
+		
+		for (uint8 Counter = 0; Counter < Children.Num(); Counter++)
+		{
+			const FNodeWidgets NodeWidgets = Children[Counter]->CreateNodeWidgets();
+
+			if (NodeWidgets.ValueWidget)
+			{
+				TSharedPtr<SHorizontalBox> CurrentHorizontalBox = SNew(SHorizontalBox);
+
+				if (Counter == 0)
+				{
+					// Add NameWidget only for the ArrayProperty not for its entries
+					CurrentHorizontalBox->AddSlot()
+					[
+						NodeWidgets.NameWidget.ToSharedRef()
+					];
+				}
+
+				CurrentHorizontalBox->AddSlot()
+					.FillWidth(1.f)
+					[
+						NodeWidgets.ValueWidget.ToSharedRef()
+					];
+
+				FieldArrayWidget->AddSlot()
+					.Padding(FMargin(3.0f, 2.0f))
+					.AutoHeight()
+					.VAlign(VAlign_Center)
+					[
+						CurrentHorizontalBox.ToSharedRef()
+					];
+			}
+			else if (NodeWidgets.WholeRowWidget)
+			{
+				FieldArrayWidget->AddSlot()
+					.Padding(FMargin(3.0f, 2.0f))
+					.AutoHeight()
+					[
+						NodeWidgets.WholeRowWidget.ToSharedRef()
+					];
+			}
+		}
+	}
+
+	RefreshPreview();
+	PathArrayWidget->SetContent(FieldArrayWidget.ToSharedRef());
+}
+
+#undef LOCTEXT_NAMESPACE
diff --git a/Engine/Plugins/VirtualProduction/RemoteControl/Source/RemoteControlUI/Private/UI/Behaviour/Builtin/Path/RCBehaviorSetAssetByPathModelNew.h b/Engine/Plugins/VirtualProduction/RemoteControl/Source/RemoteControlUI/Private/UI/Behaviour/Builtin/Path/RCBehaviorSetAssetByPathModelNew.h
new file mode 100644
--- /dev/null
+++ b/Engine/Plugins/VirtualProduction/RemoteControl/Source/RemoteControlUI/Private/UI/Behaviour/Builtin/Path/RCBehaviorSetAssetByPathModelNew.h
@@ -0,0 +1,74 @@
+// Copyright Epic Games, Inc. All Rights Reserved.
+
+#pragma once
+
+#include "Misc/NotifyHook.h"
+#include "RemoteControlEntity.h"
+#include "UI/Behaviour/RCBehaviourModel.h"
+
+class IDetailTreeNode;
+class SBox;
+class URCSetAssetByPathBehaviorNew;
+
+/*
+ * ~ FRCSetAssetByPathBehaviorModelNew ~
+ *
+ * Child Behaviour class representing the "Set Asset By Path" Behaviour's UI model.
+ *
+ * Generates several Widgets where users can enter the RootPath, Target Property, and Default Property as FStrings.
+ * The values are then put together, getting a path towards a possible object.
+ */
+class FRCSetAssetByPathBehaviorModelNew : public FRCBehaviourModel, public FNotifyHook
+{
+public:
+	explicit FRCSetAssetByPathBehaviorModelNew(URCSetAssetByPathBehaviorNew* InSetAssetByPathBehavior, const TSharedPtr<SRemoteControlPanel> InRemoteControlPanel);
+
+	/** Returns true if this behavior have a details widget or false if not*/
+	virtual bool HasBehaviourDetailsWidget() override;
+
+	/** Builds a Behavior specific widget as required for the Set Asset By Path Behavior */
+	virtual TSharedRef<SWidget> GetBehaviourDetailsWidget() override;
+
+	/*
+	 * Builds the Property Details Widget, including a generic expendable Array Widget and
+	 * further two Text Widgets representing elements needed for the Path Behavior.
+	 * All of them store the user input and use them to perform the SetAssetByPath Behavior.
+	 */
+	TSharedRef<SWidget> GetPropertyWidget();
+
+	/** Called when a controller value changed */
+	virtual void NotifyControllerValueChanged(TSharedPtr<FRCControllerModel> InControllerModel) override;
+
+	virtual TSharedPtr<SRCLogicPanelListBase> GetActionsListWidget(TSharedRef<SRCActionPanel> InActionPanel) override;
+
+	/** Refresh the PreviewPath */
+	void RefreshPreview() const;
+
+	virtual void NotifyPostChange(const FPropertyChangedEvent& InPropertyChangedEvent, class FEditPropertyChain* InPropertyThatChanged) override;
+
+private:
+	/** Callback called when the AssetPath properties change to either refresh or reconstruct it updating the path values */
+	void OnAssetPathFinishedChangingProperties(const FPropertyChangedEvent& InEvent);
+
+	/** Reconstruct the PropertyRow Generator */
+	void RegenerateWeakPtrInternal();
+	
+	/** Regenerates and creates a new PathArray Widget if changed */
+	void RefreshPathAndPreview();
+
+private:
+	/** The SetAssetByPath Behaviour associated with this Model */
+	TWeakObjectPtr<URCSetAssetByPathBehaviorNew> SetAssetByPathBehaviorWeakPtr;
+
+	/** Pointer to the Widget holding the PathArray created for this behavior */
+	TSharedPtr<SBox> PathArrayWidget;
+
+	/** Pointer to the Preview Text Widget for the Path */
+	TSharedPtr<STextBlock> PreviewPathWidget;
+
+	/** The row generator used to build the generic array value widget for the Path Array */
+	TSharedPtr<IPropertyRowGenerator> PropertyRowGeneratorArray;
+
+	/** Used to create a generic Value Widget based on the Paths Available*/
+	TArray<TSharedPtr<IDetailTreeNode>> DetailTreeNodeWeakPtrArray;
+};
diff --git a/Engine/Plugins/VirtualProduction/RemoteControl/Source/RemoteControlUI/Private/UI/Behaviour/Builtin/Path/SRCBehaviorSetAssetByPathNew.cpp b/Engine/Plugins/VirtualProduction/RemoteControl/Source/RemoteControlUI/Private/UI/Behaviour/Builtin/Path/SRCBehaviorSetAssetByPathNew.cpp
new file mode 100644
--- /dev/null
+++ b/Engine/Plugins/VirtualProduction/RemoteControl/Source/RemoteControlUI/Private/UI/Behaviour/Builtin/Path/SRCBehaviorSetAssetByPathNew.cpp
@@ -0,0 +1,131 @@
+// Copyright Epic Games, Inc. All Rights Reserved.
+
+#include "SRCBehaviorSetAssetByPathNew.h"
+
+#include "Behaviour/Builtin/Path/RCSetAssetByPathBehaviorNew.h"
+#include "IDetailTreeNode.h"
+#include "RemoteControlPreset.h"
+#include "SlateOptMacros.h"
+#include "SPositiveActionButton.h"
+#include "Styling/StyleColors.h"
+#include "UI/Behaviour/Builtin/Path/RCBehaviorSetAssetByPathModelNew.h"
+#include "UI/RemoteControlPanelStyle.h"
+#include "Widgets/Input/SButton.h"
+#include "Widgets/Text/STextBlock.h"
+
+#define LOCTEXT_NAMESPACE "SRCBehaviorSetAssetByPathNew"
+
+class SPositiveActionButton;
+
+BEGIN_SLATE_FUNCTION_BUILD_OPTIMIZATION
+
+void SRCBehaviorSetAssetByPathNew::Construct(const FArguments& InArgs, TSharedRef<FRCSetAssetByPathBehaviorModelNew> InBehaviourItem)
+{
+	SetAssetByPathWeakPtr = InBehaviourItem;
+	PathBehaviour = Cast<URCSetAssetByPathBehaviorNew>(InBehaviourItem->GetBehaviour());
+	if (!PathBehaviour.IsValid())
+	{
+		ChildSlot
+		[
+			SNullWidget::NullWidget
+		];
+		return;
+	}
+
+	const TSharedRef<SHorizontalBox> InternalExternalSwitchWidget = SNew(SHorizontalBox)
+
+		+ SHorizontalBox::Slot()
+		.HAlign(HAlign_Right)
+		.VAlign(VAlign_Center)
+		.Padding(FMargin(3.0f, 0.0f))
+		.FillWidth(1.0f)
+		.AutoWidth()
+		[
+			SNew(STextBlock)
+			.Text(LOCTEXT("AffectsTexturesOnly", "For Texture Sources:"))
+		]
+
+		+ SHorizontalBox::Slot()
+		.HAlign(HAlign_Right)
+		.VAlign(VAlign_Center)
+		.Padding(FMargin(3.0f,0.0f))
+		.FillWidth(1.0f)
+		.AutoWidth()
+		[
+			SNew(SButton)
+			.Text(FText::FromString(FString("Internal")))
+			.ButtonColorAndOpacity(this, &SRCBehaviorSetAssetByPathNew::GetInternalWidgetColor)
+			.OnPressed(this, &SRCBehaviorSetAssetByPathNew::OnInternalExternalSwitchWidgetPressed, true)
+		]
+	
+		+ SHorizontalBox::Slot()
+		.HAlign(HAlign_Right)
+		.VAlign(VAlign_Center)
+		.Padding(FMargin(3.0f,0.0f))
+		.FillWidth(1.0f)
+		.AutoWidth()
+		[
+			SNew(SButton)
+			.Text(FText::FromString(FString("External")))
+			.ButtonColorAndOpacity(this, &SRCBehaviorSetAssetByPathNew::GetExternalWidgetColor)
+			.OnPressed(this, &SRCBehaviorSetAssetByPathNew::OnInternalExternalSwitchWidgetPressed, false)
+		];
+	
+	ChildSlot
+	.Padding(8.f, 4.f)
+	[
+		SNew(SVerticalBox)
+		+ SVerticalBox::Slot()
+		[
+			InternalExternalSwitchWidget
+		]
+		+ SVerticalBox::Slot()
+		.AutoHeight()
+		[
+			InBehaviourItem->GetPropertyWidget()
+		]
+	];
+}
+
+void SRCBehaviorSetAssetByPathNew::OnInternalExternalSwitchWidgetPressed(bool bInIsInternal) const
+{
+	if (!PathBehaviour.IsValid())
+	{
+		return;
+	}
+
+	PathBehaviour->bInternal = bInIsInternal;
+
+	if (const TSharedPtr<const FRCSetAssetByPathBehaviorModelNew>& SetAssetPath = SetAssetByPathWeakPtr.Pin())
+	{
+		SetAssetPath->RefreshPreview();
+	}
+}
+
+FSlateColor SRCBehaviorSetAssetByPathNew::GetExternalWidgetColor() const
+{
+	if (!PathBehaviour.IsValid())
+	{
+		return FAppStyle::Get().GetSlateColor("Colors.AccentRed");
+	}
+
+	return PathBehaviour->bInternal 
+		? FAppStyle::Get().GetSlateColor("Colors.AccentWhite") 
+		: FAppStyle::Get().GetSlateColor("Colors.Highlight");
+}
+
+FSlateColor SRCBehaviorSetAssetByPathNew::GetInternalWidgetColor() const
+{
+	if (!PathBehaviour.IsValid())
+	{
+		return FAppStyle::Get().GetSlateColor("Colors.AccentRed");
+	}
+
+	return PathBehaviour->bInternal 
+		? FAppStyle::Get().GetSlateColor("Colors.Highlight") 
+		: FAppStyle::Get().GetSlateColor("Colors.AccentWhite");
+}
+
+END_SLATE_FUNCTION_BUILD_OPTIMIZATION
+
+#undef LOCTEXT_NAMESPACE
diff --git a/Engine/Plugins/VirtualProduction/RemoteControl/Source/RemoteControlUI/Private/UI/Behaviour/Builtin/Path/SRCBehaviorSetAssetByPathNew.h b/Engine/Plugins/VirtualProduction/RemoteControl/Source/RemoteControlUI/Private/UI/Behaviour/Builtin/Path/SRCBehaviorSetAssetByPathNew.h
new file mode 100644
--- /dev/null
+++ b/Engine/Plugins/VirtualProduction/RemoteControl/Source/RemoteControlUI/Private/UI/Behaviour/Builtin/Path/SRCBehaviorSetAssetByPathNew.h
@@ -0,0 +1,40 @@
+// Copyright Epic Games, Inc. All Rights Reserved.
+
+#pragma once
+
+#include "Styling/SlateColor.h"
+#include "Templates/SharedPointer.h"
+#include "UObject/WeakObjectPtrTemplates.h"
+#include "Widgets/SCompoundWidget.h"
+
+class FRCSetAssetByPathBehaviorModelNew;
+class UClass;
+class URCSetAssetByPathBehaviorNew;
+
+class SRCBehaviorSetAssetByPathNew : public  SCompoundWidget
+{
+public:
+	SLATE_BEGIN_ARGS(SRCBehaviorSetAssetByPathNew)
+		{}
+	SLATE_END_ARGS()
+
+	/** Constructs this widget with InArgs */
+	void Construct(const FArguments& InArgs, TSharedRef<FRCSetAssetByPathBehaviorModelNew> InBehaviourItem);
+
+private:
+	/** Called when the Internal/External button are pressed */
+	void OnInternalExternalSwitchWidgetPressed(bool bInIsInternal) const;
+
+	/** Return the color of the External widget */
+	FSlateColor GetExternalWidgetColor() const;
+
+	/** Return the color of the Internal widget */
+	FSlateColor GetInternalWidgetColor() const;
+
+private:
+	/** The Behaviour (UI model) associated with us */
+	TWeakPtr<const FRCSetAssetByPathBehaviorModelNew> SetAssetByPathWeakPtr;
+
+	/** The PathBehaviour associated with us */
+	TWeakObjectPtr<URCSetAssetByPathBehaviorNew> PathBehaviour;
+};
diff --git a/Engine/Plugins/VirtualProduction/RemoteControl/Source/RemoteControlUI/Private/UI/Behaviour/SRCBehaviourPanel.cpp b/Engine/Plugins/VirtualProduction/RemoteControl/Source/RemoteControlUI/Private/UI/Behaviour/SRCBehaviourPanel.cpp
--- a/Engine/Plugins/VirtualProduction/RemoteControl/Source/RemoteControlUI/Private/UI/Behaviour/SRCBehaviourPanel.cpp
+++ b/Engine/Plugins/VirtualProduction/RemoteControl/Source/RemoteControlUI/Private/UI/Behaviour/SRCBehaviourPanel.cpp
@@ -3,6 +3,7 @@
 #include "SRCBehaviourPanel.h"
 #include "SRCBehaviourPanelList.h"
 
+#include "Behaviour/Builtin/Path/RCBehaviourSetAssetByPathNode.h"
 #include "Behaviour/RCBehaviour.h"
 #include "Behaviour/RCBehaviourBlueprintNode.h"
 #include "Behaviour/RCBehaviourNode.h"
@@ -126,6 +127,11 @@ TSharedRef<SWidget> SRCBehaviourPanel::GetBehaviourMenuContentWidget()
 			{
 				UClass* Class = *It;
 
+				if (Class == URCBehaviourSetAssetByPathNode::StaticClass())
+				{
+					continue;
+				}
+
 				if (Class->IsChildOf(URCBehaviourNode::StaticClass()))
 				{
 					const bool bIsClassInstantiatable = Class->HasAnyClassFlags(CLASS_Deprecated | CLASS_NewerVersionExists | CLASS_Abstract) || FKismetEditorUtilities::IsClassABlueprintSkeleton(Class);
diff --git a/Engine/Plugins/VirtualProduction/RemoteControl/Source/RemoteControlUI/Private/UI/Behaviour/SRCBehaviourPanelList.cpp b/Engine/Plugins/VirtualProduction/RemoteControl/Source/RemoteControlUI/Private/UI/Behaviour/SRCBehaviourPanelList.cpp
--- a/Engine/Plugins/VirtualProduction/RemoteControl/Source/RemoteControlUI/Private/UI/Behaviour/SRCBehaviourPanelList.cpp
+++ b/Engine/Plugins/VirtualProduction/RemoteControl/Source/RemoteControlUI/Private/UI/Behaviour/SRCBehaviourPanelList.cpp
@@ -5,6 +5,7 @@
 #include "Behaviour/Builtin/Bind/RCBehaviourBind.h"
 #include "Behaviour/Builtin/Conditional/RCBehaviourConditional.h"
 #include "Behaviour/Builtin/Path/RCSetAssetByPathBehaviour.h"
+#include "Behaviour/Builtin/Path/RCSetAssetByPathBehaviorNew.h"
 #include "Controller/RCController.h"
 #include "RCBehaviourModel.h"
 #include "RemoteControlPreset.h"
@@ -14,6 +15,7 @@
 #include "UI/Action/SRCActionPanel.h"
 #include "UI/Behaviour/Builtin/Bind/RCBehaviourBindModel.h"
 #include "UI/Behaviour/Builtin/Conditional/RCBehaviourConditionalModel.h"
+#include "UI/Behaviour/Builtin/Path/RCBehaviorSetAssetByPathModelNew.h"
 #include "UI/Behaviour/Builtin/Path/RCBehaviourSetAssetByPathModel.h"
 #include "UI/Behaviour/Builtin/RangeMap/RCBehaviourRangeMapModel.h"
 #include "UI/Controller/RCControllerModel.h"
@@ -280,6 +282,10 @@ void SRCBehaviourPanelList::AddBehaviourToList(URCBehaviour* InBehaviour)
 	{
 		BehaviourItems.Add(MakeShared<FRCRangeMapBehaviourModel>(RangeMapBehaviour));
 	}
+	else if (URCSetAssetByPathBehaviorNew* SetAssetByPathBehavior = Cast<URCSetAssetByPathBehaviorNew>(InBehaviour))
+	{
+		BehaviourItems.Add(MakeShared<FRCSetAssetByPathBehaviorModelNew>(SetAssetByPathBehavior, RemoteControlPanel));
+	}
 	else
 	{
 		BehaviourItems.Add(MakeShared<FRCBehaviourModel>(InBehaviour, RemoteControlPanel));
diff --git a/Engine/Plugins/VirtualProduction/RemoteControl/Source/RemoteControlUI/Private/UI/Customizations/RCAssetPathElementCustomization.cpp b/Engine/Plugins/VirtualProduction/RemoteControl/Source/RemoteControlUI/Private/UI/Customizations/RCAssetPathElementCustomization.cpp
--- a/Engine/Plugins/VirtualProduction/RemoteControl/Source/RemoteControlUI/Private/UI/Customizations/RCAssetPathElementCustomization.cpp
+++ b/Engine/Plugins/VirtualProduction/RemoteControl/Source/RemoteControlUI/Private/UI/Customizations/RCAssetPathElementCustomization.cpp
@@ -136,13 +136,6 @@ void FRCAssetPathElementCustomization::OnPathChanged() const
 
 	bool bIsRCInput = false;
 	IsInputHandle->GetValue(bIsRCInput);
-
-	if (bIsRCInput)
-	{
-		return;
-	}
-
-	RemoveSlashFromPathEnd();
 }
 
 ECheckBoxState FRCAssetPathElementCustomization::IsChecked() const
@@ -168,11 +161,6 @@ void FRCAssetPathElementCustomization::OnCheckStateChanged(ECheckBoxState InNewS
 	}
 	const bool bIsInput = InNewState == ECheckBoxState::Checked;
 	IsInputHandle->SetValue(bIsInput);
-
-	if (!bIsInput)
-	{
-		RemoveSlashFromPathEnd();
-	}
 }
 
 FReply FRCAssetPathElementCustomization::OnGetAssetFromSelectionClicked() const
@@ -239,19 +227,4 @@ int32 FRCAssetPathElementCustomization::OnGetWidgetSwitcherIndex() const
 	return IsChecked() == ECheckBoxState::Checked ? 1 : 0;
 }
 
-void FRCAssetPathElementCustomization::RemoveSlashFromPathEnd() const
-{
-	FString CurrentPath;
-	PathHandle->GetValue(CurrentPath);
-
-	// remove all / that are at the end of the path
-	bool bRemoved = false;
-	do
-	{
-		bRemoved = CurrentPath.RemoveFromEnd(TEXT("/"), 1);
-	} while (bRemoved);
-
-	PathHandle->SetValue(CurrentPath, EPropertyValueSetFlags::NotTransactable);
-}
-
 #undef LOCTEXT_NAMESPACE
diff --git a/Engine/Plugins/VirtualProduction/RemoteControl/Source/RemoteControlUI/Private/UI/Customizations/RCAssetPathElementCustomization.h b/Engine/Plugins/VirtualProduction/RemoteControl/Source/RemoteControlUI/Private/UI/Customizations/RCAssetPathElementCustomization.h
--- a/Engine/Plugins/VirtualProduction/RemoteControl/Source/RemoteControlUI/Private/UI/Customizations/RCAssetPathElementCustomization.h
+++ b/Engine/Plugins/VirtualProduction/RemoteControl/Source/RemoteControlUI/Private/UI/Customizations/RCAssetPathElementCustomization.h
@@ -39,9 +39,6 @@ private:
 	/** Retrieve the current widget switcher index */
 	int32 OnGetWidgetSwitcherIndex() const;
 
-	/** Remove all / at the end of the path */
-	void RemoveSlashFromPathEnd() const;
-
 private:
 	TSharedPtr<SWidget> PathWidget;
 	TSharedPtr<IPropertyUtilities> PropertyUtilities;
