Remote Control: Deprecated Controller's Behaviors TSet to be private TArray with a const accessor.

diff --git a/Engine/Plugins/VirtualProduction/Avalanche/Source/AvalancheMedia/Private/Playable/AvaPlayableRemoteControl.cpp b/Engine/Plugins/VirtualProduction/Avalanche/Source/AvalancheMedia/Private/Playable/AvaPlayableRemoteControl.cpp
--- a/Engine/Plugins/VirtualProduction/Avalanche/Source/AvalancheMedia/Private/Playable/AvaPlayableRemoteControl.cpp
+++ b/Engine/Plugins/VirtualProduction/Avalanche/Source/AvalancheMedia/Private/Playable/AvaPlayableRemoteControl.cpp
@@ -338,8 +338,8 @@ namespace UE::AvaPlayableRemoteControl
 	{
 		if (URCController* const Controller = Cast<URCController>(InVirtualProperty))
 		{
-			PreviousBehavioursEnabled.Reserve(Controller->Behaviours.Num());			
-			for (URCBehaviour* const Behavior : Controller->Behaviours)
+			PreviousBehavioursEnabled.Reserve(Controller->GetBehaviors().Num());
+			for (URCBehaviour* const Behavior : Controller->GetBehaviors())
 			{
 				if (Behavior)
 				{
@@ -355,7 +355,7 @@ namespace UE::AvaPlayableRemoteControl
 		if (URCController* const Controller = Cast<URCController>(VirtualProperty))
 		{
 			int32 BehaviourIndex = 0;
-			for (URCBehaviour* const Behavior : Controller->Behaviours)
+			for (URCBehaviour* const Behavior : Controller->GetBehaviors())
 			{
 				if (Behavior)
 				{
diff --git a/Engine/Plugins/VirtualProduction/Avalanche/Source/AvalancheMediaEditor/Private/Rundown/DetailsView/RemoteControl/Controllers/SAvaRundownRCControllerPanel.cpp b/Engine/Plugins/VirtualProduction/Avalanche/Source/AvalancheMediaEditor/Private/Rundown/DetailsView/RemoteControl/Controllers/SAvaRundownRCControllerPanel.cpp
--- a/Engine/Plugins/VirtualProduction/Avalanche/Source/AvalancheMediaEditor/Private/Rundown/DetailsView/RemoteControl/Controllers/SAvaRundownRCControllerPanel.cpp
+++ b/Engine/Plugins/VirtualProduction/Avalanche/Source/AvalancheMediaEditor/Private/Rundown/DetailsView/RemoteControl/Controllers/SAvaRundownRCControllerPanel.cpp
@@ -205,7 +205,7 @@ void SAvaRundownRCControllerPanel::RefreshForManagedInstance(int32 InInstanceInd
 
 				if (URCController* const Controller = Cast<URCController>(VirtualProperty))
 				{
-					for (URCBehaviour* const Behavior : Controller->Behaviours)
+					for (URCBehaviour* const Behavior : Controller->GetBehaviors())
 					{
 						if (URCSetAssetByPathBehaviour* const AssetByPathBehavior = Cast<URCSetAssetByPathBehaviour>(Behavior))
 						{
diff --git a/Engine/Plugins/VirtualProduction/Avalanche/Source/AvalancheRemoteControl/Private/AvaRCLibrary.cpp b/Engine/Plugins/VirtualProduction/Avalanche/Source/AvalancheRemoteControl/Private/AvaRCLibrary.cpp
--- a/Engine/Plugins/VirtualProduction/Avalanche/Source/AvalancheRemoteControl/Private/AvaRCLibrary.cpp
+++ b/Engine/Plugins/VirtualProduction/Avalanche/Source/AvalancheRemoteControl/Private/AvaRCLibrary.cpp
@@ -41,7 +41,7 @@ TArray<AActor*> UAvaRCLibrary::GetControlledActors(UObject* InWorldContextObject
 
 	TArray<AActor*> ControlledActors;
 
-	for (URCBehaviour* Behavior : Controller->Behaviours)
+	for (URCBehaviour* Behavior : Controller->GetBehaviors())
 	{
 		if (!Behavior || !Behavior->ActionContainer)
 		{
diff --git a/Engine/Plugins/VirtualProduction/RemoteControl/Source/RemoteControlLogic/Private/Controller/RCController.cpp b/Engine/Plugins/VirtualProduction/RemoteControl/Source/RemoteControlLogic/Private/Controller/RCController.cpp
--- a/Engine/Plugins/VirtualProduction/RemoteControl/Source/RemoteControlLogic/Private/Controller/RCController.cpp
+++ b/Engine/Plugins/VirtualProduction/RemoteControl/Source/RemoteControlLogic/Private/Controller/RCController.cpp
@@ -12,7 +12,7 @@
 
 void URCController::UpdateEntityIds(const TMap<FGuid, FGuid>& InEntityIdMap)
 {
-	for (URCBehaviour* Behaviour : Behaviours)
+	for (URCBehaviour* Behaviour : Behaviors)
 	{
 		if (Behaviour)
 		{
@@ -42,7 +42,7 @@ URCBehaviour* URCController::AddBehaviour(TSubclassOf<URCBehaviourNode> InBehavi
 
 	NewBehaviour->Initialize();
 	
-	Behaviours.Add(NewBehaviour);
+	Behaviors.Add(NewBehaviour);
 
 	return NewBehaviour;
 }
@@ -80,14 +80,14 @@ URCBehaviour* URCController::CreateBehaviourWithoutCheck(TSubclassOf<URCBehaviou
 
 int32 URCController::RemoveBehaviour(URCBehaviour* InBehaviour)
 {
-	return Behaviours.Remove(InBehaviour);
+	return Behaviors.Remove(InBehaviour);
 }
 
 int32 URCController::RemoveBehaviour(const FGuid InBehaviourId)
 {
 	int32 RemovedCount = 0;
 
-	for (TSet<TObjectPtr<URCBehaviour>>::TIterator BehaviourIt(Behaviours); BehaviourIt; ++BehaviourIt)
+	for (TArray<TObjectPtr<URCBehaviour>>::TIterator BehaviourIt(Behaviors); BehaviourIt; ++BehaviourIt)
 	{
 		const URCBehaviour* Behaviour = *BehaviourIt;
 		if (Behaviour && Behaviour->Id == InBehaviourId)
@@ -102,17 +102,17 @@ int32 URCController::RemoveBehaviour(const FGuid InBehaviourId)
 
 bool URCController::HasBehavior(const URCBehaviour* InBehavior) const
 {
-	return InBehavior && Behaviours.Contains(InBehavior);
+	return InBehavior && Behaviors.Contains(InBehavior);
 }
 
 void URCController::EmptyBehaviours()
 {
-	Behaviours.Empty();
+	Behaviors.Empty();
 }
 
 void URCController::ExecuteBehaviours(const bool bIsPreChange/* = false*/)
 {
-	for (URCBehaviour* Behaviour : Behaviours)
+	for (URCBehaviour* Behaviour : Behaviors)
 	{
 		if (!Behaviour || !Behaviour->bIsEnabled)
 		{
@@ -146,7 +146,7 @@ URCBehaviour* URCController::DuplicateBehaviour(URCController* InController, URC
 
 	NewBehaviour->ControllerWeakPtr = InController;
 
-	InController->Behaviours.Add(NewBehaviour);
+	InController->Behaviors.Add(NewBehaviour);
 
 	return NewBehaviour;
 }
@@ -160,10 +160,28 @@ void URCController::UpdateEntityUsage(const URemoteControlPreset* InPreset, TMap
 		return;
 	}
 
-	for (URCBehaviour* Behaviour : Behaviours)
+	for (URCBehaviour* Behaviour : Behaviors)
 	{
-		Behaviour->UpdateEntityUsage(this, InEntityToBehavior);
+		if (Behaviour)
+		{
+			Behaviour->UpdateEntityUsage(this, InEntityToBehavior);
+		}
 	}
 }
 
-#undef LOCTEXT_NAMESPACE /* RCController */ 
\ No newline at end of file
+void URCController::PostLoad()
+{
+	Super::PostLoad();
+
+	PRAGMA_DISABLE_DEPRECATION_WARNINGS
+
+	if (!Behaviours.IsEmpty())
+	{
+		Behaviors = Behaviours.Array();
+		Behaviours.Empty();
+	}
+
+	PRAGMA_ENABLE_DEPRECATION_WARNINGS
+}
+
+#undef LOCTEXT_NAMESPACE /* RCController */ 
diff --git a/Engine/Plugins/VirtualProduction/RemoteControl/Source/RemoteControlLogic/Private/Controller/RCCustomControllerUtilities.cpp b/Engine/Plugins/VirtualProduction/RemoteControl/Source/RemoteControlLogic/Private/Controller/RCCustomControllerUtilities.cpp
--- a/Engine/Plugins/VirtualProduction/RemoteControl/Source/RemoteControlLogic/Private/Controller/RCCustomControllerUtilities.cpp
+++ b/Engine/Plugins/VirtualProduction/RemoteControl/Source/RemoteControlLogic/Private/Controller/RCCustomControllerUtilities.cpp
@@ -105,7 +105,7 @@ bool UE::RCCustomControllers::GetEntitiesControlledByController(const URemoteCon
 		return false;
 	}
 
-	for (const URCBehaviour* Behavior : Controller->Behaviours)
+	for (const URCBehaviour* Behavior : Controller->GetBehaviors())
 	{
 		if (!Behavior)
 		{
@@ -162,7 +162,7 @@ bool UE::RCCustomControllers::GetControllersOfEntity(const URemoteControlPreset*
 			{
 				bool bFoundEntity = false;
 
-				for (const URCBehaviour* Behavior : Controller->Behaviours)
+				for (const URCBehaviour* Behavior : Controller->GetBehaviors())
 				{
 					if (!Behavior)
 					{
diff --git a/Engine/Plugins/VirtualProduction/RemoteControl/Source/RemoteControlLogic/Private/Signature/RCSignatureControllerBindAction.cpp b/Engine/Plugins/VirtualProduction/RemoteControl/Source/RemoteControlLogic/Private/Signature/RCSignatureControllerBindAction.cpp
--- a/Engine/Plugins/VirtualProduction/RemoteControl/Source/RemoteControlLogic/Private/Signature/RCSignatureControllerBindAction.cpp
+++ b/Engine/Plugins/VirtualProduction/RemoteControl/Source/RemoteControlLogic/Private/Signature/RCSignatureControllerBindAction.cpp
@@ -87,7 +87,7 @@ bool FRCSignatureControllerBindAction::ExecuteControllerAction(const FRCSignatur
 
 	URCBehaviourBind* BindBehavior = nullptr;
 
-	for (URCBehaviour* Behavior : InController->Behaviours)
+	for (URCBehaviour* Behavior : InController->GetBehaviors())
 	{
 		URCBehaviourBind* CurrentBindBehavior = Cast<URCBehaviourBind>(Behavior);
 
diff --git a/Engine/Plugins/VirtualProduction/RemoteControl/Source/RemoteControlLogic/Private/Tests/RemoteControlLogicTests.cpp b/Engine/Plugins/VirtualProduction/RemoteControl/Source/RemoteControlLogic/Private/Tests/RemoteControlLogicTests.cpp
--- a/Engine/Plugins/VirtualProduction/RemoteControl/Source/RemoteControlLogic/Private/Tests/RemoteControlLogicTests.cpp
+++ b/Engine/Plugins/VirtualProduction/RemoteControl/Source/RemoteControlLogic/Private/Tests/RemoteControlLogicTests.cpp
@@ -154,12 +154,12 @@ bool FRemoteControlLogicTest::RunTest(const FString& Parameters)
 	TestEqual(TEXT("After empty actions the count should be = 0"), StrControllerBehaviour->GetNumActions(), 0);
  
 	// 8. Remove Behaviour by ID and by Pointer
-	const int32 BehaviourNum = FloatController->Behaviours.Num();
+	const int32 BehaviourNum = FloatController->GetBehaviors().Num();
 	FloatController->RemoveBehaviour(FloatControllerBehaviour);
 	FloatController->RemoveBehaviour(IntControllerBehaviour->Id);
-	TestEqual(TEXT("After empty 2 behaviours the count should be =-2"), FloatController->Behaviours.Num(), BehaviourNum - 2);
+	TestEqual(TEXT("After empty 2 behaviours the count should be =-2"), FloatController->GetBehaviors().Num(), BehaviourNum - 2);
 	FloatController->EmptyBehaviours();
-	TestEqual(TEXT("After empty actions the count should be = 0"), FloatController->Behaviours.Num(), 0);
+	TestEqual(TEXT("After empty actions the count should be = 0"), FloatController->GetBehaviors().Num(), 0);
  
 	// 9. Remove Controllers
 	const int32 NumPropertiesBeforeRemove = Preset->GetNumControllers();
diff --git a/Engine/Plugins/VirtualProduction/RemoteControl/Source/RemoteControlLogic/Public/Controller/RCController.h b/Engine/Plugins/VirtualProduction/RemoteControl/Source/RemoteControlLogic/Public/Controller/RCController.h
--- a/Engine/Plugins/VirtualProduction/RemoteControl/Source/RemoteControlLogic/Public/Controller/RCController.h
+++ b/Engine/Plugins/VirtualProduction/RemoteControl/Source/RemoteControlLogic/Public/Controller/RCController.h
@@ -59,6 +59,11 @@ public:
 	/** Handles modifications to controller value; evaluates all behaviours */
 	virtual void OnModifyPropertyValue() override;
 
+	TConstArrayView<TObjectPtr<URCBehaviour>> GetBehaviors() const
+	{
+		return Behaviors;
+	}
+
 	/** Duplicates an existing behaviour and adds it to the behaviour set of a given Controller */
 	static URCBehaviour* DuplicateBehaviour(URCController* InController, URCBehaviour* InBehaviour);
 
@@ -69,8 +74,18 @@ public:
 	virtual void UpdateEntityUsage(const URemoteControlPreset* InPreset, TMap<FGuid, TSet<UObject*>>& InEntityToBehavior) const override;
 	//~ End URCVirtualPropertyBase
 
+	//~ Begin UObject
+	virtual void PostLoad() override;
+	//~ End UObject
+
 public:
-	/** Set of the behaviours */
-	UPROPERTY()
+	/** Deprecated set of behaviors */
+	UE_DEPRECATED(5.7, "Please use the array Behaviors instead.")
+	UPROPERTY(meta = (DeprecatedProperty, DeprecationMessage = "Please use the array Behaviors instead."))
 	TSet<TObjectPtr<URCBehaviour>> Behaviours;
+
+protected:
+	/** Array of the behaviors */
+	UPROPERTY()
+	TArray<TObjectPtr<URCBehaviour>> Behaviors;
 };
diff --git a/Engine/Plugins/VirtualProduction/RemoteControl/Source/RemoteControlUI/Private/UI/Behaviour/SRCBehaviourPanelList.cpp b/Engine/Plugins/VirtualProduction/RemoteControl/Source/RemoteControlUI/Private/UI/Behaviour/SRCBehaviourPanelList.cpp
--- a/Engine/Plugins/VirtualProduction/RemoteControl/Source/RemoteControlUI/Private/UI/Behaviour/SRCBehaviourPanelList.cpp
+++ b/Engine/Plugins/VirtualProduction/RemoteControl/Source/RemoteControlUI/Private/UI/Behaviour/SRCBehaviourPanelList.cpp
@@ -302,7 +302,7 @@ void SRCBehaviourPanelList::Reset()
 	{
 		if (URCController* Controller = Cast<URCController>(ControllerItem->GetVirtualProperty()))
 		{
-			for (URCBehaviour* Behaviour : Controller->Behaviours)
+			for (URCBehaviour* Behaviour : Controller->GetBehaviors())
 			{
 				AddBehaviourToList(Behaviour);
 			}
diff --git a/Engine/Plugins/VirtualProduction/RemoteControl/Source/RemoteControlUI/Private/UI/Controller/SRCControllerPanelList.cpp b/Engine/Plugins/VirtualProduction/RemoteControl/Source/RemoteControlUI/Private/UI/Controller/SRCControllerPanelList.cpp
--- a/Engine/Plugins/VirtualProduction/RemoteControl/Source/RemoteControlUI/Private/UI/Controller/SRCControllerPanelList.cpp
+++ b/Engine/Plugins/VirtualProduction/RemoteControl/Source/RemoteControlUI/Private/UI/Controller/SRCControllerPanelList.cpp
@@ -1586,7 +1586,7 @@ void SRCControllerPanelList::CreateBindBehaviourAndAssignTo(URCController* Contr
 		}
 	}
 
-	for (const URCBehaviour* Behaviour : Controller->Behaviours)
+	for (const URCBehaviour* Behaviour : Controller->GetBehaviors())
 	{
 		if (Behaviour && Behaviour->IsA(URCBehaviourBind::StaticClass()))
 		{
