Remote Control: A few documentation and very minor fixes.

diff --git a/Engine/Plugins/VirtualProduction/RemoteControl/Source/RemoteControlLogic/Private/Action/Path/RCSetAssetByPathActionNew.cpp b/Engine/Plugins/VirtualProduction/RemoteControl/Source/RemoteControlLogic/Private/Action/Path/RCSetAssetByPathActionNew.cpp
--- a/Engine/Plugins/VirtualProduction/RemoteControl/Source/RemoteControlLogic/Private/Action/Path/RCSetAssetByPathActionNew.cpp
+++ b/Engine/Plugins/VirtualProduction/RemoteControl/Source/RemoteControlLogic/Private/Action/Path/RCSetAssetByPathActionNew.cpp
@@ -189,7 +189,7 @@ bool URCSetAssetByPathActionNew::SetInternalAsset(UObject* InSetterObject)
 	{
 		FRemoteControlLogger::Get().Log(SetAssetByPathNew::SetAssetByPathBehaviour, [InSetterObject]
 			{
-				return FText::FromString(FString("Path Behavior fails to set Asset: %s") + *InSetterObject->GetName());
+				return FText::FromString(FString("Path Behavior fails to set Asset: Null"));
 			}, EVerbosityLevel::Error);
 
 		return false;
diff --git a/Engine/Plugins/VirtualProduction/RemoteControl/Source/RemoteControlLogic/Public/Action/Path/RCSetAssetByPathActionNew.h b/Engine/Plugins/VirtualProduction/RemoteControl/Source/RemoteControlLogic/Public/Action/Path/RCSetAssetByPathActionNew.h
--- a/Engine/Plugins/VirtualProduction/RemoteControl/Source/RemoteControlLogic/Public/Action/Path/RCSetAssetByPathActionNew.h
+++ b/Engine/Plugins/VirtualProduction/RemoteControl/Source/RemoteControlLogic/Public/Action/Path/RCSetAssetByPathActionNew.h
@@ -11,7 +11,7 @@ class UTexture;
 struct FRemoteControlProperty;
 
 /**
- * Property Bind Action specifically for Bind Behaviour
+ * Property Set Asset By Path Action (New) specifically for Set Asset By Path Behaviour (New)
  */
 UCLASS(MinimalAPI)
 class URCSetAssetByPathActionNew : public URCPropertyAction
@@ -35,10 +35,10 @@ private:
 	/** Updates property with a existing asset. */
 	bool UpdateInternal();
 
-	/** Updatse property with an imported, external texture. */
+	/** Updates property with an imported, external texture. */
 	bool UpdateExternal();
 
-	/** Auxiliary Function which sets the given SetterObject Asset onto the Exposed Asset with the given PropertyString name. */
+	/** Auxiliary Function which sets the given InSetterObject Asset onto the Exposed Property. */
 	bool SetInternalAsset(UObject* InSetterObject);
 
 	/** Given an External Path towards a external location, loads the asset associated and places it onto an object */
diff --git a/Engine/Plugins/VirtualProduction/RemoteControl/Source/RemoteControlUI/Private/RemoteControlUIModule.cpp b/Engine/Plugins/VirtualProduction/RemoteControl/Source/RemoteControlUI/Private/RemoteControlUIModule.cpp
--- a/Engine/Plugins/VirtualProduction/RemoteControl/Source/RemoteControlUI/Private/RemoteControlUIModule.cpp
+++ b/Engine/Plugins/VirtualProduction/RemoteControl/Source/RemoteControlUI/Private/RemoteControlUIModule.cpp
@@ -1142,6 +1142,7 @@ void FRemoteControlUIModule::UnregisterStructCustomizations()
 	if (UObjectInitialized())
 	{
 		PropertyEditorModule.UnregisterCustomPropertyTypeLayout(FRCAssetPathElement::StaticStruct()->GetFName());
+		PropertyEditorModule.UnregisterCustomPropertyTypeLayout(FRCAssetPathElementNew::StaticStruct()->GetFName());
 		PropertyEditorModule.UnregisterCustomPropertyTypeLayout(FRCNetworkAddress::StaticStruct()->GetFName());
 		PropertyEditorModule.UnregisterCustomPropertyTypeLayout(FRCPassphrase::StaticStruct()->GetFName());
 	}
diff --git a/Engine/Plugins/VirtualProduction/RemoteControl/Source/RemoteControlUI/Private/UI/Behaviour/Builtin/Path/RCBehaviorSetAssetByPathModelNew.cpp b/Engine/Plugins/VirtualProduction/RemoteControl/Source/RemoteControlUI/Private/UI/Behaviour/Builtin/Path/RCBehaviorSetAssetByPathModelNew.cpp
--- a/Engine/Plugins/VirtualProduction/RemoteControl/Source/RemoteControlUI/Private/UI/Behaviour/Builtin/Path/RCBehaviorSetAssetByPathModelNew.cpp
+++ b/Engine/Plugins/VirtualProduction/RemoteControl/Source/RemoteControlUI/Private/UI/Behaviour/Builtin/Path/RCBehaviorSetAssetByPathModelNew.cpp
@@ -260,7 +260,7 @@ void FRCSetAssetByPathBehaviorModelNew::RefreshPathAndPreview()
 		PinnedNode->GetChildren(Children);
 		Children.Insert(PinnedNode.ToSharedRef(), 0);
 		
-		for (uint8 Counter = 0; Counter < Children.Num(); Counter++)
+		for (int32 Counter = 0; Counter < Children.Num(); Counter++)
 		{
 			const FNodeWidgets NodeWidgets = Children[Counter]->CreateNodeWidgets();
 
diff --git a/Engine/Plugins/VirtualProduction/RemoteControl/Source/RemoteControlUI/Private/UI/Behaviour/Builtin/Path/RCBehaviorSetAssetByPathModelNew.h b/Engine/Plugins/VirtualProduction/RemoteControl/Source/RemoteControlUI/Private/UI/Behaviour/Builtin/Path/RCBehaviorSetAssetByPathModelNew.h
--- a/Engine/Plugins/VirtualProduction/RemoteControl/Source/RemoteControlUI/Private/UI/Behaviour/Builtin/Path/RCBehaviorSetAssetByPathModelNew.h
+++ b/Engine/Plugins/VirtualProduction/RemoteControl/Source/RemoteControlUI/Private/UI/Behaviour/Builtin/Path/RCBehaviorSetAssetByPathModelNew.h
@@ -13,10 +13,10 @@ class URCSetAssetByPathBehaviorNew;
 /*
  * ~ FRCSetAssetByPathBehaviorModelNew ~
  *
- * Child Behaviour class representing the "Set Asset By Path" Behaviour's UI model.
+ * Child Behaviour class representing the "Set Asset By Path" Behavior's UI model.
  *
- * Generates several Widgets where users can enter the RootPath, Target Property, and Default Property as FStrings.
- * The values are then put together, getting a path towards a possible object.
+ * Generates several Widgets where users can build a Root Path using strings or controller values.
+ * This is assigned as a loaded asset to exposed properties defined by actions actions.
  */
 class FRCSetAssetByPathBehaviorModelNew : public FRCBehaviourModel, public FNotifyHook
 {
