Remote Control: Stateless controller (Button)

diff --git a/Engine/Plugins/VirtualProduction/RemoteControl/Source/RemoteControlLogic/Private/Behaviour/Builtin/Bind/RCBehaviourBindNode.cpp b/Engine/Plugins/VirtualProduction/RemoteControl/Source/RemoteControlLogic/Private/Behaviour/Builtin/Bind/RCBehaviourBindNode.cpp
--- a/Engine/Plugins/VirtualProduction/RemoteControl/Source/RemoteControlLogic/Private/Behaviour/Builtin/Bind/RCBehaviourBindNode.cpp
+++ b/Engine/Plugins/VirtualProduction/RemoteControl/Source/RemoteControlLogic/Private/Behaviour/Builtin/Bind/RCBehaviourBindNode.cpp
@@ -2,6 +2,9 @@
 
 #include "Behaviour/Builtin/Bind/RCBehaviourBindNode.h"
 
+#include "Behaviour/RCBehaviour.h"
+#include "Controller/RCController.h"
+#include "Controller/RCStatelessEventController.h"
 #include "Behaviour/Builtin/Bind/RCBehaviourBind.h"
 
 #define LOCTEXT_NAMESPACE "RemoteControlBehaviours"
@@ -19,7 +22,12 @@ bool URCBehaviourBindNode::Execute_Implementation(URCBehaviour* InBehaviour) con
 
 bool URCBehaviourBindNode::IsSupported_Implementation(URCBehaviour* InBehaviour) const
 {
-	return true; 
+	if (const URCController* RCController = InBehaviour->ControllerWeakPtr.Get())
+	{
+		return !FRCStatelessEventController::IsStatelessEventController(RCController);
+	}
+
+	return false;
 }
 
 UClass* URCBehaviourBindNode::GetBehaviourClass() const
diff --git a/Engine/Plugins/VirtualProduction/RemoteControl/Source/RemoteControlLogic/Private/Behaviour/Builtin/Conditional/RCBehaviourConditionalNode.cpp b/Engine/Plugins/VirtualProduction/RemoteControl/Source/RemoteControlLogic/Private/Behaviour/Builtin/Conditional/RCBehaviourConditionalNode.cpp
--- a/Engine/Plugins/VirtualProduction/RemoteControl/Source/RemoteControlLogic/Private/Behaviour/Builtin/Conditional/RCBehaviourConditionalNode.cpp
+++ b/Engine/Plugins/VirtualProduction/RemoteControl/Source/RemoteControlLogic/Private/Behaviour/Builtin/Conditional/RCBehaviourConditionalNode.cpp
@@ -3,7 +3,9 @@
 #include "Behaviour/Builtin/Conditional/RCBehaviourConditionalNode.h"
 
 #include "Behaviour/Builtin/Conditional/RCBehaviourConditional.h"
+#include "Behaviour/RCBehaviour.h"
 #include "Controller/RCController.h"
+#include "Controller/RCStatelessEventController.h"
 #include "RCVirtualProperty.h"
 
 #define LOCTEXT_NAMESPACE "RemoteControlBehaviours"
@@ -32,7 +34,13 @@ bool URCBehaviourConditionalNode::IsSupported_Implementation(URCBehaviour* InBeh
 {
 	// Conditional is supported in general for all types that are configured in BaseRemoteControl.ini
 	// Additional filtering per condition type (eg: numeric comparators for numeric fields only) is performed at UI level based on the actively selected Comparator type
-	return true; 
+	
+	if (const URCController* RCController = InBehaviour->ControllerWeakPtr.Get())
+	{
+		return !FRCStatelessEventController::IsStatelessEventController(RCController);
+	}
+
+	return false;
 }
 
 UClass* URCBehaviourConditionalNode::GetBehaviourClass() const
diff --git a/Engine/Plugins/VirtualProduction/RemoteControl/Source/RemoteControlLogic/Private/Controller/RCCustomControllerUtilities.cpp b/Engine/Plugins/VirtualProduction/RemoteControl/Source/RemoteControlLogic/Private/Controller/RCCustomControllerUtilities.cpp
--- a/Engine/Plugins/VirtualProduction/RemoteControl/Source/RemoteControlLogic/Private/Controller/RCCustomControllerUtilities.cpp
+++ b/Engine/Plugins/VirtualProduction/RemoteControl/Source/RemoteControlLogic/Private/Controller/RCCustomControllerUtilities.cpp
@@ -8,6 +8,7 @@
 #include "Controller/RCController.h"
 #include "RCVirtualProperty.h"
 #include "RCVirtualPropertyContainer.h"
+#include "RemoteControlCommon.h"
 #include "RemoteControlPreset.h"
 #include "RemoteControlPropertyIdRegistry.h"
 
@@ -15,7 +16,9 @@ namespace UE::RCCustomControllers
 {
 	// A Map to associate a custom controller name to its underlying type
 	inline static TMap<FString, EPropertyBagPropertyType> CustomControllerTypes =
-		{{CustomTextureControllerName, EPropertyBagPropertyType::String}};
+		{
+			{CustomTextureControllerName, EPropertyBagPropertyType::String}
+		};
 }
 
 bool UE::RCCustomControllers::IsCustomController(const URCVirtualPropertyBase* InController)
diff --git a/Engine/Plugins/VirtualProduction/RemoteControl/Source/RemoteControlLogic/Private/Controller/RCStatelessEventController.cpp b/Engine/Plugins/VirtualProduction/RemoteControl/Source/RemoteControlLogic/Private/Controller/RCStatelessEventController.cpp
new file mode 100644
--- /dev/null
+++ b/Engine/Plugins/VirtualProduction/RemoteControl/Source/RemoteControlLogic/Private/Controller/RCStatelessEventController.cpp
@@ -0,0 +1,23 @@
+// Copyright Epic Games, Inc. All Rights Reserved.
+
+#include "Controller/RCStatelessEventController.h"
+
+#include "UObject/UnrealType.h"
+#include "RCVirtualProperty.h"
+
+bool FRCStatelessEventController::IsStatelessEventController(const URCVirtualPropertyBase* InController)
+{
+	if (InController->GetValueType() != EPropertyBagPropertyType::Struct)
+	{
+		return false;
+	}
+
+	const FStructProperty* StructProperty = CastField<FStructProperty>(InController->GetProperty());
+
+	if (!StructProperty)
+	{
+		return false;
+	}
+
+	return StructProperty->Struct == FRCStatelessEventController::StaticStruct();
+}
diff --git a/Engine/Plugins/VirtualProduction/RemoteControl/Source/RemoteControlLogic/Public/Controller/RCCustomControllerUtilities.h b/Engine/Plugins/VirtualProduction/RemoteControl/Source/RemoteControlLogic/Public/Controller/RCCustomControllerUtilities.h
--- a/Engine/Plugins/VirtualProduction/RemoteControl/Source/RemoteControlLogic/Public/Controller/RCCustomControllerUtilities.h
+++ b/Engine/Plugins/VirtualProduction/RemoteControl/Source/RemoteControlLogic/Public/Controller/RCCustomControllerUtilities.h
@@ -49,6 +49,7 @@ namespace UE::RCCustomControllers
 
 	inline static FString CustomTextureControllerName = TEXT("Texture");
 	inline static FString CustomStringListControllerName = TEXT("StringList");
+	inline static FString CustomStatelessEventControllerName = TEXT("StatelessEvent");
 
 	inline static FName StringListControllerEntriesName = "StringOptions";
 };
diff --git a/Engine/Plugins/VirtualProduction/RemoteControl/Source/RemoteControlLogic/Public/Controller/RCStatelessEventController.h b/Engine/Plugins/VirtualProduction/RemoteControl/Source/RemoteControlLogic/Public/Controller/RCStatelessEventController.h
new file mode 100644
--- /dev/null
+++ b/Engine/Plugins/VirtualProduction/RemoteControl/Source/RemoteControlLogic/Public/Controller/RCStatelessEventController.h
@@ -0,0 +1,18 @@
+// Copyright Epic Games, Inc. All Rights Reserved.
+
+#pragma once
+
+#include "HAL/Platform.h"
+#include "UObject/Class.h"
+
+#include "RCStatelessEventController.generated.h"
+
+class URCVirtualPropertyBase;
+
+USTRUCT()
+struct FRCStatelessEventController
+{
+	GENERATED_BODY()
+
+	REMOTECONTROLLOGIC_API static bool IsStatelessEventController(const URCVirtualPropertyBase* InController);
+};
diff --git a/Engine/Plugins/VirtualProduction/RemoteControl/Source/RemoteControlUI/Private/RemoteControlUIModule.cpp b/Engine/Plugins/VirtualProduction/RemoteControl/Source/RemoteControlUI/Private/RemoteControlUIModule.cpp
--- a/Engine/Plugins/VirtualProduction/RemoteControl/Source/RemoteControlUI/Private/RemoteControlUIModule.cpp
+++ b/Engine/Plugins/VirtualProduction/RemoteControl/Source/RemoteControlUI/Private/RemoteControlUIModule.cpp
@@ -34,12 +34,12 @@
 #include "RemoteControlSettings.h"
 #include "ScopedTransaction.h"
 #include "Styling/AppStyle.h"
-#include "Styling/AppStyle.h"
 #include "Textures/SlateIcon.h"
 #include "Toolkits/GlobalEditorCommonCommands.h"
 #include "UI/Action/SRCActionPanel.h"
 #include "UI/Behaviour/Builtin/Conditional/SRCBehaviourConditional.h"
 #include "UI/Behaviour/SRCBehaviourPanel.h"
+#include "UI/Controller/CustomControllers/SCustomStatelessEventControllerWidget.h"
 #include "UI/Controller/CustomControllers/SCustomStringListControllerWidget.h"
 #include "UI/Controller/CustomControllers/SCustomTextureControllerWidget.h"
 #include "UI/Customizations/FPassphraseCustomization.h"
@@ -1234,6 +1234,10 @@ TSharedPtr<SWidget> FRemoteControlUIModule::CreateCustomControllerWidget(URCVirt
 	{
 		CustomControllerWidget = SNew(SCustomStringListControllerWidget, Cast<URCController>(InController), InOriginalPropertyHandle);
 	}
+	else if (CustomName == UE::RCCustomControllers::CustomStatelessEventControllerName)
+	{
+		CustomControllerWidget = SNew(SCustomStatelessEventControllerWidget, Cast<URCController>(InController), InOriginalPropertyHandle);
+	}
 
 	return CustomControllerWidget;
 }
diff --git a/Engine/Plugins/VirtualProduction/RemoteControl/Source/RemoteControlUI/Private/UI/Controller/CustomControllers/SCustomStatelessEventControllerWidget.cpp b/Engine/Plugins/VirtualProduction/RemoteControl/Source/RemoteControlUI/Private/UI/Controller/CustomControllers/SCustomStatelessEventControllerWidget.cpp
new file mode 100644
--- /dev/null
+++ b/Engine/Plugins/VirtualProduction/RemoteControl/Source/RemoteControlUI/Private/UI/Controller/CustomControllers/SCustomStatelessEventControllerWidget.cpp
@@ -0,0 +1,44 @@
+// Copyright Epic Games, Inc. All Rights Reserved.
+
+#include "UI/Controller/CustomControllers/SCustomStatelessEventControllerWidget.h"
+
+#include "Controller/RCController.h"
+#include "RemoteControlPreset.h"
+#include "Widgets/Input/SButton.h"
+#include "Widgets/Text/STextBlock.h"
+
+#define LOCTEXT_NAMESPACE "SCustomTextureControllerWidget"
+
+void SCustomStatelessEventControllerWidget::Construct(const FArguments& InArgs, URCController* InController, const TSharedPtr<IPropertyHandle>& InOriginalPropertyHandle)
+{
+	ControllerWeak = InController;
+
+	ChildSlot
+	[
+		SNew(SButton)
+		.OnClicked(this, &SCustomStatelessEventControllerWidget::OnButtonClicked)
+		.VAlign(VAlign_Center)
+		.HAlign(HAlign_Center)
+		[
+			SNew(STextBlock)
+			.Text(LOCTEXT("Trigger", "Trigger"))
+		]
+	];
+}
+
+FReply SCustomStatelessEventControllerWidget::OnButtonClicked()
+{
+	if (URCController* Controller = ControllerWeak.Get())
+	{
+		Controller->ExecuteBehaviours(/* Pre change */ false);
+
+		if (URemoteControlPreset* Preset = Controller->PresetWeakPtr.Get())
+		{
+			Preset->OnControllerModified().Broadcast(Preset, {Controller->Id});
+		}
+	}
+
+	return FReply::Handled();
+}
+
+#undef LOCTEXT_NAMESPACE
diff --git a/Engine/Plugins/VirtualProduction/RemoteControl/Source/RemoteControlUI/Private/UI/Controller/CustomControllers/SCustomStatelessEventControllerWidget.h b/Engine/Plugins/VirtualProduction/RemoteControl/Source/RemoteControlUI/Private/UI/Controller/CustomControllers/SCustomStatelessEventControllerWidget.h
new file mode 100644
--- /dev/null
+++ b/Engine/Plugins/VirtualProduction/RemoteControl/Source/RemoteControlUI/Private/UI/Controller/CustomControllers/SCustomStatelessEventControllerWidget.h
@@ -0,0 +1,30 @@
+// Copyright Epic Games, Inc. All Rights Reserved.
+
+#pragma once
+
+#include "Widgets/SCompoundWidget.h"
+
+#include "UObject/WeakObjectPtr.h"
+
+class IPropertyHandle;
+class URCController;
+
+class SCustomStatelessEventControllerWidget : public SCompoundWidget
+{
+public:
+	SLATE_BEGIN_ARGS(SCustomStatelessEventControllerWidget)
+		{}
+	SLATE_END_ARGS()
+
+	/**
+	 * Constructs this widget with InArgs
+	 * @param InOriginalPropertyHandle Original PropertyHandle of this CustomWidget
+	 */
+	void Construct(const FArguments& InArgs, URCController* InController, const TSharedPtr<IPropertyHandle>& InOriginalPropertyHandle);
+
+private:
+	/** The controller associated with this widget. */
+	TWeakObjectPtr<URCController> ControllerWeak;
+
+	FReply OnButtonClicked();
+};
diff --git a/Engine/Plugins/VirtualProduction/RemoteControl/Source/RemoteControlUI/Private/UI/Controller/SRCControllerPanel.cpp b/Engine/Plugins/VirtualProduction/RemoteControl/Source/RemoteControlUI/Private/UI/Controller/SRCControllerPanel.cpp
--- a/Engine/Plugins/VirtualProduction/RemoteControl/Source/RemoteControlUI/Private/UI/Controller/SRCControllerPanel.cpp
+++ b/Engine/Plugins/VirtualProduction/RemoteControl/Source/RemoteControlUI/Private/UI/Controller/SRCControllerPanel.cpp
@@ -4,6 +4,7 @@
 
 #include "Controller/RCController.h"
 #include "Controller/RCCustomControllerUtilities.h"
+#include "Controller/RCStatelessEventController.h"
 #include "Engine/Texture.h"
 #include "Framework/MultiBox/MultiBoxBuilder.h"
 #include "GraphEditorSettings.h"
@@ -23,7 +24,6 @@
 #include "Styling/RemoteControlStyles.h"
 #include "TypeSelector/SRCControllerTypeSelector.h"
 #include "UI/BaseLogicUI/RCLogicModeBase.h"
-#include "UI/Controller/SRCControllerPanel.h"
 #include "UI/Panels/SRCDockPanel.h"
 #include "UI/RCUIHelpers.h"
 #include "UI/RemoteControlPanelStyle.h"
@@ -43,7 +43,6 @@
 
 namespace UE::RemoteControl::UI::Private
 {
-	const FLazyName FRCControllerPropertyInfo::NonPropertyTypeName = TEXT("NonPropertyType");
 	const FLazyName FRCControllerPropertyInfo::CategoryName = TEXT("Category");
 	const FLazyName FRCControllerPropertyInfo::InitialDescriptionName = TEXT("InitialDescription");
 	const FString FRCControllerPropertyInfo::CategoryMeta = TEXT("Category");
@@ -473,10 +472,39 @@ TSharedRef<SWidget> SRCControllerPanel::GetControllerMenuContentWidget() const
 		MenuBuilder.AddMenuEntry(Action, MenuItemWidget);
 	}
 
+	using namespace UE::RCCustomControllers;
+
 	const UGraphEditorSettings* GraphEditorSettings = GetDefault<UGraphEditorSettings>();
 
-	using namespace UE::RCCustomControllers;
+	// Event
+	{
+		const FText ButtonName = LOCTEXT("StatelessEvent", "Event");
+
+		TSharedRef<SWidget> MenuItemWidget =
+			SNew(SHorizontalBox)
+			+ SHorizontalBox::Slot()
+			.AutoWidth()
+			.HAlign(HAlign_Left)
+			.VAlign(VAlign_Fill)
+			[
+				UE::RCUIHelpers::GetTypeColorWidget(GraphEditorSettings->DelegatePinTypeColor, ButtonName)
+			]
+			+ SHorizontalBox::Slot()
+			.AutoWidth()
+			.Padding(FMargin(6.f, 0.f))
+			[
+				SNew(STextBlock)
+				.Text(ButtonName)
+			];
+
+		FRCControllerPropertyInfo ButtonPropertyInfo;
+		ButtonPropertyInfo.Type = EPropertyBagPropertyType::Struct;
+		ButtonPropertyInfo.ValueTypeObject = FRCStatelessEventController::StaticStruct();
+		ButtonPropertyInfo.CustomMetaData.Add(CustomControllerNameKey, CustomStatelessEventControllerName);
 
+		FUIAction Action(FExecuteAction::CreateSP(this, &SRCControllerPanel::OnAddControllerClicked, ButtonPropertyInfo.Type, ButtonPropertyInfo));
+		MenuBuilder.AddMenuEntry(Action, MenuItemWidget);
+	}
 
 	// All enums
 	{
@@ -727,17 +755,7 @@ void SRCControllerPanel::OnAddControllerClicked(const EPropertyBagPropertyType I
 
 	using namespace UE::RemoteControl::UI::Private;
 
-	FName InitialPropertyName = NAME_None;
-	
-	if (const FString* NonPropertyTypeValue = InControllerPropertyInfo.CustomMetaData.Find(FRCControllerPropertyInfo::NonPropertyTypeName))
-	{
-		if ((*NonPropertyTypeValue) == FRCControllerPropertyInfo::CategoryMeta)
-		{
-			InitialPropertyName = **NonPropertyTypeValue;
-		}
-	}
-
-	URCVirtualPropertyInContainer* NewVirtualProperty = Preset->AddController(URCController::StaticClass(), InValueType, InControllerPropertyInfo.ValueTypeObject, InitialPropertyName);
+	URCVirtualPropertyInContainer* NewVirtualProperty = Preset->AddController(URCController::StaticClass(), InValueType, InControllerPropertyInfo.ValueTypeObject, NAME_None);
 
 	if (!NewVirtualProperty)
 	{
diff --git a/Engine/Plugins/VirtualProduction/RemoteControl/Source/RemoteControlUI/Private/UI/Controller/SRCControllerPanel.h b/Engine/Plugins/VirtualProduction/RemoteControl/Source/RemoteControlUI/Private/UI/Controller/SRCControllerPanel.h
--- a/Engine/Plugins/VirtualProduction/RemoteControl/Source/RemoteControlUI/Private/UI/Controller/SRCControllerPanel.h
+++ b/Engine/Plugins/VirtualProduction/RemoteControl/Source/RemoteControlUI/Private/UI/Controller/SRCControllerPanel.h
@@ -15,7 +15,6 @@ struct FRCPanelStyle;
 
 namespace UE::RemoteControl::UI::Private
 {
-
 /**
  * Helper struct used when generating Controllers.
  * Wraps the TypeObject and the Bag Type.
@@ -26,7 +25,6 @@ namespace UE::RemoteControl::UI::Private
  */
 struct FRCControllerPropertyInfo
 {
-	static const FLazyName NonPropertyTypeName;
 	static const FLazyName CategoryName;
 	static const FLazyName InitialDescriptionName;
 	static const FString CategoryMeta;
