Remote Control: Made behaviors rearrangable in the RC panel.

diff --git a/Engine/Plugins/VirtualProduction/RemoteControl/Source/RemoteControlLogic/Private/Controller/RCController.cpp b/Engine/Plugins/VirtualProduction/RemoteControl/Source/RemoteControlLogic/Private/Controller/RCController.cpp
--- a/Engine/Plugins/VirtualProduction/RemoteControl/Source/RemoteControlLogic/Private/Controller/RCController.cpp
+++ b/Engine/Plugins/VirtualProduction/RemoteControl/Source/RemoteControlLogic/Private/Controller/RCController.cpp
@@ -151,6 +151,54 @@ URCBehaviour* URCController::DuplicateBehaviour(URCController* InController, URC
 	return NewBehaviour;
 }
 
+bool URCController::ReorderBehaviorItems(URCBehaviour* InDroppedOnBehavior, bool bInBelowItem, TArray<URCBehaviour*> InDroppedBehaviors)
+{
+	if (!Behaviors.Contains(InDroppedOnBehavior) || InDroppedBehaviors.IsEmpty())
+	{
+		return false;
+	}
+
+	TArray<URCBehaviour*> NewBehaviors;
+
+	auto AddDroppedItems = [this, &InDroppedBehaviors, &NewBehaviors]
+		{
+			for (URCBehaviour* Behavior : InDroppedBehaviors)
+			{
+				if (Behavior)
+				{
+					if (Behaviors.Contains(Behavior))
+					{
+						NewBehaviors.Add(Behavior);
+					}
+				}
+			}
+		};
+
+	for (URCBehaviour* Behavior : Behaviors)
+	{
+		if (Behavior == InDroppedOnBehavior && !bInBelowItem)
+		{
+			AddDroppedItems();
+		}
+
+		if (!InDroppedBehaviors.Contains(Behavior))
+		{
+			NewBehaviors.Add(Behavior);
+		}
+
+		if (Behavior == InDroppedOnBehavior && bInBelowItem)
+		{
+			AddDroppedItems();
+		}
+	}
+
+	Behaviors = NewBehaviors;
+
+	OnBehaviourListModified.Broadcast();
+
+	return true;
+}
+
 void URCController::UpdateEntityUsage(const URemoteControlPreset* InPreset, TMap<FGuid, TSet<UObject*>>& InEntityToBehavior) const
 {
 	URemoteControlPreset* Preset = PresetWeakPtr.Get();
diff --git a/Engine/Plugins/VirtualProduction/RemoteControl/Source/RemoteControlLogic/Public/Controller/RCController.h b/Engine/Plugins/VirtualProduction/RemoteControl/Source/RemoteControlLogic/Public/Controller/RCController.h
--- a/Engine/Plugins/VirtualProduction/RemoteControl/Source/RemoteControlLogic/Public/Controller/RCController.h
+++ b/Engine/Plugins/VirtualProduction/RemoteControl/Source/RemoteControlLogic/Public/Controller/RCController.h
@@ -67,6 +67,9 @@ public:
 	/** Duplicates an existing behaviour and adds it to the behaviour set of a given Controller */
 	static URCBehaviour* DuplicateBehaviour(URCController* InController, URCBehaviour* InBehaviour);
 
+	/** Rearranges the behaviors based on the dropped behaviours. */
+	bool ReorderBehaviorItems(URCBehaviour* InDroppedOnBehavior, bool bInBelowItem, TArray<URCBehaviour*> InDroppedBehaviors);
+
 	/** Delegate that notifies changes to the list of behaviours*/
 	FOnBehaviourListModified OnBehaviourListModified;
 
diff --git a/Engine/Plugins/VirtualProduction/RemoteControl/Source/RemoteControlUI/Private/UI/Behaviour/SRCBehaviorPanelRow.cpp b/Engine/Plugins/VirtualProduction/RemoteControl/Source/RemoteControlUI/Private/UI/Behaviour/SRCBehaviorPanelRow.cpp
new file mode 100644
--- /dev/null
+++ b/Engine/Plugins/VirtualProduction/RemoteControl/Source/RemoteControlUI/Private/UI/Behaviour/SRCBehaviorPanelRow.cpp
@@ -0,0 +1,193 @@
+// Copyright Epic Games, Inc. All Rights Reserved.
+
+#include "UI/Behaviour/SRCBehaviorPanelRow.h"
+
+#include "Behaviour/RCBehaviour.h"
+#include "DragAndDrop/DecoratedDragDropOp.h"
+#include "Styles/SlateBrushTemplates.h"
+#include "Styling/RemoteControlStyles.h"
+#include "UI/Behaviour/SRCBehaviourPanelList.h"
+#include "UI/RemoteControlPanelStyle.h"
+#include "Widgets/SBoxPanel.h"
+#include "Widgets/Input/SCheckBox.h"
+
+#define LOCTEXT_NAMESPACE "SRCBehaviorPanelRow"
+
+/*
+* ~ FRCBehaviorPanelDragDrop ~
+* Facilitates drag-drop operation
+*/
+class FRCBehaviorPanelDragDrop final : public FDecoratedDragDropOp
+{
+public:
+	DRAG_DROP_OPERATOR_TYPE(FRCBehaviorPanelDragDrop, FDecoratedDragDropOp)
+
+	explicit FRCBehaviorPanelDragDrop(const TArray<FRCBehaviorPanelListItem>& InItems)
+		: Items(InItems)
+	{
+		MouseCursor = EMouseCursor::GrabHandClosed;
+	}
+
+	TArray<FRCBehaviorPanelListItem> GetItems() const
+	{
+		return Items;
+	}
+
+private:
+	TArray<FRCBehaviorPanelListItem> Items;
+};
+
+void SRCBehaviorPanelRow::Construct(const FArguments& InArgs, const TSharedRef<STableViewBase>& InOwnerTableView, TSharedRef<SRCBehaviourPanelList> InBehaviorPanelList, FRCBehaviorPanelListItem InBehaviorItem)
+{
+	BehaviorItem = InBehaviorItem;
+	BehaviorPanelListWeak = InBehaviorPanelList;
+
+	const FRCPanelStyle* RCPanelStyle = &FRemoteControlPanelStyle::Get()->GetWidgetStyle<FRCPanelStyle>("RemoteControlPanel.MinorPanel");
+
+	STableRow::Construct(
+		STableRow::FArguments()
+		.Style(InArgs._Style)
+		.Padding(0.f)
+		.OnDragDetected(this, &SRCBehaviorPanelRow::OnRowDragDetected)
+		.OnCanAcceptDrop(this, &SRCBehaviorPanelRow::CanAcceptDrop)
+		.OnAcceptDrop(this, &SRCBehaviorPanelRow::OnAcceptDrop)
+		.ShowWires(false)
+		.Content()
+		[
+			SNew(SHorizontalBox)
+			+ SHorizontalBox::Slot()
+			.VAlign(VAlign_Center)
+			.AutoWidth()
+			.Padding(4.f, 0.f)
+			[
+				SNew(SImage)
+				.Image(FSlateBrushTemplates::DragHandle())
+			]
+
+			+ SHorizontalBox::Slot()
+			[
+				InBehaviorItem->GetWidget()
+			]
+
+			// Toggle Behaviour Button
+			+ SHorizontalBox::Slot()
+			.VAlign(VAlign_Center)
+			.AutoWidth()
+			.Padding(4.f, 0.f)
+			[
+				SNew(SCheckBox)
+				.AddMetaData<FTagMetaData>(FTagMetaData(TEXT("Toggle Behaviour")))
+				.ToolTipText(LOCTEXT("EditModeTooltip", "Enable/Disable this Behaviour.\nWhen a behaviour is disabled its Actions will not be processed when the Controller value changes"))
+				.HAlign(HAlign_Center)
+				.ForegroundColor(FLinearColor::White)
+				.Style(&RCPanelStyle->ToggleButtonStyle)
+				.IsChecked(this, &SRCBehaviorPanelRow::IsBehaviourChecked)
+				.OnCheckStateChanged(this, &SRCBehaviorPanelRow::OnToggleEnableBehaviour)
+			]
+		],
+		InOwnerTableView
+	);
+}
+
+FRCBehaviorPanelListItem SRCBehaviorPanelRow::GetBehavior() const
+{
+	return BehaviorItem;
+}
+
+TOptional<EItemDropZone> SRCBehaviorPanelRow::CanAcceptDrop(const FDragDropEvent& InDragDropEvent, EItemDropZone InDropZone, FRCBehaviorPanelListItem InItem)
+{
+	if (!BehaviorPanelListWeak.IsValid() || !BehaviorItem.IsValid())
+	{
+		return {};
+	}
+
+	TSharedPtr<FRCBehaviorPanelDragDrop> BehaviorDragDrop = InDragDropEvent.GetOperationAs<FRCBehaviorPanelDragDrop>();
+
+	if (!BehaviorDragDrop.IsValid())
+	{
+		return {};
+	}
+
+	const TArray<FRCBehaviorPanelListItem> Items = BehaviorDragDrop->GetItems();
+
+	if (Items.IsEmpty())
+	{
+		return {};
+	}
+
+	if (Items.Num() == 1 && Items[0].Get() == BehaviorItem.Get())
+	{
+		return {};
+	}
+
+	return InDropZone;
+}
+
+FReply SRCBehaviorPanelRow::OnAcceptDrop(const FDragDropEvent& InDragDropEvent, EItemDropZone InDropZone, FRCBehaviorPanelListItem InItem)
+{
+	if (!BehaviorItem.IsValid())
+	{
+		return FReply::Handled();
+	}
+
+	TSharedPtr<SRCBehaviourPanelList> BehaviorPanelList = BehaviorPanelListWeak.Pin();
+
+	if (!BehaviorPanelList.IsValid())
+	{
+		return FReply::Handled();
+	}
+
+	TSharedPtr<FRCBehaviorPanelDragDrop> BehaviorDragDrop = InDragDropEvent.GetOperationAs<FRCBehaviorPanelDragDrop>();
+
+	if (!BehaviorDragDrop.IsValid())
+	{
+		return FReply::Handled();
+	}
+
+	const TArray<FRCBehaviorPanelListItem> Items = BehaviorDragDrop->GetItems();
+
+	if (Items.IsEmpty())
+	{
+		return FReply::Handled();
+	}
+
+	BehaviorPanelList->ReorderBehaviorItems(BehaviorItem.ToSharedRef(), InDropZone, Items);
+
+	return FReply::Handled();
+}
+
+FReply SRCBehaviorPanelRow::OnRowDragDetected(const FGeometry& InMyGeometry, const FPointerEvent& InMouseEvent)
+{
+	TSharedPtr<SRCBehaviourPanelList> BehaviorPanelList = BehaviorPanelListWeak.Pin();
+
+	if (!BehaviorPanelList.IsValid())
+	{
+		return FReply::Unhandled();
+	}
+
+	if (!InMouseEvent.IsMouseButtonDown(EKeys::LeftMouseButton))
+	{
+		return FReply::Unhandled();
+	}
+
+	TSharedRef<FRCBehaviorPanelDragDrop> BehaviorDragDrop = MakeShared<FRCBehaviorPanelDragDrop>(BehaviorPanelList->GetSelectedBehaviourItems());
+	return FReply::Handled().BeginDragDrop(BehaviorDragDrop);
+}
+
+ECheckBoxState SRCBehaviorPanelRow::IsBehaviourChecked() const
+{
+	return BehaviorItem.IsValid() && BehaviorItem->IsBehaviourEnabled() ? ECheckBoxState::Checked : ECheckBoxState::Unchecked;
+}
+
+void SRCBehaviorPanelRow::OnToggleEnableBehaviour(ECheckBoxState State)
+{
+	if (BehaviorItem.IsValid())
+	{
+		if (TSharedPtr<SRCBehaviourPanelList> BehaviorPanelList = BehaviorPanelListWeak.Pin())
+		{
+			BehaviorPanelList->SetIsBehaviourEnabled(BehaviorItem.ToSharedRef(), State == ECheckBoxState::Checked);
+		}
+	}
+}
+
+#undef LOCTEXT_NAMESPACE
diff --git a/Engine/Plugins/VirtualProduction/RemoteControl/Source/RemoteControlUI/Private/UI/Behaviour/SRCBehaviorPanelRow.h b/Engine/Plugins/VirtualProduction/RemoteControl/Source/RemoteControlUI/Private/UI/Behaviour/SRCBehaviorPanelRow.h
new file mode 100644
--- /dev/null
+++ b/Engine/Plugins/VirtualProduction/RemoteControl/Source/RemoteControlUI/Private/UI/Behaviour/SRCBehaviorPanelRow.h
@@ -0,0 +1,47 @@
+// Copyright Epic Games, Inc. All Rights Reserved.
+
+#pragma once
+
+#include "Misc/Optional.h"
+#include "Widgets/Views/SListView.h"
+#include "Widgets/Views/STableRow.h"
+
+class FRCBehaviourModel;
+class SRCBehaviourPanelList;
+class URemoteControlPreset;
+
+using FRCBehaviorPanelListItem = TSharedPtr<FRCBehaviourModel>;
+
+class SRCBehaviorPanelRow : public STableRow<FRCBehaviorPanelListItem>
+{
+	using FSuperRowType = STableRow<FRCBehaviorPanelListItem>;
+
+public:
+	SLATE_BEGIN_ARGS(SRCBehaviorPanelRow) 
+		{}
+		SLATE_STYLE_ARGUMENT(FTableRowStyle, Style)
+	SLATE_END_ARGS()
+
+	void Construct(const FArguments& InArgs, const TSharedRef<STableViewBase>& InOwnerTableView,
+		TSharedRef<SRCBehaviourPanelList> InBehaviorPanelList, FRCBehaviorPanelListItem InBehaviorItem);
+
+	FRCBehaviorPanelListItem GetBehavior() const;
+
+private:
+	/** Called to validate whether a drag drop event is accepted for a behavior item row */
+	TOptional<EItemDropZone> CanAcceptDrop(const FDragDropEvent& InDragDropEvent, EItemDropZone InDropZone, FRCBehaviorPanelListItem InItem);
+
+	/** Called when a drop happens on this behavior item row */
+	FReply OnAcceptDrop(const FDragDropEvent& InDragDropEvent, EItemDropZone InDropZone, FRCBehaviorPanelListItem InItem);
+
+	FReply OnRowDragDetected(const FGeometry& InMyGeometry, const FPointerEvent& InMouseEvent);
+
+	/** Return whether or not the given behaviour is checked */
+	ECheckBoxState IsBehaviourChecked() const;
+
+	/** Executed when toggling the behaviour state */
+	void OnToggleEnableBehaviour(ECheckBoxState State);
+
+	TWeakPtr<SRCBehaviourPanelList> BehaviorPanelListWeak;
+	FRCBehaviorPanelListItem BehaviorItem;
+};
\ No newline at end of file
diff --git a/Engine/Plugins/VirtualProduction/RemoteControl/Source/RemoteControlUI/Private/UI/Behaviour/SRCBehaviourPanelList.cpp b/Engine/Plugins/VirtualProduction/RemoteControl/Source/RemoteControlUI/Private/UI/Behaviour/SRCBehaviourPanelList.cpp
--- a/Engine/Plugins/VirtualProduction/RemoteControl/Source/RemoteControlUI/Private/UI/Behaviour/SRCBehaviourPanelList.cpp
+++ b/Engine/Plugins/VirtualProduction/RemoteControl/Source/RemoteControlUI/Private/UI/Behaviour/SRCBehaviourPanelList.cpp
@@ -1,6 +1,6 @@
 // Copyright Epic Games, Inc. All Rights Reserved.
 
-#include "SRCBehaviourPanelList.h"
+#include "UI/Behaviour/SRCBehaviourPanelList.h"
 
 #include "Behaviour/Builtin/Bind/RCBehaviourBind.h"
 #include "Behaviour/Builtin/Conditional/RCBehaviourConditional.h"
@@ -9,6 +9,7 @@
 #include "Controller/RCController.h"
 #include "RCBehaviourModel.h"
 #include "RemoteControlPreset.h"
+#include "ScopedTransaction.h"
 #include "SlateOptMacros.h"
 #include "SRCBehaviourPanel.h"
 #include "Styling/RemoteControlStyles.h"
@@ -18,6 +19,7 @@
 #include "UI/Behaviour/Builtin/Path/RCBehaviorSetAssetByPathModelNew.h"
 #include "UI/Behaviour/Builtin/Path/RCBehaviourSetAssetByPathModel.h"
 #include "UI/Behaviour/Builtin/RangeMap/RCBehaviourRangeMapModel.h"
+#include "UI/Behaviour/SRCBehaviorPanelRow.h"
 #include "UI/Controller/RCControllerModel.h"
 #include "UI/RemoteControlPanelStyle.h"
 #include "UI/SRCPanelExposedEntitiesList.h"
@@ -107,8 +109,8 @@ void SRCBehaviourPanelList::AddNewLogicItem(UObject* InLogicItem)
 
 void SRCBehaviourPanelList::AddSpecialContextMenuOptions(FMenuBuilder& MenuBuilder)
 {
-	const FUIAction EnableAction(FExecuteAction::CreateSP(this, &SRCBehaviourPanelList::SetIsBehaviourEnabled, true));
-	const FUIAction DisableAction(FExecuteAction::CreateSP(this, &SRCBehaviourPanelList::SetIsBehaviourEnabled, false));
+	const FUIAction EnableAction(FExecuteAction::CreateSP(this, &SRCBehaviourPanelList::SetAreSelectedBehaviorsEnabled, true));
+	const FUIAction DisableAction(FExecuteAction::CreateSP(this, &SRCBehaviourPanelList::SetAreSelectedBehaviorsEnabled, false));
 
 	static const FText EnableLabel = LOCTEXT("LabelEnableBehaviour", "Enable");
 	static const FText EnableTooltip = LOCTEXT("TooltipEnableBehaviour", "Enables the current behaviours. Restores functioning of Actions associated with these behaviours.");
@@ -213,7 +215,7 @@ void SRCBehaviourPanelList::SetControllerItem(const TSharedPtr<FRCControllerMode
 	Reset();
 }
 
-void SRCBehaviourPanelList::SetIsBehaviourEnabled(const bool bIsEnabled)
+void SRCBehaviourPanelList::SetAreSelectedBehaviorsEnabled(const bool bIsEnabled)
 {
 	// Disable all selected behaviour here
 	for (const TSharedPtr<FRCLogicModeBase>& LogicItem : GetSelectedLogicItems())
@@ -236,6 +238,63 @@ void SRCBehaviourPanelList::SetIsBehaviourEnabled(const bool bIsEnabled)
 	}
 }
 
+void SRCBehaviourPanelList::SetIsBehaviourEnabled(TSharedRef<FRCBehaviourModel> InBehaviorModel, const bool bIsEnabled)
+{
+	// Disable all selected behaviour here
+	InBehaviorModel->SetIsBehaviourEnabled(bIsEnabled);
+
+	// Disable only the ActionPanel of the currently selected behaviour.
+	// When selecting a different behaviour the ActionPanel will be re-constructed,
+	// it will get the current status of the behaviour (enabled or disabled) and will match it for the ActionPanel as well
+	if (const TSharedPtr<SRemoteControlPanel> RemoteControlPanel = GetRemoteControlPanel())
+	{
+		if (const TSharedPtr<SRCActionPanel> ActionPanel = RemoteControlPanel->GetLogicActionPanel())
+		{
+			ActionPanel->RefreshIsBehaviourEnabled(bIsEnabled);
+		}
+	}
+}
+
+void SRCBehaviourPanelList::ReorderBehaviorItems(TSharedRef<FRCBehaviourModel> InDroppedOnModel, EItemDropZone InDropZone, TArray<TSharedPtr<FRCBehaviourModel>> InDroppedModels)
+{
+	TSharedPtr<FRCControllerModel> ControllerMoel = ControllerItemWeakPtr.Pin();
+
+	if (!ControllerMoel.IsValid())
+	{
+		return;
+	}
+
+	URCController* Controller = Cast<URCController>(ControllerMoel->GetVirtualProperty());
+
+	if (!Controller)
+	{
+		return;
+	}
+
+	URCBehaviour* DroppedOnBehavior = InDroppedOnModel->GetBehaviour();
+
+	if (!DroppedOnBehavior)
+	{
+		return;
+	}
+
+	TArray<URCBehaviour*> Behaviors;
+	Behaviors.Reserve(InDroppedModels.Num());
+
+	for (const TSharedPtr<FRCBehaviourModel>& BehaviorModel : InDroppedModels)
+	{
+		if (URCBehaviour* Behavior = BehaviorModel->GetBehaviour())
+		{
+			Behaviors.Add(Behavior);
+		}
+	}
+
+	FScopedTransaction Transaction(LOCTEXT("ReorderBehaviors", "Reorder Behaviors"));
+	Controller->Modify();
+
+	Controller->ReorderBehaviorItems(DroppedOnBehavior, InDropZone != EItemDropZone::AboveItem, Behaviors);
+}
+
 void SRCBehaviourPanelList::SetIsBehaviourEnabled(const TSharedPtr<FRCBehaviourModel>& InBehaviourModel, const bool bIsEnabled)
 {
 	// Disable the behaviour
@@ -318,41 +377,8 @@ void SRCBehaviourPanelList::Reset()
 TSharedRef<ITableRow> SRCBehaviourPanelList::OnGenerateWidgetForList(TSharedPtr<FRCBehaviourModel> InItem,
 	const TSharedRef<STableViewBase>& OwnerTable)
 {
-	return SNew(STableRow<TSharedPtr<FString>>, OwnerTable)
-		.Style(&RCPanelStyle->TableRowStyle)
-		[
-			SNew(SHorizontalBox)
-			+ SHorizontalBox::Slot()
-			[
-				InItem->GetWidget()
-			]
-
-			// Toggle Behaviour Button
-			+ SHorizontalBox::Slot()
-			.VAlign(VAlign_Center)
-			.AutoWidth()
-			.Padding(4.f, 0.f)
-			[
-				SNew(SCheckBox)
-				.AddMetaData<FTagMetaData>(FTagMetaData(TEXT("Toggle Behaviour")))
-				.ToolTipText(LOCTEXT("EditModeTooltip", "Enable/Disable this Behaviour.\nWhen a behaviour is disabled its Actions will not be processed when the Controller value changes"))
-				.HAlign(HAlign_Center)
-				.ForegroundColor(FLinearColor::White)
-				.Style(&RCPanelStyle->ToggleButtonStyle)
-				.IsChecked(this, &SRCBehaviourPanelList::IsBehaviourChecked, InItem)
-				.OnCheckStateChanged(this, &SRCBehaviourPanelList::OnToggleEnableBehaviour, InItem)
-			]
-		];
-}
-
-ECheckBoxState SRCBehaviourPanelList::IsBehaviourChecked(const TSharedPtr<FRCBehaviourModel> InBehaviourModel) const
-{
-	return InBehaviourModel.IsValid() && InBehaviourModel->IsBehaviourEnabled() ? ECheckBoxState::Checked : ECheckBoxState::Unchecked;
-}
-
-void SRCBehaviourPanelList::OnToggleEnableBehaviour(ECheckBoxState State, const TSharedPtr<FRCBehaviourModel> InBehaviourModel)
-{
-	SetIsBehaviourEnabled(InBehaviourModel, State == ECheckBoxState::Checked);
+	return SNew(SRCBehaviorPanelRow, OwnerTable, SharedThis(this), InItem)
+		.Style(&RCPanelStyle->TableRowStyle);
 }
 
 void SRCBehaviourPanelList::OnTreeSelectionChanged(TSharedPtr<FRCBehaviourModel> InItem, ESelectInfo::Type)
diff --git a/Engine/Plugins/VirtualProduction/RemoteControl/Source/RemoteControlUI/Private/UI/Behaviour/SRCBehaviourPanelList.h b/Engine/Plugins/VirtualProduction/RemoteControl/Source/RemoteControlUI/Private/UI/Behaviour/SRCBehaviourPanelList.h
--- a/Engine/Plugins/VirtualProduction/RemoteControl/Source/RemoteControlUI/Private/UI/Behaviour/SRCBehaviourPanelList.h
+++ b/Engine/Plugins/VirtualProduction/RemoteControl/Source/RemoteControlUI/Private/UI/Behaviour/SRCBehaviourPanelList.h
@@ -15,6 +15,7 @@ class SRCBehaviourPanel;
 class SRemoteControlPanel;
 class STableViewBase;
 class URemoteControlPreset;
+enum class EItemDropZone;
 template <typename ItemType> class SListView;
 
 /*
@@ -84,10 +85,16 @@ public:
 	/** Set a new controller item to update the behavior view based on the new controller */
 	void SetControllerItem(const TSharedPtr<FRCControllerModel>& InNewControllerItem);
 
-private:
-
 	/** Enables or Disables the currently selected behaviour */
-	void SetIsBehaviourEnabled(const bool bIsEnabled);
+	void SetAreSelectedBehaviorsEnabled(const bool bIsEnabled);
+
+	/** Enables or Disables the given behaviour */
+	void SetIsBehaviourEnabled(TSharedRef<FRCBehaviourModel> InBehaviorModel, const bool bIsEnabled);
+
+	/** Rearranges the behaviors based on the dropped behaviours. */
+	void ReorderBehaviorItems(TSharedRef<FRCBehaviourModel> InDroppedOnModel, EItemDropZone InDropZone, TArray<TSharedPtr<FRCBehaviourModel>> InDroppedModels);
+
+private:
 
 	/** Enables or Disables the passed behaviour */
 	void SetIsBehaviourEnabled(const TSharedPtr<FRCBehaviourModel>& InBehaviourModel, const bool bIsEnabled);
@@ -95,13 +102,7 @@ private:
 	void AddBehaviourToList(URCBehaviour* InBehaviour);
 
 	/** OnGenerateRow delegate for the Behaviours List View */
-	TSharedRef<ITableRow> OnGenerateWidgetForList( TSharedPtr<FRCBehaviourModel> InItem, const TSharedRef<STableViewBase>& OwnerTable );
-
-	/** Return whether or not the given behaviour is checked */
-	ECheckBoxState IsBehaviourChecked(const TSharedPtr<FRCBehaviourModel> InBehaviourModel) const;
-
-	/** Executed when toggling the behaviour state */
-	void OnToggleEnableBehaviour(ECheckBoxState State, TSharedPtr<FRCBehaviourModel> InBehaviourModel);
+	TSharedRef<ITableRow> OnGenerateWidgetForList(TSharedPtr<FRCBehaviourModel> InItem, const TSharedRef<STableViewBase>& OwnerTable);
 
 	/** OnSelectionChanged delegate for Behaviours List View */
 	void OnTreeSelectionChanged(TSharedPtr<FRCBehaviourModel> InItem , ESelectInfo::Type);
diff --git a/Engine/Plugins/VirtualProduction/RemoteControl/Source/RemoteControlUI/Private/UI/Controller/CustomControllers/SCustomStringListControllerListEditorWidget.cpp b/Engine/Plugins/VirtualProduction/RemoteControl/Source/RemoteControlUI/Private/UI/Controller/CustomControllers/SCustomStringListControllerListEditorWidget.cpp
--- a/Engine/Plugins/VirtualProduction/RemoteControl/Source/RemoteControlUI/Private/UI/Controller/CustomControllers/SCustomStringListControllerListEditorWidget.cpp
+++ b/Engine/Plugins/VirtualProduction/RemoteControl/Source/RemoteControlUI/Private/UI/Controller/CustomControllers/SCustomStringListControllerListEditorWidget.cpp
@@ -225,7 +225,9 @@ public:
 
 	FCustomStringListControllerEnumEditorRowDragDropOp(FName InEntry)
 		: Entry(InEntry)
-	{}
+	{
+		MouseCursor = EMouseCursor::GrabHandClosed;
+	}
 
 	virtual ~FCustomStringListControllerEnumEditorRowDragDropOp() = default;
 
diff --git a/Engine/Plugins/VirtualProduction/RemoteControl/Source/RemoteControlUI/Private/UI/Controller/SRCControllerPanelList.h b/Engine/Plugins/VirtualProduction/RemoteControl/Source/RemoteControlUI/Private/UI/Controller/SRCControllerPanelList.h
--- a/Engine/Plugins/VirtualProduction/RemoteControl/Source/RemoteControlUI/Private/UI/Controller/SRCControllerPanelList.h
+++ b/Engine/Plugins/VirtualProduction/RemoteControl/Source/RemoteControlUI/Private/UI/Controller/SRCControllerPanelList.h
@@ -43,6 +43,7 @@ public:
 	explicit FRCControllerPanelDragDrop(const TArray<FRCControllerPanelListItem>& InItems)
 		: Items(InItems)
 	{
+		MouseCursor = EMouseCursor::GrabHandClosed;
 	}
 
 	TArray<FRCControllerPanelListItem> GetItems() const
diff --git a/Engine/Plugins/VirtualProduction/RemoteControl/Source/RemoteControlUI/Private/UI/SRCPanelExposedEntity.h b/Engine/Plugins/VirtualProduction/RemoteControl/Source/RemoteControlUI/Private/UI/SRCPanelExposedEntity.h
--- a/Engine/Plugins/VirtualProduction/RemoteControl/Source/RemoteControlUI/Private/UI/SRCPanelExposedEntity.h
+++ b/Engine/Plugins/VirtualProduction/RemoteControl/Source/RemoteControlUI/Private/UI/SRCPanelExposedEntity.h
@@ -144,6 +144,8 @@ public:
 		: NodeId(InNodeId)
 		, SelectedEntities(InSelectedEntities)
 	{
+		MouseCursor = EMouseCursor::GrabHandClosed;
+
 		Algo::TransformIf(SelectedEntities, SelectedFieldsIds,
 			[] (const TSharedPtr<SRCPanelTreeNode>& InNode)
 			{
diff --git a/Engine/Plugins/VirtualProduction/RemoteControl/Source/RemoteControlUI/Private/UI/SRCPanelFieldGroup.h b/Engine/Plugins/VirtualProduction/RemoteControl/Source/RemoteControlUI/Private/UI/SRCPanelFieldGroup.h
--- a/Engine/Plugins/VirtualProduction/RemoteControl/Source/RemoteControlUI/Private/UI/SRCPanelFieldGroup.h
+++ b/Engine/Plugins/VirtualProduction/RemoteControl/Source/RemoteControlUI/Private/UI/SRCPanelFieldGroup.h
@@ -123,6 +123,8 @@ public:
 	FFieldGroupDragDropOp(TSharedPtr<SRCPanelGroup> InWidget, FGuid InId)
 		: Id(MoveTemp(InId))
 	{
+		MouseCursor = EMouseCursor::GrabHandClosed;
+
 		DecoratorWidget = SNew(SBorder)
 			.Padding(0.f)
 			.BorderImage(FAppStyle::Get().GetBrush("Graph.ConnectorFeedback.Border"))
