Remote Control: When adding a path behavor, automatically add an entry for the owning controller.

diff --git a/Engine/Plugins/VirtualProduction/RemoteControl/Source/RemoteControlLogic/Private/Behaviour/Builtin/Path/RCSetAssetByPathBehaviorNew.cpp b/Engine/Plugins/VirtualProduction/RemoteControl/Source/RemoteControlLogic/Private/Behaviour/Builtin/Path/RCSetAssetByPathBehaviorNew.cpp
--- a/Engine/Plugins/VirtualProduction/RemoteControl/Source/RemoteControlLogic/Private/Behaviour/Builtin/Path/RCSetAssetByPathBehaviorNew.cpp
+++ b/Engine/Plugins/VirtualProduction/RemoteControl/Source/RemoteControlLogic/Private/Behaviour/Builtin/Path/RCSetAssetByPathBehaviorNew.cpp
@@ -30,6 +30,18 @@ URCSetAssetByPathActionNew* URCSetAssetByPathBehaviorNew::AddPropertyBindAction(
 	return BindAction;
 }
 
+void URCSetAssetByPathBehaviorNew::Initialize()
+{
+	// Path already creates itself an empty entry. Let's add one bound to the owning controller. At least in name.
+
+	if (URCController* Controller = Cast<URCController>(GetOuter()))
+	{
+		FRCAssetPathElementNew& ControllerPathElement = PathStruct.AssetPath.AddDefaulted_GetRef();
+		ControllerPathElement.bIsInput = true;
+		ControllerPathElement.Path = Controller->DisplayName.ToString();
+	}
+}
+
 URCAction* URCSetAssetByPathBehaviorNew::AddAction(const TSharedRef<const FRemoteControlField> InRemoteControlField)
 {
 	const TSharedRef<const FRemoteControlProperty> RemoteControlProperty = StaticCastSharedRef<const FRemoteControlProperty>(InRemoteControlField);
diff --git a/Engine/Plugins/VirtualProduction/RemoteControl/Source/RemoteControlLogic/Private/Controller/RCController.cpp b/Engine/Plugins/VirtualProduction/RemoteControl/Source/RemoteControlLogic/Private/Controller/RCController.cpp
--- a/Engine/Plugins/VirtualProduction/RemoteControl/Source/RemoteControlLogic/Private/Controller/RCController.cpp
+++ b/Engine/Plugins/VirtualProduction/RemoteControl/Source/RemoteControlLogic/Private/Controller/RCController.cpp
@@ -40,6 +40,7 @@ URCBehaviour* URCController::AddBehaviour(TSubclassOf<URCBehaviourNode> InBehavi
 		return nullptr;
 	}
 
+	// This is the only time this is called - when the behavior is first created and added.
 	NewBehaviour->Initialize();
 	
 	Behaviors.Add(NewBehaviour);
diff --git a/Engine/Plugins/VirtualProduction/RemoteControl/Source/RemoteControlLogic/Public/Behaviour/Builtin/Path/RCSetAssetByPathBehaviorNew.h b/Engine/Plugins/VirtualProduction/RemoteControl/Source/RemoteControlLogic/Public/Behaviour/Builtin/Path/RCSetAssetByPathBehaviorNew.h
--- a/Engine/Plugins/VirtualProduction/RemoteControl/Source/RemoteControlLogic/Public/Behaviour/Builtin/Path/RCSetAssetByPathBehaviorNew.h
+++ b/Engine/Plugins/VirtualProduction/RemoteControl/Source/RemoteControlLogic/Public/Behaviour/Builtin/Path/RCSetAssetByPathBehaviorNew.h
@@ -79,6 +79,8 @@ class URCSetAssetByPathBehaviorNew : public URCBehaviour
 
 public:
 	//~ Begin URCBehaviour interface
+	/** Add default entries to the path. */
+	virtual void Initialize() override;
 	/** Add a Logic action using a remote control field as input */
 	REMOTECONTROLLOGIC_API virtual URCAction* AddAction(const TSharedRef<const FRemoteControlField> InRemoteControlField) override;
 	REMOTECONTROLLOGIC_API virtual bool SupportPropertyId() const override { return false; }
diff --git a/Engine/Plugins/VirtualProduction/RemoteControl/Source/RemoteControlLogic/Public/Behaviour/RCBehaviour.h b/Engine/Plugins/VirtualProduction/RemoteControl/Source/RemoteControlLogic/Public/Behaviour/RCBehaviour.h
--- a/Engine/Plugins/VirtualProduction/RemoteControl/Source/RemoteControlLogic/Public/Behaviour/RCBehaviour.h
+++ b/Engine/Plugins/VirtualProduction/RemoteControl/Source/RemoteControlLogic/Public/Behaviour/RCBehaviour.h
@@ -36,7 +36,7 @@ class REMOTECONTROLLOGIC_API URCBehaviour : public UObject
 public:
 	URCBehaviour();
 
-	/** Initialize behaviour functionality */
+	/** Initialize behaviour functionality. Called just once when created by a controller. */
 	virtual void Initialize() {}
 
 	/** Execute the behaviour */
