Remote Control: Invalid paths are now highlighted in red on the path behavior panel.

diff --git a/Engine/Plugins/VirtualProduction/RemoteControl/Source/RemoteControlUI/Private/UI/Behaviour/Builtin/Path/RCBehaviorSetAssetByPathModelNew.cpp b/Engine/Plugins/VirtualProduction/RemoteControl/Source/RemoteControlUI/Private/UI/Behaviour/Builtin/Path/RCBehaviorSetAssetByPathModelNew.cpp
--- a/Engine/Plugins/VirtualProduction/RemoteControl/Source/RemoteControlUI/Private/UI/Behaviour/Builtin/Path/RCBehaviorSetAssetByPathModelNew.cpp
+++ b/Engine/Plugins/VirtualProduction/RemoteControl/Source/RemoteControlUI/Private/UI/Behaviour/Builtin/Path/RCBehaviorSetAssetByPathModelNew.cpp
@@ -2,13 +2,16 @@
 
 #include "UI/Behaviour/Builtin/Path/RCBehaviorSetAssetByPathModelNew.h"
 
+#include "AssetRegistry/IAssetRegistry.h"
 #include "Behaviour/Builtin/Path/RCSetAssetByPathBehaviorNew.h"
 #include "Controller/RCController.h"
 #include "IDetailTreeNode.h"
 #include "IPropertyRowGenerator.h"
+#include "Misc/Paths.h"
 #include "Modules/ModuleManager.h"
 #include "PropertyEditorModule.h"
 #include "RemoteControlPreset.h"
+#include "Styling/StyleColors.h"
 #include "UI/Action/Path/RCActionSetAssetByPathModelNew.h"
 #include "UI/Action/SRCActionPanelList.h"
 #include "UI/Behaviour/Builtin/Path/SRCBehaviorSetAssetByPathNew.h"
@@ -24,7 +27,9 @@ FRCSetAssetByPathBehaviorModelNew::FRCSetAssetByPathBehaviorModelNew(URCSetAsset
 	SetAssetByPathBehaviorWeakPtr = InSetAssetByPathBehavior;
 
 	PreviewPathWidget = SNew(STextBlock)
-		.AutoWrapText(true);
+		.AutoWrapText(true)
+		.ColorAndOpacity(FStyleColors::Error);
+
 	PathArrayWidget = SNew(SBox);
 
 	FPropertyRowGeneratorArgs ArgsRowGeneratorArray;
@@ -134,13 +139,48 @@ TSharedPtr<SRCLogicPanelListBase> FRCSetAssetByPathBehaviorModelNew::GetActionsL
 void FRCSetAssetByPathBehaviorModelNew::RefreshPreview() const
 {
 	FString PreviewText = TEXT("Path is invalid!");
+	bool bInternal = true;
 
 	if (URCSetAssetByPathBehaviorNew* SetAssetByPathBehavior = SetAssetByPathBehaviorWeakPtr.Get())
 	{
 		PreviewText = SetAssetByPathBehavior->GetCurrentPath();
+		bInternal = SetAssetByPathBehavior->bInternal;
 	}
 
 	PreviewPathWidget->SetText(FText::FromString(PreviewText));
+
+	bool bIsValidPath = false;
+
+	if (!bInternal)
+	{
+		bIsValidPath = FPaths::FileExists(PreviewText);
+	}
+	else
+	{
+		// Specifying a path without the internal package name works, but cannot be found by the asset registry.
+		// Double up the last path element if it does not contain a dot to complete the path.
+		FAssetData AssetData;
+		FString CorrectedPath = PreviewText;
+
+		if (!CorrectedPath.EndsWith(TEXT("/")))
+		{
+			const FString Filename = FPaths::GetBaseFilename(PreviewText);
+			int32 DotIndex;
+
+			if (!Filename.FindChar(TEXT('.'), DotIndex))
+			{
+				CorrectedPath += TEXT(".") + Filename;
+			}
+		}
+
+		bIsValidPath = IAssetRegistry::Get()->TryGetAssetByObjectPath(CorrectedPath, AssetData) == UE::AssetRegistry::EExists::Exists;
+	}
+	
+	PreviewPathWidget->SetColorAndOpacity(
+		bIsValidPath
+			? FStyleColors::Foreground
+			: FStyleColors::Error
+	);
 }
 
 void FRCSetAssetByPathBehaviorModelNew::NotifyPostChange(const FPropertyChangedEvent& InPropertyChangedEvent, FEditPropertyChain* InPropertyThatChanged)
