Remote Control: Path behavior path preview highlight redux

diff --git a/Engine/Plugins/VirtualProduction/RemoteControl/Source/RemoteControlUI/Private/UI/Behaviour/Builtin/Path/RCBehaviorSetAssetByPathModelNew.cpp b/Engine/Plugins/VirtualProduction/RemoteControl/Source/RemoteControlUI/Private/UI/Behaviour/Builtin/Path/RCBehaviorSetAssetByPathModelNew.cpp
--- a/Engine/Plugins/VirtualProduction/RemoteControl/Source/RemoteControlUI/Private/UI/Behaviour/Builtin/Path/RCBehaviorSetAssetByPathModelNew.cpp
+++ b/Engine/Plugins/VirtualProduction/RemoteControl/Source/RemoteControlUI/Private/UI/Behaviour/Builtin/Path/RCBehaviorSetAssetByPathModelNew.cpp
@@ -164,7 +164,7 @@ void FRCSetAssetByPathBehaviorModelNew::RefreshPreview() const
 
 		if (!CorrectedPath.EndsWith(TEXT("/")))
 		{
-			const FString Filename = FPaths::GetBaseFilename(PreviewText);
+			const FString Filename = FPaths::GetCleanFilename(PreviewText);
 			int32 DotIndex;
 
 			if (!Filename.FindChar(TEXT('.'), DotIndex))
