Remote Control: Path behavior row editor header now exists outside of the scroll box for greater clarity.

diff --git a/Engine/Plugins/VirtualProduction/RemoteControl/Source/RemoteControlUI/Private/UI/Behaviour/Builtin/Path/RCBehaviorSetAssetByPathModelNew.cpp b/Engine/Plugins/VirtualProduction/RemoteControl/Source/RemoteControlUI/Private/UI/Behaviour/Builtin/Path/RCBehaviorSetAssetByPathModelNew.cpp
--- a/Engine/Plugins/VirtualProduction/RemoteControl/Source/RemoteControlUI/Private/UI/Behaviour/Builtin/Path/RCBehaviorSetAssetByPathModelNew.cpp
+++ b/Engine/Plugins/VirtualProduction/RemoteControl/Source/RemoteControlUI/Private/UI/Behaviour/Builtin/Path/RCBehaviorSetAssetByPathModelNew.cpp
@@ -85,15 +85,7 @@ TSharedRef<SWidget> FRCSetAssetByPathBehaviorModelNew::GetPropertyWidget()
 		.Padding(FMargin(3.f, 3.f))
 		.AutoHeight()
 		[
-			SNew(SBox)
-			.MaxDesiredHeight(150.f)
-			[
-				SNew(SScrollBox)
-				+ SScrollBox::Slot()
-				[
-					PathArrayWidget->AsShared()
-				]
-			]
+			PathArrayWidget->AsShared()
 		];
 
 	FieldWidget->AddSlot()
@@ -291,6 +283,26 @@ void FRCSetAssetByPathBehaviorModelNew::RefreshPathAndPreview()
 	{
 		Behavior->RefreshPathArray();
 	}
+
+	TSharedPtr<SBox> Header;
+	TSharedPtr<SVerticalBox> Rows;
+
+	TSharedRef<SVerticalBox> Content = SNew(SVerticalBox)
+		+ SVerticalBox::Slot()
+		.AutoHeight()
+		[
+			SAssignNew(Header, SBox)
+			.Padding(FMargin(3.0f, 2.0f))
+		]
+		+ SVerticalBox::Slot()
+		.MaxHeight(150.f)
+		[
+			SNew(SScrollBox)
+			+ SScrollBox::Slot()
+			[
+				SAssignNew(Rows, SVerticalBox)
+			]
+		];
 	
 	for (const TSharedPtr<IDetailTreeNode>& DetailTreeNodeArray : DetailTreeNodeWeakPtrArray)
 	{
@@ -308,6 +320,7 @@ void FRCSetAssetByPathBehaviorModelNew::RefreshPathAndPreview()
 			{
 				TSharedPtr<SHorizontalBox> CurrentHorizontalBox = SNew(SHorizontalBox);
 
+				// Header/Control row
 				if (Counter == 0)
 				{
 					// Add NameWidget only for the ArrayProperty not for its entries
@@ -315,6 +328,15 @@ void FRCSetAssetByPathBehaviorModelNew::RefreshPathAndPreview()
 					[
 						NodeWidgets.NameWidget.ToSharedRef()
 					];
+
+					CurrentHorizontalBox->AddSlot()
+						.FillWidth(1.f)
+						[
+							NodeWidgets.ValueWidget.ToSharedRef()
+						];
+
+					Header->SetContent(CurrentHorizontalBox.ToSharedRef());
+					continue;
 				}
 
 				CurrentHorizontalBox->AddSlot()
@@ -323,18 +345,15 @@ void FRCSetAssetByPathBehaviorModelNew::RefreshPathAndPreview()
 						NodeWidgets.ValueWidget.ToSharedRef()
 					];
 
-				FieldArrayWidget->AddSlot()
-					.Padding(FMargin(3.0f, 2.0f))
+				Rows->AddSlot()
 					.AutoHeight()
-					.VAlign(VAlign_Center)
 					[
 						CurrentHorizontalBox.ToSharedRef()
 					];
 			}
 			else if (NodeWidgets.WholeRowWidget)
 			{
-				FieldArrayWidget->AddSlot()
-					.Padding(FMargin(3.0f, 2.0f))
+				Rows->AddSlot()
 					.AutoHeight()
 					[
 						NodeWidgets.WholeRowWidget.ToSharedRef()
@@ -344,7 +363,7 @@ void FRCSetAssetByPathBehaviorModelNew::RefreshPathAndPreview()
 	}
 
 	RefreshPreview();
-	PathArrayWidget->SetContent(FieldArrayWidget.ToSharedRef());
+	PathArrayWidget->SetContent(Content);
 }
 
 #undef LOCTEXT_NAMESPACE
