Remote Control: Path behavior "get asset path" now adds a slash on the end of the provided path

diff --git a/Engine/Plugins/VirtualProduction/RemoteControl/Source/RemoteControlUI/Private/UI/Customizations/RCAssetPathElementCustomization.cpp b/Engine/Plugins/VirtualProduction/RemoteControl/Source/RemoteControlUI/Private/UI/Customizations/RCAssetPathElementCustomization.cpp
--- a/Engine/Plugins/VirtualProduction/RemoteControl/Source/RemoteControlUI/Private/UI/Customizations/RCAssetPathElementCustomization.cpp
+++ b/Engine/Plugins/VirtualProduction/RemoteControl/Source/RemoteControlUI/Private/UI/Customizations/RCAssetPathElementCustomization.cpp
@@ -185,7 +185,11 @@ FReply FRCAssetPathElementCustomization::OnGetAssetFromSelectionClicked() const
 	PathString.FindLastChar('/', IndexOfLast);
 	if (IndexOfLast != INDEX_NONE)
 	{
-		PathString.RemoveAt(IndexOfLast, PathString.Len() - IndexOfLast);
+		if (IndexOfLast != (PathString.Len() - 1))
+		{
+			// Keep the slash
+			PathString.RemoveAt(IndexOfLast + 1, PathString.Len() - IndexOfLast - 1);
+		}
 	}
 	else
 	{
