Remote Control: Disables the 'create controller' button for path behavior if the path element is not empty.

diff --git a/Engine/Plugins/VirtualProduction/RemoteControl/Source/RemoteControlUI/Private/UI/Customizations/RCAssetPathElementCustomization.cpp b/Engine/Plugins/VirtualProduction/RemoteControl/Source/RemoteControlUI/Private/UI/Customizations/RCAssetPathElementCustomization.cpp
--- a/Engine/Plugins/VirtualProduction/RemoteControl/Source/RemoteControlUI/Private/UI/Customizations/RCAssetPathElementCustomization.cpp
+++ b/Engine/Plugins/VirtualProduction/RemoteControl/Source/RemoteControlUI/Private/UI/Customizations/RCAssetPathElementCustomization.cpp
@@ -108,6 +108,7 @@ void FRCAssetPathElementCustomization::CustomizeHeader(TSharedRef<IPropertyHandl
 			[
 				SNew(SButton)
 				.ButtonStyle(FAppStyle::Get(), "SimpleButton")
+				.IsEnabled(this, &FRCAssetPathElementCustomization::IsCreateControllerButtonEnabled)
 				.OnClicked(this, &FRCAssetPathElementCustomization::OnCreateControllerButtonClicked)
 				.ToolTip(CreateControllerToolTipWidget)
 				.IsFocusable(false)
@@ -209,6 +210,14 @@ FReply FRCAssetPathElementCustomization::OnGetAssetFromSelectionClicked() const
 	return FReply::Handled();
 }
 
+bool FRCAssetPathElementCustomization::IsCreateControllerButtonEnabled() const
+{
+	FString Value;
+	PathHandle->GetValue(Value);
+
+	return Value.IsEmpty();
+}
+
 FReply FRCAssetPathElementCustomization::OnCreateControllerButtonClicked() const
 {
 	if (PropertyUtilities.IsValid() && ArrayEntryHandle.IsValid())
diff --git a/Engine/Plugins/VirtualProduction/RemoteControl/Source/RemoteControlUI/Private/UI/Customizations/RCAssetPathElementCustomization.h b/Engine/Plugins/VirtualProduction/RemoteControl/Source/RemoteControlUI/Private/UI/Customizations/RCAssetPathElementCustomization.h
--- a/Engine/Plugins/VirtualProduction/RemoteControl/Source/RemoteControlUI/Private/UI/Customizations/RCAssetPathElementCustomization.h
+++ b/Engine/Plugins/VirtualProduction/RemoteControl/Source/RemoteControlUI/Private/UI/Customizations/RCAssetPathElementCustomization.h
@@ -33,6 +33,9 @@ private:
 	/** Callback called when you click on the arrow button next to the entries to retrieve the path of the selected Asset */
 	FReply OnGetAssetFromSelectionClicked() const;
 
+	/** Only lets you click the button if the path is empty. */
+	bool IsCreateControllerButtonEnabled() const;
+
 	/** Callback called when you click on the plus button next to the entries to create a controller associated with the current RC Input entry */
 	FReply OnCreateControllerButtonClicked() const;
 
