Remote Control: "Fill path from selected asset" will now look for selected exposed properties if there are no asset selected.
- GetValue and GetValueInArray of the IRemoteControlPropertyHandle now correctly take UObject* references.

diff --git a/Engine/Plugins/VirtualProduction/RemoteControl/Source/RemoteControl/Private/RemoteControlPropertyHandle.cpp b/Engine/Plugins/VirtualProduction/RemoteControl/Source/RemoteControl/Private/RemoteControlPropertyHandle.cpp
--- a/Engine/Plugins/VirtualProduction/RemoteControl/Source/RemoteControl/Private/RemoteControlPropertyHandle.cpp
+++ b/Engine/Plugins/VirtualProduction/RemoteControl/Source/RemoteControl/Private/RemoteControlPropertyHandle.cpp
@@ -1886,17 +1886,9 @@ bool FRemoteControlPropertyHandleObject::SetValue(UObject* InValue)
 	return SetPropertyValue(*this, InValue);
 }
 
-bool FRemoteControlPropertyHandleObject::GetValue(UObject* OutValue) const
+bool FRemoteControlPropertyHandleObject::GetValue(UObject*& OutValue) const
 {
-	if (const TSharedPtr<FRemoteControlProperty>& RemoteControlPropertyPtr = GetRCProperty())
-	{
-		if (const TSharedPtr<IRemoteControlPropertyHandle>& PropertyHandle = RemoteControlPropertyPtr->GetPropertyHandle())
-		{
-			return PropertyHandle->GetValue(OutValue);
-		}
-	}
-
-	return false;
+	return GetPropertyValue<FObjectProperty>(*this, OutValue);
 }
 
 bool FRemoteControlPropertyHandleObject::SetValueInArray(UObject* InValue, int32 InIndex)
@@ -1911,7 +1903,7 @@ bool FRemoteControlPropertyHandleObject::SetValueInArray(UObject* InValue, int32
 	return false;
 }
 
-bool FRemoteControlPropertyHandleObject::GetValueInArray(UObject* OutValue, int32 InIndex) const
+bool FRemoteControlPropertyHandleObject::GetValueInArray(UObject*& OutValue, int32 InIndex) const
 {
 	if (InIndex < 0)
 	{
diff --git a/Engine/Plugins/VirtualProduction/RemoteControl/Source/RemoteControl/Private/RemoteControlPropertyHandle.h b/Engine/Plugins/VirtualProduction/RemoteControl/Source/RemoteControl/Private/RemoteControlPropertyHandle.h
--- a/Engine/Plugins/VirtualProduction/RemoteControl/Source/RemoteControl/Private/RemoteControlPropertyHandle.h
+++ b/Engine/Plugins/VirtualProduction/RemoteControl/Source/RemoteControl/Private/RemoteControlPropertyHandle.h
@@ -96,8 +96,8 @@ public:
 
 	virtual bool SetValue(UObject* InValue) override { return false; }
 	virtual bool SetValueInArray(UObject* InValue, int32 InIndex) override { return false; }
-	virtual bool GetValue(UObject* OutValue) const override { return false; }
-	virtual bool GetValueInArray(UObject* OutValue, int32 InIndex) const override { return false; }
+	virtual bool GetValue(UObject*& OutValue) const override { return false; }
+	virtual bool GetValueInArray(UObject*& OutValue, int32 InIndex) const override { return false; }
 
 	virtual TSharedPtr<IRemoteControlPropertyHandle> GetChildHandle(FName ChildName, bool bRecurse = true) const override;
 	virtual TSharedPtr<IRemoteControlPropertyHandle> GetChildHandle(int32 Index) const override;
@@ -612,11 +612,11 @@ public:
 
 	//~ Begin IRemoteControlPropertyHandle Interface
 	virtual bool SetValue(UObject* InValue) override;
-	virtual bool GetValue(UObject* OutValue) const override;
+	virtual bool GetValue(UObject*& OutValue) const override;
 
 	/** This is currently used for the specific case in which we are setting a material inside the override materials of e.g. a Static Mesh*/
 	virtual bool SetValueInArray(UObject* InValue, int32 InIndex) override;
-	virtual bool GetValueInArray(UObject* OutValue, int32 InIndex) const override;
+	virtual bool GetValueInArray(UObject*& OutValue, int32 InIndex) const override;
 	//~ End IRemoteControlPropertyHandle Interface
 
 	/** Set a Static Mesh value using the handled Object */
diff --git a/Engine/Plugins/VirtualProduction/RemoteControl/Source/RemoteControl/Public/IRemoteControlPropertyHandle.h b/Engine/Plugins/VirtualProduction/RemoteControl/Source/RemoteControl/Public/IRemoteControlPropertyHandle.h
--- a/Engine/Plugins/VirtualProduction/RemoteControl/Source/RemoteControl/Public/IRemoteControlPropertyHandle.h
+++ b/Engine/Plugins/VirtualProduction/RemoteControl/Source/RemoteControl/Public/IRemoteControlPropertyHandle.h
@@ -69,8 +69,8 @@ public:
 	virtual bool GetValue(FRotator& OutValue) const = 0;
 	virtual bool GetValue(FColor& OutValue) const = 0;
 	virtual bool GetValue(FLinearColor& OutValue) const = 0;
-	virtual bool GetValue(UObject* OutValue) const = 0;
-	virtual bool GetValueInArray(UObject* OutValue, int32 InIndex) const = 0;
+	virtual bool GetValue(UObject*& OutValue) const = 0;
+	virtual bool GetValueInArray(UObject*& OutValue, int32 InIndex) const = 0;
 	
 	/**
 	 * Sets the typed value of a property.
diff --git a/Engine/Plugins/VirtualProduction/RemoteControl/Source/RemoteControlUI/Private/RemoteControlUIModule.cpp b/Engine/Plugins/VirtualProduction/RemoteControl/Source/RemoteControlUI/Private/RemoteControlUIModule.cpp
--- a/Engine/Plugins/VirtualProduction/RemoteControl/Source/RemoteControlUI/Private/RemoteControlUIModule.cpp
+++ b/Engine/Plugins/VirtualProduction/RemoteControl/Source/RemoteControlUI/Private/RemoteControlUIModule.cpp
@@ -1124,7 +1124,6 @@ void FRemoteControlUIModule::RegisterStructCustomizations()
 	}
 
 	PropertyEditorModule.RegisterCustomPropertyTypeLayout(FRCAssetPathElement::StaticStruct()->GetFName(), FOnGetPropertyTypeCustomizationInstance::CreateStatic(&FRCAssetPathElementCustomization::MakeInstance));
-	PropertyEditorModule.RegisterCustomPropertyTypeLayout(FRCAssetPathElementNew::StaticStruct()->GetFName(), FOnGetPropertyTypeCustomizationInstance::CreateStatic(&FRCAssetPathElementCustomization::MakeInstance));
 	PropertyEditorModule.RegisterCustomPropertyTypeLayout(FRCNetworkAddress::StaticStruct()->GetFName(), FOnGetPropertyTypeCustomizationInstance::CreateStatic(&FNetworkAddressCustomization::MakeInstance));
 	PropertyEditorModule.RegisterCustomPropertyTypeLayout(FRCPassphrase::StaticStruct()->GetFName(), FOnGetPropertyTypeCustomizationInstance::CreateStatic(&FPassphraseCustomization::MakeInstance));
 }
diff --git a/Engine/Plugins/VirtualProduction/RemoteControl/Source/RemoteControlUI/Private/UI/Behaviour/Builtin/Path/RCBehaviorSetAssetByPathModelNew.cpp b/Engine/Plugins/VirtualProduction/RemoteControl/Source/RemoteControlUI/Private/UI/Behaviour/Builtin/Path/RCBehaviorSetAssetByPathModelNew.cpp
--- a/Engine/Plugins/VirtualProduction/RemoteControl/Source/RemoteControlUI/Private/UI/Behaviour/Builtin/Path/RCBehaviorSetAssetByPathModelNew.cpp
+++ b/Engine/Plugins/VirtualProduction/RemoteControl/Source/RemoteControlUI/Private/UI/Behaviour/Builtin/Path/RCBehaviorSetAssetByPathModelNew.cpp
@@ -16,6 +16,7 @@
 #include "UI/Action/SRCActionPanelList.h"
 #include "UI/Behaviour/Builtin/Path/SRCBehaviorSetAssetByPathNew.h"
 #include "UI/Controller/SRCControllerPanel.h"
+#include "UI/Customizations/RCAssetPathElementCustomization.h"
 #include "UI/SRemoteControlPanel.h"
 #include "Widgets/SNullWidget.h"
 
@@ -25,7 +26,10 @@ FRCSetAssetByPathBehaviorModelNew::FRCSetAssetByPathBehaviorModelNew(URCSetAsset
 	: FRCBehaviourModel(InSetAssetByPathBehavior, InRemoteControlPanel)
 {
 	SetAssetByPathBehaviorWeakPtr = InSetAssetByPathBehavior;
+}
 
+void FRCSetAssetByPathBehaviorModelNew::Initialize()
+{
 	PreviewPathWidget = SNew(STextBlock)
 		.AutoWrapText(true)
 		.ColorAndOpacity(FStyleColors::Error);
@@ -36,9 +40,19 @@ FRCSetAssetByPathBehaviorModelNew::FRCSetAssetByPathBehaviorModelNew(URCSetAsset
 	ArgsRowGeneratorArray.NotifyHook = this;
 	PropertyRowGeneratorArray = FModuleManager::GetModuleChecked<FPropertyEditorModule>("PropertyEditor").CreatePropertyRowGenerator(ArgsRowGeneratorArray);
 
-	if (InSetAssetByPathBehavior)
+	PropertyRowGeneratorArray->RegisterInstancedCustomPropertyTypeLayout(
+		FRCAssetPathElementNew::StaticStruct()->GetFName(),
+		FOnGetPropertyTypeCustomizationInstance::CreateLambda(
+			[ThisWeak = SharedThis(this).ToWeakPtr()]()
+			{
+				return FRCAssetPathElementCustomization::MakeInstance(ThisWeak.Pin());
+			}
+		)
+	);
+
+	if (URCSetAssetByPathBehaviorNew* SetAssetByPathBehavior = SetAssetByPathBehaviorWeakPtr.Get())
 	{
-		PropertyRowGeneratorArray->SetStructure(MakeShareable(new FStructOnScope(FRCSetAssetPathNew::StaticStruct(), (uint8*)&InSetAssetByPathBehavior->PathStruct)));
+		PropertyRowGeneratorArray->SetStructure(MakeShareable(new FStructOnScope(FRCSetAssetPathNew::StaticStruct(), (uint8*)&SetAssetByPathBehavior->PathStruct)));
 
 		PropertyRowGeneratorArray->OnFinishedChangingProperties().AddLambda([this](const FPropertyChangedEvent& InEvent)
 			{
diff --git a/Engine/Plugins/VirtualProduction/RemoteControl/Source/RemoteControlUI/Private/UI/Behaviour/Builtin/Path/RCBehaviorSetAssetByPathModelNew.h b/Engine/Plugins/VirtualProduction/RemoteControl/Source/RemoteControlUI/Private/UI/Behaviour/Builtin/Path/RCBehaviorSetAssetByPathModelNew.h
--- a/Engine/Plugins/VirtualProduction/RemoteControl/Source/RemoteControlUI/Private/UI/Behaviour/Builtin/Path/RCBehaviorSetAssetByPathModelNew.h
+++ b/Engine/Plugins/VirtualProduction/RemoteControl/Source/RemoteControlUI/Private/UI/Behaviour/Builtin/Path/RCBehaviorSetAssetByPathModelNew.h
@@ -23,6 +23,8 @@ class FRCSetAssetByPathBehaviorModelNew : public FRCBehaviourModel, public FNoti
 public:
 	explicit FRCSetAssetByPathBehaviorModelNew(URCSetAssetByPathBehaviorNew* InSetAssetByPathBehavior, const TSharedPtr<SRemoteControlPanel> InRemoteControlPanel);
 
+	void Initialize();
+
 	/** Returns true if this behavior have a details widget or false if not*/
 	virtual bool HasBehaviourDetailsWidget() override;
 
diff --git a/Engine/Plugins/VirtualProduction/RemoteControl/Source/RemoteControlUI/Private/UI/Behaviour/SRCBehaviourPanelList.cpp b/Engine/Plugins/VirtualProduction/RemoteControl/Source/RemoteControlUI/Private/UI/Behaviour/SRCBehaviourPanelList.cpp
--- a/Engine/Plugins/VirtualProduction/RemoteControl/Source/RemoteControlUI/Private/UI/Behaviour/SRCBehaviourPanelList.cpp
+++ b/Engine/Plugins/VirtualProduction/RemoteControl/Source/RemoteControlUI/Private/UI/Behaviour/SRCBehaviourPanelList.cpp
@@ -343,7 +343,10 @@ void SRCBehaviourPanelList::AddBehaviourToList(URCBehaviour* InBehaviour)
 	}
 	else if (URCSetAssetByPathBehaviorNew* SetAssetByPathBehavior = Cast<URCSetAssetByPathBehaviorNew>(InBehaviour))
 	{
-		BehaviourItems.Add(MakeShared<FRCSetAssetByPathBehaviorModelNew>(SetAssetByPathBehavior, RemoteControlPanel));
+		TSharedRef<FRCSetAssetByPathBehaviorModelNew> NewBehavior = MakeShared<FRCSetAssetByPathBehaviorModelNew>(SetAssetByPathBehavior, RemoteControlPanel);
+		NewBehavior->Initialize();
+
+		BehaviourItems.Add(NewBehavior);
 	}
 	else
 	{
diff --git a/Engine/Plugins/VirtualProduction/RemoteControl/Source/RemoteControlUI/Private/UI/Controller/SRCControllerPanel.cpp b/Engine/Plugins/VirtualProduction/RemoteControl/Source/RemoteControlUI/Private/UI/Controller/SRCControllerPanel.cpp
--- a/Engine/Plugins/VirtualProduction/RemoteControl/Source/RemoteControlUI/Private/UI/Controller/SRCControllerPanel.cpp
+++ b/Engine/Plugins/VirtualProduction/RemoteControl/Source/RemoteControlUI/Private/UI/Controller/SRCControllerPanel.cpp
@@ -253,6 +253,24 @@ TArray<TSharedPtr<FRCLogicModeBase>> SRCControllerPanel::GetSelectedLogicItems()
 	return {};
 }
 
+TArray<TSharedPtr<FRCControllerModel>> SRCControllerPanel::GetSelectedControllers() const
+{
+	TArray<TSharedPtr<FRCLogicModeBase>> SelectedLogicModes = GetSelectedLogicItems();
+
+	TArray<TSharedPtr<FRCControllerModel>> SelectedControllers;
+	SelectedControllers.Reserve(SelectedLogicModes.Num());
+
+	for (const TSharedPtr<FRCLogicModeBase>& LogicMode : SelectedLogicModes)
+	{
+		if (LogicMode.IsValid() && LogicMode->GetModelType() == FRCControllerModel::ControllerModelName)
+		{
+			SelectedControllers.Add(StaticCastSharedPtr<FRCControllerModel>(LogicMode));
+		}
+	}
+
+	return SelectedControllers;
+}
+
 int32 SRCControllerPanel::NumControllerItems() const
 {
 	if (ControllerPanelList.IsValid())
diff --git a/Engine/Plugins/VirtualProduction/RemoteControl/Source/RemoteControlUI/Private/UI/Controller/SRCControllerPanel.h b/Engine/Plugins/VirtualProduction/RemoteControl/Source/RemoteControlUI/Private/UI/Controller/SRCControllerPanel.h
--- a/Engine/Plugins/VirtualProduction/RemoteControl/Source/RemoteControlUI/Private/UI/Controller/SRCControllerPanel.h
+++ b/Engine/Plugins/VirtualProduction/RemoteControl/Source/RemoteControlUI/Private/UI/Controller/SRCControllerPanel.h
@@ -89,6 +89,9 @@ public:
 	/** Returns the UI items currently selected by the user (if any). To be implemented per child panel*/
 	virtual TArray<TSharedPtr<FRCLogicModeBase>> GetSelectedLogicItems() const override;
 
+	/** Get the currently selected controller models. */
+	TArray<TSharedPtr<FRCControllerModel>> GetSelectedControllers() const;
+
 	/** Retrieve the current number of Controllers */
 	int32 NumControllerItems() const;
 
diff --git a/Engine/Plugins/VirtualProduction/RemoteControl/Source/RemoteControlUI/Private/UI/Customizations/RCAssetPathElementCustomization.cpp b/Engine/Plugins/VirtualProduction/RemoteControl/Source/RemoteControlUI/Private/UI/Customizations/RCAssetPathElementCustomization.cpp
--- a/Engine/Plugins/VirtualProduction/RemoteControl/Source/RemoteControlUI/Private/UI/Customizations/RCAssetPathElementCustomization.cpp
+++ b/Engine/Plugins/VirtualProduction/RemoteControl/Source/RemoteControlUI/Private/UI/Customizations/RCAssetPathElementCustomization.cpp
@@ -8,8 +8,14 @@
 #include "Framework/Application/SlateApplication.h"
 #include "IDetailChildrenBuilder.h"
 #include "IPropertyUtilities.h"
+#include "IRemoteControlPropertyHandle.h"
 #include "PropertyCustomizationHelpers.h"
+#include "UI/Behaviour/Builtin/Path/RCBehaviorSetAssetByPathModelNew.h"
 #include "UI/RemoteControlPanelStyle.h"
+#include "UI/SRCPanelExposedEntitiesList.h"
+#include "UI/SRCPanelExposedField.h"
+#include "UI/SRCPanelTreeNode.h"
+#include "UI/SRemoteControlPanel.h"
 #include "Widgets/Images/SImage.h"
 #include "Widgets/Input/SButton.h"
 #include "Widgets/Input/SCheckBox.h"
@@ -22,7 +28,18 @@
 
 TSharedRef<IPropertyTypeCustomization> FRCAssetPathElementCustomization::MakeInstance()
 {
-	return MakeShared<FRCAssetPathElementCustomization>();
+	return MakeShared<FRCAssetPathElementCustomization>(nullptr);
+}
+
+TSharedRef<IPropertyTypeCustomization> FRCAssetPathElementCustomization::MakeInstance(TSharedPtr<FRCSetAssetByPathBehaviorModelNew> InPathBehaviorModelNew)
+{
+	return MakeShared<FRCAssetPathElementCustomization>(InPathBehaviorModelNew);
+}
+
+FRCAssetPathElementCustomization::FRCAssetPathElementCustomization(TSharedPtr<FRCSetAssetByPathBehaviorModelNew> InPathBehaviorModelNew)
+	: PathBehaviorModelNewWeak(InPathBehaviorModelNew)
+{
+
 }
 
 void FRCAssetPathElementCustomization::CustomizeHeader(TSharedRef<IPropertyHandle> InPropertyHandle, FDetailWidgetRow& InHeaderRow, IPropertyTypeCustomizationUtils& InCustomizationUtils)
@@ -166,24 +183,23 @@ void FRCAssetPathElementCustomization::OnCheckStateChanged(ECheckBoxState InNewS
 
 FReply FRCAssetPathElementCustomization::OnGetAssetFromSelectionClicked() const
 {
-	TArray<FAssetData> AssetData;
-	GEditor->GetContentBrowserSelections(AssetData);
-	if (AssetData.Num() <= 0)
+	FString PathString = GetSelectedAssetPath();
+
+	if (PathString.IsEmpty())
 	{
 		return FReply::Handled();
 	}
-				
-	// Clear it, in case it is an already used one.
-	// Use the first one in the Array
-	int32 IndexOfLast;
-	const UObject* SelectedAsset = AssetData[0].GetAsset();
-	FString PathString = SelectedAsset->GetPathName();
 
 	// Remove the initial Game
 	PathString.RemoveFromStart("/Game/");
 
+	// Clear it, in case it is an already used one.
+	// Use the first one in the Array
+	int32 IndexOfLast;
+
 	// Remove anything after the last /
 	PathString.FindLastChar('/', IndexOfLast);
+
 	if (IndexOfLast != INDEX_NONE)
 	{
 		if (IndexOfLast != (PathString.Len() - 1))
@@ -240,4 +256,95 @@ int32 FRCAssetPathElementCustomization::OnGetWidgetSwitcherIndex() const
 	return IsChecked() == ECheckBoxState::Checked ? 1 : 0;
 }
 
+FString FRCAssetPathElementCustomization::GetSelectedAssetPath() const
+{
+	FString AssetPath = GetSelectedAssetPath_ContentBrowser();
+
+	if (!AssetPath.IsEmpty())
+	{
+		return AssetPath;
+	}
+
+	return GetSelectedAssetPath_EntityList();
+}
+
+FString FRCAssetPathElementCustomization::GetSelectedAssetPath_ContentBrowser() const
+{
+	TArray<FAssetData> AssetData;
+	GEditor->GetContentBrowserSelections(AssetData);
+
+	if (AssetData.Num() > 0)
+	{
+		const UObject* SelectedAsset = AssetData[0].GetAsset();
+		return SelectedAsset->GetPathName();
+	}
+
+	return FString();
+}
+
+FString FRCAssetPathElementCustomization::GetSelectedAssetPath_EntityList() const
+{
+	TSharedPtr<FRCSetAssetByPathBehaviorModelNew> PathBehaviorModel = PathBehaviorModelNewWeak.Pin();
+
+	if (!PathBehaviorModel.IsValid())
+	{
+		return FString();
+	}
+
+	TSharedPtr<SRemoteControlPanel> RCPanel = PathBehaviorModel->GetRemoteControlPanel();
+
+	if (!RCPanel.IsValid())
+	{
+		return FString();
+	}
+
+	TSharedPtr<SRCPanelExposedEntitiesList> EntityList = RCPanel->GetEntityList();
+
+	if (!EntityList.IsValid())
+	{
+		return FString();
+	}
+
+	TArray<TSharedPtr<SRCPanelTreeNode>> SelectedEntities = EntityList->GetSelectedEntities();
+
+	for (const TSharedPtr<SRCPanelTreeNode>& SelectedEntity : SelectedEntities)
+	{
+		if (!SelectedEntity.IsValid())
+		{
+			continue;
+		}
+
+		if (SelectedEntity->GetRCType() != SRCPanelTreeNode::Field)
+		{
+			continue;
+		}
+
+		TSharedPtr<FRemoteControlField> RCField = StaticCastSharedPtr<SRCPanelExposedField>(SelectedEntity)->GetRemoteControlField().Pin();
+
+		if (!RCField.IsValid() || RCField->FieldType != EExposedFieldType::Property || !RCField->IsBound())
+		{
+			continue;
+		}
+
+		TSharedPtr<IRemoteControlPropertyHandle> PropertyHandle = StaticCastSharedPtr<FRemoteControlProperty>(RCField)->GetPropertyHandle();
+
+		if (!PropertyHandle.IsValid())
+		{
+			continue;
+		}
+
+		UObject* Object = nullptr;
+
+		if (PropertyHandle->GetValue(Object))
+		{
+			if (Object && Object->IsAsset())
+			{
+				return Object->GetPathName();
+			}
+		}
+	}
+
+	return FString();
+}
+
 #undef LOCTEXT_NAMESPACE
diff --git a/Engine/Plugins/VirtualProduction/RemoteControl/Source/RemoteControlUI/Private/UI/Customizations/RCAssetPathElementCustomization.h b/Engine/Plugins/VirtualProduction/RemoteControl/Source/RemoteControlUI/Private/UI/Customizations/RCAssetPathElementCustomization.h
--- a/Engine/Plugins/VirtualProduction/RemoteControl/Source/RemoteControlUI/Private/UI/Customizations/RCAssetPathElementCustomization.h
+++ b/Engine/Plugins/VirtualProduction/RemoteControl/Source/RemoteControlUI/Private/UI/Customizations/RCAssetPathElementCustomization.h
@@ -7,6 +7,7 @@
 #include "Styling/SlateTypes.h"
 
 class FDetailWidgetRow;
+class FRCSetAssetByPathBehaviorModelNew;
 class IDetailChildrenBuilder;
 class IPropertyHandle;
 
@@ -14,6 +15,9 @@ class FRCAssetPathElementCustomization : public IPropertyTypeCustomization
 {
 public:
 	static TSharedRef<IPropertyTypeCustomization> MakeInstance();
+	static TSharedRef<IPropertyTypeCustomization> MakeInstance(TSharedPtr<FRCSetAssetByPathBehaviorModelNew> InPathBehaviorModelNew);
+
+	FRCAssetPathElementCustomization(TSharedPtr<FRCSetAssetByPathBehaviorModelNew> InPathBehaviorModelNew);
 
 	//~ Begin IPropertyTypeCustomization
 	virtual void CustomizeHeader(TSharedRef<IPropertyHandle> InPropertyHandle, FDetailWidgetRow& InHeaderRow, IPropertyTypeCustomizationUtils& InCustomizationUtils) override;
@@ -42,7 +46,14 @@ private:
 	/** Retrieve the current widget switcher index */
 	int32 OnGetWidgetSwitcherIndex() const;
 
+	FString GetSelectedAssetPath() const;
+
+	FString GetSelectedAssetPath_ContentBrowser() const;
+
+	FString GetSelectedAssetPath_EntityList() const;
+
 private:
+	TWeakPtr<FRCSetAssetByPathBehaviorModelNew> PathBehaviorModelNewWeak;
 	TSharedPtr<SWidget> PathWidget;
 	TSharedPtr<IPropertyUtilities> PropertyUtilities;
 	TSharedPtr<IPropertyHandle> ArrayEntryHandle;
diff --git a/Engine/Plugins/VirtualProduction/RemoteControl/Source/RemoteControlUI/Private/UI/SRemoteControlPanel.h b/Engine/Plugins/VirtualProduction/RemoteControl/Source/RemoteControlUI/Private/UI/SRemoteControlPanel.h
--- a/Engine/Plugins/VirtualProduction/RemoteControl/Source/RemoteControlUI/Private/UI/SRemoteControlPanel.h
+++ b/Engine/Plugins/VirtualProduction/RemoteControl/Source/RemoteControlUI/Private/UI/SRemoteControlPanel.h
@@ -43,11 +43,14 @@ class SBox;
 class SClassViewer;
 class SComboButton;
 class SRCActionPanel;
+class SRCBehaviourPanel;
+class SRCControllerPanel;
 class SRCLogicPanelBase;
 class SRCPanelDrawer;
 class SRCPanelExposedEntitiesList;
 class SRCPanelFilter;
 class SRCPanelFunctionPicker;
+class SRCSignaturePanel;
 class SRemoteControlPanel;
 class SSearchBox;
 class STextBlock;
@@ -463,13 +466,13 @@ private:
 	// ~ Remote Control Logic Panels ~
 
 	/** Controller panel UI widget for Remote Control Logic*/
-	TSharedPtr<class SRCControllerPanel> ControllerPanel;
+	TSharedPtr<SRCControllerPanel> ControllerPanel;
 	/** Behaviour panel UI widget for Remote Control Logic*/
-	TSharedPtr<class SRCBehaviourPanel> BehaviorPanel;
+	TSharedPtr<SRCBehaviourPanel> BehaviorPanel;
 	/** Action panel UI widget for Remote Control Logic*/
-	TSharedPtr<class SRCActionPanel> ActionPanel;
+	TSharedPtr<SRCActionPanel> ActionPanel;
 	/** Signature panel UI widget */
-	TSharedPtr<class SRCSignaturePanel> SignaturePanel;
+	TSharedPtr<SRCSignaturePanel> SignaturePanel;
 
 	/** LogicClipboardItems - Holds the items copied from a Logic panel
 	*
