Remote Control: Extended path behavior selector button to have right click override

diff --git a/Engine/Plugins/VirtualProduction/RemoteControl/Source/RemoteControlUI/Private/UI/Customizations/RCAssetPathElementCustomization.cpp b/Engine/Plugins/VirtualProduction/RemoteControl/Source/RemoteControlUI/Private/UI/Customizations/RCAssetPathElementCustomization.cpp
--- a/Engine/Plugins/VirtualProduction/RemoteControl/Source/RemoteControlUI/Private/UI/Customizations/RCAssetPathElementCustomization.cpp
+++ b/Engine/Plugins/VirtualProduction/RemoteControl/Source/RemoteControlUI/Private/UI/Customizations/RCAssetPathElementCustomization.cpp
@@ -26,6 +26,89 @@
 
 #define LOCTEXT_NAMESPACE "RCAssetPathElementCustomization"
 
+class SRCAssetPathSelectorButton : public SButton
+{
+public:
+	SLATE_BEGIN_ARGS(SRCAssetPathSelectorButton)
+		{}
+		SLATE_EVENT(FOnGetContent, OnGetMenuContent)
+	SLATE_END_ARGS()
+
+	void Construct(const FArguments& InArgs, TSharedRef<FRCAssetPathElementCustomization> InCustomization)
+	{
+		CustomizationWeak = InCustomization;
+		OnGetMenuContent = InArgs._OnGetMenuContent;
+
+		const TSharedRef<SToolTip> GetAssetPathToolTipWidget = SNew(SToolTip)
+			.Text(LOCTEXT("RCGetAssetPathButton_Tooltip", 
+				"Get the path of selected assets"
+				"\n\n"
+				"- Left click will attempt to get the first selected asset in the content browser and then the first valid selected exposed entity."
+				"\n"
+				"- Right click will open a menu allowing you to specify where to look, overriding the default priority."
+			));
+
+		SButton::Construct(
+			SButton::FArguments()
+			.ButtonStyle(FAppStyle::Get(), "SimpleButton")
+			.OnClicked(this, &SRCAssetPathSelectorButton::OnLeftClicked)
+			.ToolTip(GetAssetPathToolTipWidget)
+			.IsFocusable(false)
+			.ContentPadding(0)
+			.Content()
+			[
+				SNew(SImage)
+				.Image(FAppStyle::GetBrush("Icons.Use"))
+				.ColorAndOpacity(FSlateColor::UseForeground())
+			]
+		);
+	}
+
+	virtual FReply OnMouseButtonDown(const FGeometry& InMyGeometry, const FPointerEvent& InMouseEvent) override
+	{
+		if (!IsEnabled())
+		{
+			return FReply::Unhandled();
+		}
+
+		if (InMouseEvent.GetEffectingButton() == EKeys::RightMouseButton)
+		{
+			if (OnGetMenuContent.IsBound())
+			{
+				if (TSharedPtr<SWidget> MenuContent = OnGetMenuContent.Execute())
+				{
+					FSlateApplication& SlateApplication = FSlateApplication::Get();
+
+					FWidgetPath MyWidgetPath;
+					SlateApplication.GeneratePathToWidgetUnchecked(AsShared(), MyWidgetPath);
+
+					const FVector2D CursorLocation = SlateApplication.GetCursorPos();
+
+					SlateApplication.PushMenu(AsShared(), MyWidgetPath, MenuContent.ToSharedRef(), CursorLocation + FVector2D(16), FPopupTransitionEffect::ContextMenu);
+				}
+			}
+				
+			return FReply::Handled();
+		}
+
+		return SButton::OnMouseButtonDown(InMyGeometry, InMouseEvent);
+	}
+
+protected:
+	TWeakPtr<FRCAssetPathElementCustomization> CustomizationWeak;
+	FOnGetContent OnGetMenuContent;
+
+	FReply OnLeftClicked()
+	{
+		if (TSharedPtr<FRCAssetPathElementCustomization> Customization = CustomizationWeak.Pin())
+		{
+			return Customization->OnGetAssetFromSelectionClicked();
+		}
+
+		return FReply::Unhandled();
+	}
+};
+
 TSharedRef<IPropertyTypeCustomization> FRCAssetPathElementCustomization::MakeInstance()
 {
 	return MakeShared<FRCAssetPathElementCustomization>(nullptr);
@@ -56,9 +139,6 @@ void FRCAssetPathElementCustomization::CustomizeHeader(TSharedRef<IPropertyHandl
 
 	PathHandle->SetOnPropertyValueChanged(FSimpleDelegate::CreateSP(this, &FRCAssetPathElementCustomization::OnPathChanged));
 
-	const TSharedRef<SToolTip> GetAssetPathToolTipWidget = SNew(SToolTip)
-		.Text(LOCTEXT("RCGetAssetPathButton_Tooltip", "Get the path of the currently first selected asset in the content browser and set it to the current path"));
-
 	const TSharedRef<SToolTip> CreateControllerToolTipWidget = SNew(SToolTip)
 		.Text(LOCTEXT("RCCreateController_Tooltip", "Create a controller for the given RC Input path entry"));
 
@@ -103,21 +183,12 @@ void FRCAssetPathElementCustomization::CustomizeHeader(TSharedRef<IPropertyHandl
 		[
 			SNew(SWidgetSwitcher)
 			.WidgetIndex(this, &FRCAssetPathElementCustomization::OnGetWidgetSwitcherIndex)
+
 			// [0] Get Current Selected Asset Path Button
 			+ SWidgetSwitcher::Slot()
 			[
-				SNew(SButton)
-				.ButtonStyle(FAppStyle::Get(), "SimpleButton")
-				.OnClicked(this, &FRCAssetPathElementCustomization::OnGetAssetFromSelectionClicked)
-				.ToolTip(GetAssetPathToolTipWidget)
-				.IsFocusable(false)
-				.ContentPadding(0)
-				.Content()
-				[
-					SNew(SImage)
-					.Image(FAppStyle::GetBrush("Icons.Use"))
-					.ColorAndOpacity(FSlateColor::UseForeground())
-				]
+				SNew(SRCAssetPathSelectorButton, SharedThis(this))
+				.OnGetMenuContent(FOnGetContent::CreateSP(this, &FRCAssetPathElementCustomization::GetPathSelectorMenuContent))
 			]
 
 			// [1] Create Controller button
@@ -145,7 +216,7 @@ void FRCAssetPathElementCustomization::CustomizeChildren(TSharedRef<IPropertyHan
 {
 }
 
-void FRCAssetPathElementCustomization::OnPathChanged() const
+void FRCAssetPathElementCustomization::OnPathChanged()
 {
 	if (!PathHandle.IsValid() || !IsInputHandle.IsValid())
 	{
@@ -171,7 +242,7 @@ ECheckBoxState FRCAssetPathElementCustomization::IsChecked() const
 	return ReturnValue;
 }
 
-void FRCAssetPathElementCustomization::OnCheckStateChanged(ECheckBoxState InNewState) const
+void FRCAssetPathElementCustomization::OnCheckStateChanged(ECheckBoxState InNewState)
 {
 	if (!IsInputHandle.IsValid())
 	{
@@ -181,48 +252,10 @@ void FRCAssetPathElementCustomization::OnCheckStateChanged(ECheckBoxState InNewS
 	IsInputHandle->SetValue(bIsInput);
 }
 
-FReply FRCAssetPathElementCustomization::OnGetAssetFromSelectionClicked() const
+FReply FRCAssetPathElementCustomization::OnGetAssetFromSelectionClicked()
 {
-	FString PathString = GetSelectedAssetPath();
-
-	if (PathString.IsEmpty())
-	{
-		return FReply::Handled();
-	}
-
-	// Remove the initial Game
-	PathString.RemoveFromStart("/Game/");
-
-	// Clear it, in case it is an already used one.
-	// Use the first one in the Array
-	int32 IndexOfLast;
-
-	// Remove anything after the last /
-	PathString.FindLastChar('/', IndexOfLast);
-
-	if (IndexOfLast != INDEX_NONE)
-	{
-		if (IndexOfLast != (PathString.Len() - 1))
-		{
-			// Keep the slash
-			PathString.RemoveAt(IndexOfLast + 1, PathString.Len() - IndexOfLast - 1);
-		}
-	}
-	else
-	{
-		// if the Index is -1 then it means that we are selecting an asset already in the topmost folder
-		// So we clear the string since it will just contains the AssetName
-		PathString.Empty();
-	}
-
-	if (IsInputHandle.IsValid())
-	{
-		IsInputHandle->SetValue(false);
-	}
-	if (PathHandle.IsValid())
-	{
-		PathHandle->SetValue(PathString);
-	}
+	SetAssetFromPath(GetSelectedAssetPath());
+	
 	return FReply::Handled();
 }
 
@@ -234,7 +267,7 @@ bool FRCAssetPathElementCustomization::IsCreateControllerButtonEnabled() const
 	return Value.IsEmpty();
 }
 
-FReply FRCAssetPathElementCustomization::OnCreateControllerButtonClicked() const
+FReply FRCAssetPathElementCustomization::OnCreateControllerButtonClicked()
 {
 	if (PropertyUtilities.IsValid() && ArrayEntryHandle.IsValid())
 	{
@@ -347,4 +380,89 @@ FString FRCAssetPathElementCustomization::GetSelectedAssetPath_EntityList() cons
 	return FString();
 }
 
+TSharedRef<SWidget> FRCAssetPathElementCustomization::GetPathSelectorMenuContent()
+{
+	FMenuBuilder MenuBuilder(/* Should close after selection */ true, /* Command list */ nullptr);
+
+	MenuBuilder.BeginSection("SelectSource", LOCTEXT("SelectSource", "Select Source"));
+
+	FUIAction ContentBrowserSource(
+		FExecuteAction::CreateSP(this, &FRCAssetPathElementCustomization::SetAssetPathFromContentBrowser)
+	);
+
+	MenuBuilder.AddMenuEntry(
+		LOCTEXT("ContentBrowser", "Content Browser"),
+		LOCTEXT("ContentBrowserTooltip", "Take the asset path from the selected assets in the content browser."),
+		FSlateIcon(),
+		ContentBrowserSource
+	);
+
+	FUIAction EntityListSource(
+		FExecuteAction::CreateSP(this, &FRCAssetPathElementCustomization::SetAssetPathFromEntityList)
+	);
+
+	MenuBuilder.AddMenuEntry(
+		LOCTEXT("EntityList", "Entity List"),
+		LOCTEXT("EntityListTooltip", "Take the asset path from the first valid selected exposed entity."),
+		FSlateIcon(),
+		EntityListSource
+	);
+
+	MenuBuilder.EndSection();
+
+	return MenuBuilder.MakeWidget();
+}
+
+void FRCAssetPathElementCustomization::SetAssetPathFromContentBrowser()
+{
+	SetAssetFromPath(GetSelectedAssetPath_ContentBrowser());
+}
+
+void FRCAssetPathElementCustomization::SetAssetPathFromEntityList()
+{
+	SetAssetFromPath(GetSelectedAssetPath_EntityList());
+}
+
+void FRCAssetPathElementCustomization::SetAssetFromPath(FString InPath)
+{
+	if (InPath.IsEmpty())
+	{
+		return;
+	}
+
+	// Remove the initial Game
+	InPath.RemoveFromStart("/Game/");
+
+	// Clear it, in case it is an already used one.
+	// Use the first one in the Array
+	int32 IndexOfLast;
+
+	// Remove anything after the last /
+	InPath.FindLastChar('/', IndexOfLast);
+
+	if (IndexOfLast != INDEX_NONE)
+	{
+		if (IndexOfLast != (InPath.Len() - 1))
+		{
+			// Keep the slash
+			InPath.RemoveAt(IndexOfLast + 1, InPath.Len() - IndexOfLast - 1);
+		}
+	}
+	else
+	{
+		// if the Index is -1 then it means that we are selecting an asset already in the topmost folder
+		// So we clear the string since it will just contains the AssetName
+		InPath.Empty();
+	}
+
+	if (IsInputHandle.IsValid())
+	{
+		IsInputHandle->SetValue(false);
+	}
+	if (PathHandle.IsValid())
+	{
+		PathHandle->SetValue(InPath);
+	}
+}
+
 #undef LOCTEXT_NAMESPACE
diff --git a/Engine/Plugins/VirtualProduction/RemoteControl/Source/RemoteControlUI/Private/UI/Customizations/RCAssetPathElementCustomization.h b/Engine/Plugins/VirtualProduction/RemoteControl/Source/RemoteControlUI/Private/UI/Customizations/RCAssetPathElementCustomization.h
--- a/Engine/Plugins/VirtualProduction/RemoteControl/Source/RemoteControlUI/Private/UI/Customizations/RCAssetPathElementCustomization.h
+++ b/Engine/Plugins/VirtualProduction/RemoteControl/Source/RemoteControlUI/Private/UI/Customizations/RCAssetPathElementCustomization.h
@@ -13,6 +13,8 @@ class IPropertyHandle;
 
 class FRCAssetPathElementCustomization : public IPropertyTypeCustomization
 {
+	friend class SRCAssetPathSelectorButton;
+
 public:
 	static TSharedRef<IPropertyTypeCustomization> MakeInstance();
 	static TSharedRef<IPropertyTypeCustomization> MakeInstance(TSharedPtr<FRCSetAssetByPathBehaviorModelNew> InPathBehaviorModelNew);
@@ -26,22 +28,22 @@ public:
 
 private:
 	/** Called when the path change */
-	void OnPathChanged() const;
+	void OnPathChanged();
 
 	/** Returns if the CheckBox is checked or not */
 	ECheckBoxState IsChecked() const;
 
 	/** Callback called when you click on the CheckBox to enable/disable it */
-	void OnCheckStateChanged(ECheckBoxState InNewState) const;
+	void OnCheckStateChanged(ECheckBoxState InNewState);
 
 	/** Callback called when you click on the arrow button next to the entries to retrieve the path of the selected Asset */
-	FReply OnGetAssetFromSelectionClicked() const;
+	FReply OnGetAssetFromSelectionClicked();
 
 	/** Only lets you click the button if the path is empty. */
 	bool IsCreateControllerButtonEnabled() const;
 
 	/** Callback called when you click on the plus button next to the entries to create a controller associated with the current RC Input entry */
-	FReply OnCreateControllerButtonClicked() const;
+	FReply OnCreateControllerButtonClicked();
 
 	/** Retrieve the current widget switcher index */
 	int32 OnGetWidgetSwitcherIndex() const;
@@ -52,6 +54,14 @@ private:
 
 	FString GetSelectedAssetPath_EntityList() const;
 
+	TSharedRef<SWidget> GetPathSelectorMenuContent();
+
+	void SetAssetPathFromContentBrowser();
+
+	void SetAssetPathFromEntityList();
+
+	void SetAssetFromPath(FString InPath);
+
 private:
 	TWeakPtr<FRCSetAssetByPathBehaviorModelNew> PathBehaviorModelNewWeak;
 	TSharedPtr<SWidget> PathWidget;
