Remote Control: Path elements in the Path Behavior are now rearrangeable by dragging

diff --git a/Engine/Plugins/VirtualProduction/RemoteControl/Source/RemoteControlUI/Private/UI/Behaviour/Builtin/Path/RCBehaviorSetAssetByPathModelNew.cpp b/Engine/Plugins/VirtualProduction/RemoteControl/Source/RemoteControlUI/Private/UI/Behaviour/Builtin/Path/RCBehaviorSetAssetByPathModelNew.cpp
--- a/Engine/Plugins/VirtualProduction/RemoteControl/Source/RemoteControlUI/Private/UI/Behaviour/Builtin/Path/RCBehaviorSetAssetByPathModelNew.cpp
+++ b/Engine/Plugins/VirtualProduction/RemoteControl/Source/RemoteControlUI/Private/UI/Behaviour/Builtin/Path/RCBehaviorSetAssetByPathModelNew.cpp
@@ -11,14 +11,19 @@
 #include "Modules/ModuleManager.h"
 #include "PropertyEditorModule.h"
 #include "RemoteControlPreset.h"
+#include "ScopedTransaction.h"
 #include "Styling/StyleColors.h"
 #include "UI/Action/Path/RCActionSetAssetByPathModelNew.h"
 #include "UI/Action/SRCActionPanelList.h"
 #include "UI/Behaviour/Builtin/Path/SRCBehaviorSetAssetByPathNew.h"
+#include "UI/Behaviour/Builtin/Path/SRCBehaviorSetAssetByPathNewElementRow.h"
 #include "UI/Controller/SRCControllerPanel.h"
 #include "UI/Customizations/RCAssetPathElementCustomization.h"
 #include "UI/SRemoteControlPanel.h"
+#include "Widgets/SBoxPanel.h"
 #include "Widgets/SNullWidget.h"
+#include "Widgets/Views/SListView.h"
+#include "Widgets/Views/STableRow.h"
 
 #define LOCTEXT_NAMESPACE "RCSetAssetByPathBehaviorModelNew"
 
@@ -36,6 +41,11 @@ void FRCSetAssetByPathBehaviorModelNew::Initialize()
 
 	PathArrayWidget = SNew(SBox);
 
+	ElementsView = SNew(SListView<TSharedPtr<FRCPathBehaviorElementRow>>)
+		.ListItemsSource(&ElementItems)
+		.SelectionMode(ESelectionMode::Single)
+		.OnGenerateRow(this, &FRCSetAssetByPathBehaviorModelNew::OnGenerateWidgetForList);
+
 	FPropertyRowGeneratorArgs ArgsRowGeneratorArray;
 	ArgsRowGeneratorArray.NotifyHook = this;
 	PropertyRowGeneratorArray = FModuleManager::GetModuleChecked<FPropertyEditorModule>("PropertyEditor").CreatePropertyRowGenerator(ArgsRowGeneratorArray);
@@ -197,6 +207,56 @@ void FRCSetAssetByPathBehaviorModelNew::NotifyPostChange(const FPropertyChangedE
 	}
 }
 
+void FRCSetAssetByPathBehaviorModelNew::ReorderElementItems(int32 InDroppedOnIndex, EItemDropZone InDropZone, int32 InDroppedIndex)
+{
+	if (InDroppedOnIndex == InDroppedIndex)
+	{
+		return;
+	}
+
+	URCSetAssetByPathBehaviorNew* SetAssetByPathBehavior = SetAssetByPathBehaviorWeakPtr.Get();
+
+	if (!SetAssetByPathBehavior)
+	{
+		return;
+	}
+
+	if (!SetAssetByPathBehavior->PathStruct.AssetPath.IsValidIndex(InDroppedOnIndex)
+		|| !SetAssetByPathBehavior->PathStruct.AssetPath.IsValidIndex(InDroppedIndex))
+	{
+		return;
+	}
+
+	FScopedTransaction Transaction(LOCTEXT("ReorderPath", "Reorder Path"));
+	SetAssetByPathBehavior->Modify();
+
+	const bool bIsAbove = InDropZone == EItemDropZone::AboveItem;
+	TArray<FRCAssetPathElementNew> NewAssetPath;
+
+	for (int32 Index = 0; Index < SetAssetByPathBehavior->PathStruct.AssetPath.Num(); ++Index)
+	{
+		if (Index == InDroppedOnIndex && bIsAbove)
+		{
+			NewAssetPath.Add(SetAssetByPathBehavior->PathStruct.AssetPath[InDroppedIndex]);
+		}
+
+		if (Index != InDroppedIndex)
+		{
+			NewAssetPath.Add(SetAssetByPathBehavior->PathStruct.AssetPath[Index]);
+		}
+
+		if (Index == InDroppedOnIndex && !bIsAbove)
+		{
+			NewAssetPath.Add(SetAssetByPathBehavior->PathStruct.AssetPath[InDroppedIndex]);
+		}
+	}
+
+	SetAssetByPathBehavior->PathStruct.AssetPath = NewAssetPath;
+	SetAssetByPathBehavior->Execute();
+
+	RefreshPathAndPreview();
+}
+
 void FRCSetAssetByPathBehaviorModelNew::OnAssetPathFinishedChangingProperties(const FPropertyChangedEvent& InEvent)
 {
 	if (InEvent.ChangeType == EPropertyChangeType::ValueSet)
@@ -311,13 +371,11 @@ void FRCSetAssetByPathBehaviorModelNew::RefreshPathAndPreview()
 		+ SVerticalBox::Slot()
 		.MaxHeight(150.f)
 		[
-			SNew(SScrollBox)
-			+ SScrollBox::Slot()
-			[
-				SAssignNew(Rows, SVerticalBox)
-			]
+			ElementsView.ToSharedRef()
 		];
 	
+	ElementItems.Empty();
+
 	for (const TSharedPtr<IDetailTreeNode>& DetailTreeNodeArray : DetailTreeNodeWeakPtrArray)
 	{
 		const TSharedPtr<IDetailTreeNode> PinnedNode = DetailTreeNodeArray;
@@ -332,7 +390,7 @@ void FRCSetAssetByPathBehaviorModelNew::RefreshPathAndPreview()
 
 			if (NodeWidgets.ValueWidget)
 			{
-				TSharedPtr<SHorizontalBox> CurrentHorizontalBox = SNew(SHorizontalBox);
+				TSharedRef<SHorizontalBox> CurrentHorizontalBox = SNew(SHorizontalBox);
 
 				// Header/Control row
 				if (Counter == 0)
@@ -349,7 +407,7 @@ void FRCSetAssetByPathBehaviorModelNew::RefreshPathAndPreview()
 							NodeWidgets.ValueWidget.ToSharedRef()
 						];
 
-					Header->SetContent(CurrentHorizontalBox.ToSharedRef());
+					Header->SetContent(CurrentHorizontalBox);
 					continue;
 				}
 
@@ -359,25 +417,24 @@ void FRCSetAssetByPathBehaviorModelNew::RefreshPathAndPreview()
 						NodeWidgets.ValueWidget.ToSharedRef()
 					];
 
-				Rows->AddSlot()
-					.AutoHeight()
-					[
-						CurrentHorizontalBox.ToSharedRef()
-					];
+				ElementItems.Add(MakeShared<FRCPathBehaviorElementRow>(Counter - 1, CurrentHorizontalBox));
 			}
 			else if (NodeWidgets.WholeRowWidget)
 			{
-				Rows->AddSlot()
-					.AutoHeight()
-					[
-						NodeWidgets.WholeRowWidget.ToSharedRef()
-					];
+				ElementItems.Add(MakeShared<FRCPathBehaviorElementRow>(Counter - 1, NodeWidgets.WholeRowWidget.ToSharedRef()));
 			}
 		}
 	}
 
+	ElementsView->RequestListRefresh();
+
 	RefreshPreview();
 	PathArrayWidget->SetContent(Content);
 }
 
+TSharedRef<ITableRow> FRCSetAssetByPathBehaviorModelNew::OnGenerateWidgetForList(TSharedPtr<FRCPathBehaviorElementRow> InItem, const TSharedRef<STableViewBase>& InOwnerTable)
+{
+	return SNew(SRCBehaviorSetAssetByPathNewElementRow, InOwnerTable, SharedThis(this), InItem);
+}
+
 #undef LOCTEXT_NAMESPACE
diff --git a/Engine/Plugins/VirtualProduction/RemoteControl/Source/RemoteControlUI/Private/UI/Behaviour/Builtin/Path/RCBehaviorSetAssetByPathModelNew.h b/Engine/Plugins/VirtualProduction/RemoteControl/Source/RemoteControlUI/Private/UI/Behaviour/Builtin/Path/RCBehaviorSetAssetByPathModelNew.h
--- a/Engine/Plugins/VirtualProduction/RemoteControl/Source/RemoteControlUI/Private/UI/Behaviour/Builtin/Path/RCBehaviorSetAssetByPathModelNew.h
+++ b/Engine/Plugins/VirtualProduction/RemoteControl/Source/RemoteControlUI/Private/UI/Behaviour/Builtin/Path/RCBehaviorSetAssetByPathModelNew.h
@@ -7,8 +7,13 @@
 #include "UI/Behaviour/RCBehaviourModel.h"
 
 class IDetailTreeNode;
+class ITableRow;
 class SBox;
+class STableViewBase;
 class URCSetAssetByPathBehaviorNew;
+enum class EItemDropZone;
+struct FRCPathBehaviorElementRow;
+template<typename T> class SListView;
 
 /*
  * ~ FRCSetAssetByPathBehaviorModelNew ~
@@ -48,6 +53,8 @@ public:
 
 	virtual void NotifyPostChange(const FPropertyChangedEvent& InPropertyChangedEvent, class FEditPropertyChain* InPropertyThatChanged) override;
 
+	void ReorderElementItems(int32 InDroppedOnIndex, EItemDropZone InDropZone, int32 InDroppedIndex);
+
 private:
 	/** Callback called when the AssetPath properties change to either refresh or reconstruct it updating the path values */
 	void OnAssetPathFinishedChangingProperties(const FPropertyChangedEvent& InEvent);
@@ -58,6 +65,9 @@ private:
 	/** Regenerates and creates a new PathArray Widget if changed */
 	void RefreshPathAndPreview();
 
+	/** Creates a row for the elements view. */
+	TSharedRef<ITableRow> OnGenerateWidgetForList(TSharedPtr<FRCPathBehaviorElementRow> InItem, const TSharedRef<STableViewBase>& InOwnerTable);
+
 private:
 	/** The SetAssetByPath Behaviour associated with this Model */
 	TWeakObjectPtr<URCSetAssetByPathBehaviorNew> SetAssetByPathBehaviorWeakPtr;
@@ -65,6 +75,12 @@ private:
 	/** Pointer to the Widget holding the PathArray created for this behavior */
 	TSharedPtr<SBox> PathArrayWidget;
 
+	/** List view containing the path elements. */
+	TSharedPtr<SListView<TSharedPtr<FRCPathBehaviorElementRow>>> ElementsView;
+
+	/** Data for each row in the list view. */
+	TArray<TSharedPtr<FRCPathBehaviorElementRow>> ElementItems;
+
 	/** Pointer to the Preview Text Widget for the Path */
 	TSharedPtr<STextBlock> PreviewPathWidget;
 
diff --git a/Engine/Plugins/VirtualProduction/RemoteControl/Source/RemoteControlUI/Private/UI/Behaviour/Builtin/Path/SRCBehaviorSetAssetByPathNewElementRow.cpp b/Engine/Plugins/VirtualProduction/RemoteControl/Source/RemoteControlUI/Private/UI/Behaviour/Builtin/Path/SRCBehaviorSetAssetByPathNewElementRow.cpp
new file mode 100644
--- /dev/null
+++ b/Engine/Plugins/VirtualProduction/RemoteControl/Source/RemoteControlUI/Private/UI/Behaviour/Builtin/Path/SRCBehaviorSetAssetByPathNewElementRow.cpp
@@ -0,0 +1,150 @@
+// Copyright Epic Games, Inc. All Rights Reserved.
+
+#include "UI/Behaviour/Builtin/Path/SRCBehaviorSetAssetByPathNewElementRow.h"
+
+#include "DragAndDrop/DecoratedDragDropOp.h"
+#include "Styles/SlateBrushTemplates.h"
+#include "Styling/RemoteControlStyles.h"
+#include "UI/Behaviour/Builtin/Path/RCBehaviorSetAssetByPathModelNew.h"
+#include "UI/RemoteControlPanelStyle.h"
+
+#define LOCTEXT_NAMESPACE "SRCBehaviorPanelRow"
+
+/*
+* ~ FRCBehaviorSetAssetByPathNewElementDragDrop ~
+* Facilitates drag-drop operation
+*/
+class FRCBehaviorSetAssetByPathNewElementDragDrop final : public FDecoratedDragDropOp
+{
+public:
+	DRAG_DROP_OPERATOR_TYPE(FRCBehaviorSetAssetByPathNewElementDragDrop, FDecoratedDragDropOp)
+
+	explicit FRCBehaviorSetAssetByPathNewElementDragDrop(int32 InIndex)
+		: Index(InIndex)
+	{
+		MouseCursor = EMouseCursor::GrabHandClosed;
+	}
+
+	int32 GetIndex() const
+	{
+		return Index;
+	}
+
+private:
+	int32 Index;
+};
+
+void SRCBehaviorSetAssetByPathNewElementRow::Construct(const FArguments& InArgs, const TSharedRef<STableViewBase>& InOwnerTableView, TSharedRef<FRCSetAssetByPathBehaviorModelNew> InSetAssetByPathBehaviorModelNew, RCBehaviorSetAssetByPathNewElementListItem InElementItem)
+{
+	ElementItem = InElementItem;
+	SetAssetByPathBehaviorModelNewWeak = InSetAssetByPathBehaviorModelNew;
+
+	SetCursor(EMouseCursor::GrabHand);
+
+	const FRCPanelStyle* RCPanelStyle = &FRemoteControlPanelStyle::Get()->GetWidgetStyle<FRCPanelStyle>("RemoteControlPanel.MinorPanel");
+
+	STableRow::Construct(
+		STableRow::FArguments()
+		.Padding(0.f)
+		.OnDragDetected(this, &SRCBehaviorSetAssetByPathNewElementRow::OnRowDragDetected)
+		.OnCanAcceptDrop(this, &SRCBehaviorSetAssetByPathNewElementRow::CanAcceptDrop)
+		.OnAcceptDrop(this, &SRCBehaviorSetAssetByPathNewElementRow::OnAcceptDrop)
+		.ShowWires(false)
+		.Content()
+		[
+			SNew(SHorizontalBox)
+			+ SHorizontalBox::Slot()
+			.VAlign(VAlign_Center)
+			.AutoWidth()
+			.Padding(4.f, 0.f)
+			[
+				SNew(SImage)
+				.Image(FSlateBrushTemplates::DragHandle())
+			]
+
+			+ SHorizontalBox::Slot()
+			[
+				InElementItem->Widget
+			]
+		],
+		InOwnerTableView
+	);
+}
+
+int32 SRCBehaviorSetAssetByPathNewElementRow::GetElementIndex() const
+{
+	return ElementItem->Index;
+}
+
+TOptional<EItemDropZone> SRCBehaviorSetAssetByPathNewElementRow::CanAcceptDrop(const FDragDropEvent& InDragDropEvent, EItemDropZone InDropZone, RCBehaviorSetAssetByPathNewElementListItem InItem)
+{
+	if (!SetAssetByPathBehaviorModelNewWeak.IsValid() || !ElementItem.IsValid())
+	{
+		return {};
+	}
+
+	TSharedPtr<FRCBehaviorSetAssetByPathNewElementDragDrop> PathBehaviorElementDragDrop = InDragDropEvent.GetOperationAs<FRCBehaviorSetAssetByPathNewElementDragDrop>();
+
+	if (!PathBehaviorElementDragDrop.IsValid())
+	{
+		return {};
+	}
+
+	if (ElementItem->Index == PathBehaviorElementDragDrop->GetIndex())
+	{
+		return {};
+	}
+
+	return InDropZone;
+}
+
+FReply SRCBehaviorSetAssetByPathNewElementRow::OnAcceptDrop(const FDragDropEvent& InDragDropEvent, EItemDropZone InDropZone, RCBehaviorSetAssetByPathNewElementListItem InItem)
+{
+	if (!ElementItem.IsValid())
+	{
+		return FReply::Handled();
+	}
+
+	TSharedPtr<FRCSetAssetByPathBehaviorModelNew> SetAssetByPathBehaviorModelNew = SetAssetByPathBehaviorModelNewWeak.Pin();
+
+	if (!SetAssetByPathBehaviorModelNew.IsValid())
+	{
+		return FReply::Handled();
+	}
+
+	TSharedPtr<FRCBehaviorSetAssetByPathNewElementDragDrop> PathBehaviorElementDragDrop = InDragDropEvent.GetOperationAs<FRCBehaviorSetAssetByPathNewElementDragDrop>();
+
+	if (!PathBehaviorElementDragDrop.IsValid())
+	{
+		return FReply::Handled();
+	}
+
+	if (ElementItem->Index == PathBehaviorElementDragDrop->GetIndex())
+	{
+		return FReply::Handled();
+	}
+
+	SetAssetByPathBehaviorModelNew->ReorderElementItems(ElementItem->Index, InDropZone, PathBehaviorElementDragDrop->GetIndex());
+
+	return FReply::Handled();
+}
+
+FReply SRCBehaviorSetAssetByPathNewElementRow::OnRowDragDetected(const FGeometry& InMyGeometry, const FPointerEvent& InMouseEvent)
+{
+	TSharedPtr<FRCSetAssetByPathBehaviorModelNew> SetAssetByPathBehaviorModelNew = SetAssetByPathBehaviorModelNewWeak.Pin();
+
+	if (!SetAssetByPathBehaviorModelNew.IsValid())
+	{
+		return FReply::Unhandled();
+	}
+
+	if (!InMouseEvent.IsMouseButtonDown(EKeys::LeftMouseButton))
+	{
+		return FReply::Unhandled();
+	}
+
+	TSharedRef<FRCBehaviorSetAssetByPathNewElementDragDrop> PathBehaviorElementDragDrop = MakeShared<FRCBehaviorSetAssetByPathNewElementDragDrop>(ElementItem->Index);
+	return FReply::Handled().BeginDragDrop(PathBehaviorElementDragDrop);
+}
+
+#undef LOCTEXT_NAMESPACE
diff --git a/Engine/Plugins/VirtualProduction/RemoteControl/Source/RemoteControlUI/Private/UI/Behaviour/Builtin/Path/SRCBehaviorSetAssetByPathNewElementRow.h b/Engine/Plugins/VirtualProduction/RemoteControl/Source/RemoteControlUI/Private/UI/Behaviour/Builtin/Path/SRCBehaviorSetAssetByPathNewElementRow.h
new file mode 100644
--- /dev/null
+++ b/Engine/Plugins/VirtualProduction/RemoteControl/Source/RemoteControlUI/Private/UI/Behaviour/Builtin/Path/SRCBehaviorSetAssetByPathNewElementRow.h
@@ -0,0 +1,45 @@
+// Copyright Epic Games, Inc. All Rights Reserved.
+
+#pragma once
+
+#include "HAL/Platform.h"
+#include "Misc/Optional.h"
+#include "Widgets/Views/SListView.h"
+#include "Widgets/Views/STableRow.h"
+
+struct FRCPathBehaviorElementRow
+{
+	int32 Index;
+	TSharedRef<SWidget> Widget;
+};
+
+class FRCSetAssetByPathBehaviorModelNew;
+
+using RCBehaviorSetAssetByPathNewElementListItem = TSharedPtr<FRCPathBehaviorElementRow>;
+
+class SRCBehaviorSetAssetByPathNewElementRow : public STableRow<RCBehaviorSetAssetByPathNewElementListItem>
+{
+	using FSuperRowType = STableRow<RCBehaviorSetAssetByPathNewElementListItem>;
+
+public:
+	SLATE_BEGIN_ARGS(SRCBehaviorSetAssetByPathNewElementRow) 
+		{}
+	SLATE_END_ARGS()
+
+	void Construct(const FArguments& InArgs, const TSharedRef<STableViewBase>& InOwnerTableView,
+		TSharedRef<FRCSetAssetByPathBehaviorModelNew> InSetAssetByPathBehaviorModelNew, RCBehaviorSetAssetByPathNewElementListItem InElementItem);
+
+	int32 GetElementIndex() const;
+
+private:
+	/** Called to validate whether a drag drop event is accepted for a behavior item row */
+	TOptional<EItemDropZone> CanAcceptDrop(const FDragDropEvent& InDragDropEvent, EItemDropZone InDropZone, RCBehaviorSetAssetByPathNewElementListItem InItem);
+
+	/** Called when a drop happens on this behavior item row */
+	FReply OnAcceptDrop(const FDragDropEvent& InDragDropEvent, EItemDropZone InDropZone, RCBehaviorSetAssetByPathNewElementListItem InItem);
+
+	FReply OnRowDragDetected(const FGeometry& InMyGeometry, const FPointerEvent& InMouseEvent);
+
+	TWeakPtr<FRCSetAssetByPathBehaviorModelNew> SetAssetByPathBehaviorModelNewWeak;
+	RCBehaviorSetAssetByPathNewElementListItem ElementItem;
+};
\ No newline at end of file
diff --git a/Engine/Plugins/VirtualProduction/RemoteControl/Source/RemoteControlUI/Private/UI/Behaviour/SRCBehaviorPanelRow.cpp b/Engine/Plugins/VirtualProduction/RemoteControl/Source/RemoteControlUI/Private/UI/Behaviour/SRCBehaviorPanelRow.cpp
--- a/Engine/Plugins/VirtualProduction/RemoteControl/Source/RemoteControlUI/Private/UI/Behaviour/SRCBehaviorPanelRow.cpp
+++ b/Engine/Plugins/VirtualProduction/RemoteControl/Source/RemoteControlUI/Private/UI/Behaviour/SRCBehaviorPanelRow.cpp
@@ -42,6 +42,8 @@ void SRCBehaviorPanelRow::Construct(const FArguments& InArgs, const TSharedRef<S
 	BehaviorItem = InBehaviorItem;
 	BehaviorPanelListWeak = InBehaviorPanelList;
 
+	SetCursor(EMouseCursor::GrabHand);
+
 	const FRCPanelStyle* RCPanelStyle = &FRemoteControlPanelStyle::Get()->GetWidgetStyle<FRCPanelStyle>("RemoteControlPanel.MinorPanel");
 
 	STableRow::Construct(
diff --git a/Engine/Plugins/VirtualProduction/RemoteControl/Source/RemoteControlUI/Private/UI/Controller/CustomControllers/SCustomStringListControllerListEditorWidget.cpp b/Engine/Plugins/VirtualProduction/RemoteControl/Source/RemoteControlUI/Private/UI/Controller/CustomControllers/SCustomStringListControllerListEditorWidget.cpp
--- a/Engine/Plugins/VirtualProduction/RemoteControl/Source/RemoteControlUI/Private/UI/Controller/CustomControllers/SCustomStringListControllerListEditorWidget.cpp
+++ b/Engine/Plugins/VirtualProduction/RemoteControl/Source/RemoteControlUI/Private/UI/Controller/CustomControllers/SCustomStringListControllerListEditorWidget.cpp
@@ -254,6 +254,8 @@ public:
 	{
 		ParentWidgetWeak = InParentWidget;
 		Entry = InEntry;
+
+		SetCursor(EMouseCursor::GrabHand);
 		
 		STableRow<FName>::Construct(
 			STableRow<FName>::FArguments()
diff --git a/Engine/Plugins/VirtualProduction/RemoteControl/Source/RemoteControlUI/Private/UI/Controller/SRCControllerPanelCategoryRow.cpp b/Engine/Plugins/VirtualProduction/RemoteControl/Source/RemoteControlUI/Private/UI/Controller/SRCControllerPanelCategoryRow.cpp
--- a/Engine/Plugins/VirtualProduction/RemoteControl/Source/RemoteControlUI/Private/UI/Controller/SRCControllerPanelCategoryRow.cpp
+++ b/Engine/Plugins/VirtualProduction/RemoteControl/Source/RemoteControlUI/Private/UI/Controller/SRCControllerPanelCategoryRow.cpp
@@ -22,6 +22,8 @@ void SRCControllerPanelCategoryRow::Construct(const FArguments& InArgs, const TS
 	CategoryItem = InCategoryItem;
 	ControllerPanelListWeak = InControllerPanelList.ToWeakPtr();
 
+	SetCursor(EMouseCursor::GrabHand);
+
 	TSharedPtr<FRCCategoryModel> CategoryModel = GetCategory();
 
 	STableRow::Construct(
diff --git a/Engine/Plugins/VirtualProduction/RemoteControl/Source/RemoteControlUI/Private/UI/Controller/SRCControllerPanelControllerRow.cpp b/Engine/Plugins/VirtualProduction/RemoteControl/Source/RemoteControlUI/Private/UI/Controller/SRCControllerPanelControllerRow.cpp
--- a/Engine/Plugins/VirtualProduction/RemoteControl/Source/RemoteControlUI/Private/UI/Controller/SRCControllerPanelControllerRow.cpp
+++ b/Engine/Plugins/VirtualProduction/RemoteControl/Source/RemoteControlUI/Private/UI/Controller/SRCControllerPanelControllerRow.cpp
@@ -28,6 +28,8 @@ void SRCControllerPanelControllerRow::Construct(const FArguments& InArgs, const
 	ControllerItem = InControllerItem;
 	ControllerPanelListWeak = InControllerPanelList.ToWeakPtr();
 
+	SetCursor(EMouseCursor::GrabHand);
+
 	FSuperRowType::Construct(FSuperRowType::FArguments()
 		.Style(InArgs._Style)
 		.Padding(0.f)
diff --git a/Engine/Plugins/VirtualProduction/RemoteControl/Source/RemoteControlUI/Private/UI/SRCPanelDragHandle.h b/Engine/Plugins/VirtualProduction/RemoteControl/Source/RemoteControlUI/Private/UI/SRCPanelDragHandle.h
--- a/Engine/Plugins/VirtualProduction/RemoteControl/Source/RemoteControlUI/Private/UI/SRCPanelDragHandle.h
+++ b/Engine/Plugins/VirtualProduction/RemoteControl/Source/RemoteControlUI/Private/UI/SRCPanelDragHandle.h
@@ -27,6 +27,9 @@ public:
 	{
 		Widget = InArgs._Widget;
 		Id = MoveTemp(InId);
+
+		SetCursor(EMouseCursor::GrabHand);
+
 		ChildSlot
 		[
 			SNew(SBox)
