[Remote Control] Adding missing OnControllerModified() broadcast in the FRemoteControlModule::SetPresetController code path.

This is the code path used by WebRemoteControl to set the controller values. It is needed for rundown server to broadcast controller events. It is also used by the rundown editor to refresh the UI, so this might fix a few UI refresh bug as well.

diff --git a/Engine/Plugins/VirtualProduction/RemoteControl/Source/RemoteControl/Private/RemoteControlModule.cpp b/Engine/Plugins/VirtualProduction/RemoteControl/Source/RemoteControl/Private/RemoteControlModule.cpp
--- a/Engine/Plugins/VirtualProduction/RemoteControl/Source/RemoteControl/Private/RemoteControlModule.cpp
+++ b/Engine/Plugins/VirtualProduction/RemoteControl/Source/RemoteControl/Private/RemoteControlModule.cpp
@@ -2037,7 +2037,10 @@ bool FRemoteControlModule::SetPresetController(const FName PresetName, class URC
 	}
 
 	bool bSuccess = VirtualProperty->DeserializeFromBackend(Backend);
-
+	if (bSuccess)
+	{
+		Preset->OnControllerModified().Broadcast(Preset, {VirtualProperty->Id});
+	}
 	return bSuccess;
 }
 
