Remote Control: Redesign of path preview widget
- Also moved "affects textures" to button tooltip.

diff --git a/Engine/Plugins/VirtualProduction/RemoteControl/Source/RemoteControlUI/Private/UI/Behaviour/Builtin/Path/RCBehaviorSetAssetByPathModelNew.cpp b/Engine/Plugins/VirtualProduction/RemoteControl/Source/RemoteControlUI/Private/UI/Behaviour/Builtin/Path/RCBehaviorSetAssetByPathModelNew.cpp
--- a/Engine/Plugins/VirtualProduction/RemoteControl/Source/RemoteControlUI/Private/UI/Behaviour/Builtin/Path/RCBehaviorSetAssetByPathModelNew.cpp
+++ b/Engine/Plugins/VirtualProduction/RemoteControl/Source/RemoteControlUI/Private/UI/Behaviour/Builtin/Path/RCBehaviorSetAssetByPathModelNew.cpp
@@ -12,6 +12,7 @@
 #include "PropertyEditorModule.h"
 #include "RemoteControlPreset.h"
 #include "ScopedTransaction.h"
+#include "Styling/SlateStyleMacros.h"
 #include "Styling/StyleColors.h"
 #include "UI/Action/Path/RCActionSetAssetByPathModelNew.h"
 #include "UI/Action/SRCActionPanelList.h"
@@ -27,6 +28,12 @@
 
 #define LOCTEXT_NAMESPACE "RCSetAssetByPathBehaviorModelNew"
 
+namespace UE::RemoteControl::UI::Private
+{
+	constexpr const TCHAR* ErrorImage = TEXT("SourceControl.StatusIcon.Error");
+	constexpr const TCHAR* ValidImage = TEXT("SourceControl.StatusIcon.On");
+}
+
 FRCSetAssetByPathBehaviorModelNew::FRCSetAssetByPathBehaviorModelNew(URCSetAssetByPathBehaviorNew* InSetAssetByPathBehavior, const TSharedPtr<SRemoteControlPanel> InRemoteControlPanel)
 	: FRCBehaviourModel(InSetAssetByPathBehavior, InRemoteControlPanel)
 {
@@ -37,7 +44,7 @@ void FRCSetAssetByPathBehaviorModelNew::Initialize()
 {
 	PreviewPathWidget = SNew(STextBlock)
 		.AutoWrapText(true)
-		.ColorAndOpacity(FStyleColors::Error);
+		.Font(DEFAULT_FONT("Regular", 10));
 
 	PathArrayWidget = SNew(SBox);
 
@@ -46,6 +53,9 @@ void FRCSetAssetByPathBehaviorModelNew::Initialize()
 		.SelectionMode(ESelectionMode::Single)
 		.OnGenerateRow(this, &FRCSetAssetByPathBehaviorModelNew::OnGenerateWidgetForList);
 
+	ValidationImage = SNew(SImage)
+		.DesiredSizeOverride(FVector2D(12));
+
 	FPropertyRowGeneratorArgs ArgsRowGeneratorArray;
 	ArgsRowGeneratorArray.NotifyHook = this;
 	PropertyRowGeneratorArray = FModuleManager::GetModuleChecked<FPropertyEditorModule>("PropertyEditor").CreatePropertyRowGenerator(ArgsRowGeneratorArray);
@@ -113,17 +123,33 @@ TSharedRef<SWidget> FRCSetAssetByPathBehaviorModelNew::GetPropertyWidget()
 		];
 
 	FieldWidget->AddSlot()
-		.Padding(FMargin(3.f, 3.f))
+		.Padding(FMargin(3.f, 8.f))
 		.AutoHeight()
 		[
-			SNew(SHorizontalBox)
-			+ SHorizontalBox::Slot()
-			.AutoWidth()
+			SNew(SVerticalBox)
+			+ SVerticalBox::Slot()
+			.AutoHeight()
 			[
-				SNew(STextBlock)
-				.Text(FText::FromString("Preview Path: "))
+				SNew(SHorizontalBox)
+				+ SHorizontalBox::Slot()
+				.AutoWidth()
+				.VAlign(VAlign_Center)
+				[
+					SNew(STextBlock)
+					.Text(FText::FromString("Path Preview"))
+					.Font(DEFAULT_FONT("Regular", 8))
+				]
+				+ SHorizontalBox::Slot()
+				.AutoWidth()
+				.VAlign(VAlign_Center)
+				.Padding(5.f, 0.f, 0.f, 0.f)
+				[
+					ValidationImage.ToSharedRef()
+				]
 			]
-			+ SHorizontalBox::Slot()
+			+ SVerticalBox::Slot()
+			.AutoHeight()
+			.Padding(0.f, 5.f, 0.f, 0.f)
 			[
 				PreviewPathWidget->AsShared()
 			]
@@ -192,10 +218,16 @@ void FRCSetAssetByPathBehaviorModelNew::RefreshPreview() const
 		bIsValidPath = IAssetRegistry::Get()->TryGetAssetByObjectPath(CorrectedPath, AssetData) == UE::AssetRegistry::EExists::Exists;
 	}
 	
-	PreviewPathWidget->SetColorAndOpacity(
+	ValidationImage->SetImage(FAppStyle::Get().GetBrush(
+		bIsValidPath
+			? UE::RemoteControl::UI::Private::ValidImage
+			: UE::RemoteControl::UI::Private::ErrorImage
+	));
+
+	ValidationImage->SetToolTipText(
 		bIsValidPath
-			? FStyleColors::Foreground
-			: FStyleColors::Error
+			? LOCTEXT("ValidImageTooltip", "Path points to a valid asset.")
+			: LOCTEXT("InvalidImageTooltip", "Path does not point to a valid asset!")
 	);
 }
 
diff --git a/Engine/Plugins/VirtualProduction/RemoteControl/Source/RemoteControlUI/Private/UI/Behaviour/Builtin/Path/RCBehaviorSetAssetByPathModelNew.h b/Engine/Plugins/VirtualProduction/RemoteControl/Source/RemoteControlUI/Private/UI/Behaviour/Builtin/Path/RCBehaviorSetAssetByPathModelNew.h
--- a/Engine/Plugins/VirtualProduction/RemoteControl/Source/RemoteControlUI/Private/UI/Behaviour/Builtin/Path/RCBehaviorSetAssetByPathModelNew.h
+++ b/Engine/Plugins/VirtualProduction/RemoteControl/Source/RemoteControlUI/Private/UI/Behaviour/Builtin/Path/RCBehaviorSetAssetByPathModelNew.h
@@ -9,6 +9,7 @@
 class IDetailTreeNode;
 class ITableRow;
 class SBox;
+class SImage;
 class STableViewBase;
 class URCSetAssetByPathBehaviorNew;
 enum class EItemDropZone;
@@ -84,6 +85,9 @@ private:
 	/** Pointer to the Preview Text Widget for the Path */
 	TSharedPtr<STextBlock> PreviewPathWidget;
 
+	/** Image showing whether the path is valid or not. */
+	TSharedPtr<SImage> ValidationImage;
+
 	/** The row generator used to build the generic array value widget for the Path Array */
 	TSharedPtr<IPropertyRowGenerator> PropertyRowGeneratorArray;
 
diff --git a/Engine/Plugins/VirtualProduction/RemoteControl/Source/RemoteControlUI/Private/UI/Behaviour/Builtin/Path/SRCBehaviorSetAssetByPathNew.cpp b/Engine/Plugins/VirtualProduction/RemoteControl/Source/RemoteControlUI/Private/UI/Behaviour/Builtin/Path/SRCBehaviorSetAssetByPathNew.cpp
--- a/Engine/Plugins/VirtualProduction/RemoteControl/Source/RemoteControlUI/Private/UI/Behaviour/Builtin/Path/SRCBehaviorSetAssetByPathNew.cpp
+++ b/Engine/Plugins/VirtualProduction/RemoteControl/Source/RemoteControlUI/Private/UI/Behaviour/Builtin/Path/SRCBehaviorSetAssetByPathNew.cpp
@@ -10,7 +10,8 @@
 #include "Styling/StyleColors.h"
 #include "UI/Behaviour/Builtin/Path/RCBehaviorSetAssetByPathModelNew.h"
 #include "UI/RemoteControlPanelStyle.h"
-#include "Widgets/Input/SButton.h"
+#include "Widgets/SBoxPanel.h"
+#include "Widgets/Input/SCheckBox.h"
 #include "Widgets/Text/STextBlock.h"
 
 #define LOCTEXT_NAMESPACE "SRCBehaviorSetAssetByPathNew"
@@ -32,18 +33,9 @@ void SRCBehaviorSetAssetByPathNew::Construct(const FArguments& InArgs, TSharedRe
 		return;
 	}
 
-	const TSharedRef<SHorizontalBox> InternalExternalSwitchWidget = SNew(SHorizontalBox)
+	const FText TooltipText = LOCTEXT("OnlyAffectsTextures", "Only affects texture assets.");
 
-		+ SHorizontalBox::Slot()
-		.HAlign(HAlign_Right)
-		.VAlign(VAlign_Center)
-		.Padding(FMargin(3.0f, 0.0f))
-		.FillWidth(1.0f)
-		.AutoWidth()
-		[
-			SNew(STextBlock)
-			.Text(LOCTEXT("AffectsTexturesOnly", "For Texture Sources:"))
-		]
+	const TSharedRef<SHorizontalBox> InternalExternalSwitchWidget = SNew(SHorizontalBox)
 
 		+ SHorizontalBox::Slot()
 		.HAlign(HAlign_Right)
@@ -52,10 +44,16 @@ void SRCBehaviorSetAssetByPathNew::Construct(const FArguments& InArgs, TSharedRe
 		.FillWidth(1.0f)
 		.AutoWidth()
 		[
-			SNew(SButton)
-			.Text(FText::FromString(FString("Internal")))
-			.ButtonColorAndOpacity(this, &SRCBehaviorSetAssetByPathNew::GetInternalWidgetColor)
-			.OnPressed(this, &SRCBehaviorSetAssetByPathNew::OnInternalExternalSwitchWidgetPressed, true)
+			SNew(SCheckBox)
+			.Style(FAppStyle::Get(), "DetailsView.SectionButton")
+			.IsChecked(this, &SRCBehaviorSetAssetByPathNew::GetInternalExternalSwitchState, /* Internal */ true)
+			.OnCheckStateChanged(this, &SRCBehaviorSetAssetByPathNew::OnInternalExternalSwitchStateChanged, /* Internal */ true)
+			.ToolTipText(TooltipText)
+			.Content()
+			[
+				SNew(STextBlock)
+				.Text(LOCTEXT("Internal", "Internal"))
+			]
 		]
 	
 		+ SHorizontalBox::Slot()
@@ -65,10 +63,16 @@ void SRCBehaviorSetAssetByPathNew::Construct(const FArguments& InArgs, TSharedRe
 		.FillWidth(1.0f)
 		.AutoWidth()
 		[
-			SNew(SButton)
-			.Text(FText::FromString(FString("External")))
-			.ButtonColorAndOpacity(this, &SRCBehaviorSetAssetByPathNew::GetExternalWidgetColor)
-			.OnPressed(this, &SRCBehaviorSetAssetByPathNew::OnInternalExternalSwitchWidgetPressed, false)
+			SNew(SCheckBox)
+			.Style(FAppStyle::Get(), "DetailsView.SectionButton")
+			.IsChecked(this, &SRCBehaviorSetAssetByPathNew::GetInternalExternalSwitchState, /* Internal */ false)
+			.OnCheckStateChanged(this, &SRCBehaviorSetAssetByPathNew::OnInternalExternalSwitchStateChanged, /* Internal */ false)
+			.ToolTipText(TooltipText)
+			.Content()
+			[
+				SNew(STextBlock)
+				.Text(LOCTEXT("External", "External"))
+			]
 		];
 	
 	ChildSlot
@@ -87,43 +91,36 @@ void SRCBehaviorSetAssetByPathNew::Construct(const FArguments& InArgs, TSharedRe
 	];
 }
 
-void SRCBehaviorSetAssetByPathNew::OnInternalExternalSwitchWidgetPressed(bool bInIsInternal) const
+ECheckBoxState SRCBehaviorSetAssetByPathNew::GetInternalExternalSwitchState(bool bInIsInternal) const
 {
 	if (!PathBehaviour.IsValid())
 	{
-		return;
+		return ECheckBoxState::Undetermined;
 	}
 
-	PathBehaviour->bInternal = bInIsInternal;
-
-	if (const TSharedPtr<const FRCSetAssetByPathBehaviorModelNew>& SetAssetPath = SetAssetByPathWeakPtr.Pin())
-	{
-		SetAssetPath->RefreshPreview();
-	}
+	return PathBehaviour->bInternal == bInIsInternal
+		? ECheckBoxState::Checked
+		: ECheckBoxState::Unchecked;
 }
 
-FSlateColor SRCBehaviorSetAssetByPathNew::GetExternalWidgetColor() const
+void SRCBehaviorSetAssetByPathNew::OnInternalExternalSwitchStateChanged(ECheckBoxState InState, bool bInIsInternal)
 {
 	if (!PathBehaviour.IsValid())
 	{
-		return FAppStyle::Get().GetSlateColor("Colors.AccentRed");
+		return;
 	}
 
-	return PathBehaviour->bInternal 
-		? FAppStyle::Get().GetSlateColor("Colors.AccentWhite") 
-		: FAppStyle::Get().GetSlateColor("Colors.Highlight");
-}
-
-FSlateColor SRCBehaviorSetAssetByPathNew::GetInternalWidgetColor() const
-{
-	if (!PathBehaviour.IsValid())
+	if (InState != ECheckBoxState::Checked)
 	{
-		return FAppStyle::Get().GetSlateColor("Colors.AccentRed");
+		return;
 	}
 
-	return PathBehaviour->bInternal 
-		? FAppStyle::Get().GetSlateColor("Colors.Highlight") 
-		: FAppStyle::Get().GetSlateColor("Colors.AccentWhite");
+	PathBehaviour->bInternal = bInIsInternal;
+
+	if (const TSharedPtr<const FRCSetAssetByPathBehaviorModelNew>& SetAssetPath = SetAssetByPathWeakPtr.Pin())
+	{
+		SetAssetPath->RefreshPreview();
+	}
 }
 
 END_SLATE_FUNCTION_BUILD_OPTIMIZATION
diff --git a/Engine/Plugins/VirtualProduction/RemoteControl/Source/RemoteControlUI/Private/UI/Behaviour/Builtin/Path/SRCBehaviorSetAssetByPathNew.h b/Engine/Plugins/VirtualProduction/RemoteControl/Source/RemoteControlUI/Private/UI/Behaviour/Builtin/Path/SRCBehaviorSetAssetByPathNew.h
--- a/Engine/Plugins/VirtualProduction/RemoteControl/Source/RemoteControlUI/Private/UI/Behaviour/Builtin/Path/SRCBehaviorSetAssetByPathNew.h
+++ b/Engine/Plugins/VirtualProduction/RemoteControl/Source/RemoteControlUI/Private/UI/Behaviour/Builtin/Path/SRCBehaviorSetAssetByPathNew.h
@@ -10,6 +10,7 @@
 class FRCSetAssetByPathBehaviorModelNew;
 class UClass;
 class URCSetAssetByPathBehaviorNew;
+enum class ECheckBoxState : uint8;
 
 class SRCBehaviorSetAssetByPathNew : public  SCompoundWidget
 {
@@ -22,14 +23,11 @@ public:
 	void Construct(const FArguments& InArgs, TSharedRef<FRCSetAssetByPathBehaviorModelNew> InBehaviourItem);
 
 private:
-	/** Called when the Internal/External button are pressed */
-	void OnInternalExternalSwitchWidgetPressed(bool bInIsInternal) const;
-
-	/** Return the color of the External widget */
-	FSlateColor GetExternalWidgetColor() const;
+	/** Determines the state of the internal/external buttons. */
+	ECheckBoxState GetInternalExternalSwitchState(bool bInIsInternal) const;
 
-	/** Return the color of the Internal widget */
-	FSlateColor GetInternalWidgetColor() const;
+	/** Called when the Internal/External button are pressed */
+	void OnInternalExternalSwitchStateChanged(ECheckBoxState InState, bool bInIsInternal);
 
 private:
 	/** The Behaviour (UI model) associated with us */
