Remote Control: Fixed incorrect usage of FText::FromString

diff --git a/Engine/Plugins/VirtualProduction/RemoteControl/Source/RemoteControlUI/Private/UI/Behaviour/Builtin/Path/RCBehaviorSetAssetByPathModelNew.cpp b/Engine/Plugins/VirtualProduction/RemoteControl/Source/RemoteControlUI/Private/UI/Behaviour/Builtin/Path/RCBehaviorSetAssetByPathModelNew.cpp
--- a/Engine/Plugins/VirtualProduction/RemoteControl/Source/RemoteControlUI/Private/UI/Behaviour/Builtin/Path/RCBehaviorSetAssetByPathModelNew.cpp
+++ b/Engine/Plugins/VirtualProduction/RemoteControl/Source/RemoteControlUI/Private/UI/Behaviour/Builtin/Path/RCBehaviorSetAssetByPathModelNew.cpp
@@ -136,7 +136,7 @@ TSharedRef<SWidget> FRCSetAssetByPathBehaviorModelNew::GetPropertyWidget()
 				.VAlign(VAlign_Center)
 				[
 					SNew(STextBlock)
-					.Text(FText::FromString("Path Preview"))
+					.Text(LOCTEXT("PathPreview", "Path Preview"))
 					.Font(DEFAULT_FONT("Regular", 8))
 				]
 				+ SHorizontalBox::Slot()
@@ -180,7 +180,7 @@ TSharedPtr<SRCLogicPanelListBase> FRCSetAssetByPathBehaviorModelNew::GetActionsL
 
 void FRCSetAssetByPathBehaviorModelNew::RefreshPreview() const
 {
-	FString PreviewText = TEXT("Path is invalid!");
+	FString PreviewText = LOCTEXT("InvalidPath", "Invalid Path!").ToString();
 	bool bInternal = true;
 
 	if (URCSetAssetByPathBehaviorNew* SetAssetByPathBehavior = SetAssetByPathBehaviorWeakPtr.Get())
