[Remote Control] Adding a command to enable remote control interception local forwarding to make it easier to test.

Testing the RCI (remote control interception) requires a full Display cluster render node setup, which makes testing (or debugging) it time consuming.
This change will help by adding a "local forwarding" mode to RCI, so that it is possible to have the RC operations locally (i.e. within the local process) go through RCI.
This is not testing the Display Cluster RCI, but just the RCI core itself to ensure it has the coverage. This is exposed in a command "RemoteControl.EnableLocalForwardingInterceptor".

Also in this change, made sure the "function call" from the RC Panel also go through the same InvokeCall api. Consistency in the implementation using the same code path and allow for testing the RC interception code more easily.

diff --git a/Engine/Plugins/VirtualProduction/RemoteControl/Source/RemoteControl/Private/RemoteControlModule.cpp b/Engine/Plugins/VirtualProduction/RemoteControl/Source/RemoteControl/Private/RemoteControlModule.cpp
--- a/Engine/Plugins/VirtualProduction/RemoteControl/Source/RemoteControl/Private/RemoteControlModule.cpp
+++ b/Engine/Plugins/VirtualProduction/RemoteControl/Source/RemoteControl/Private/RemoteControlModule.cpp
@@ -35,6 +35,7 @@
 #include "Serialization/PropertyMapStructDeserializerBackendWrapper.h"
 #include "StructDeserializer.h"
 #include "StructSerializer.h"
+#include "Tests/RemoteControlInterceptionForwardInterceptor.h"
 #include "UObject/Class.h"
 #include "UObject/FieldPath.h"
 #include "UObject/UnrealType.h"
@@ -936,10 +937,34 @@ void FRemoteControlModule::StartupModule()
 	RegisterPropertyIdHandler();
 
 	PopulateDisallowedFunctions();
+
+	ConsoleObjects.Add(IConsoleManager::Get().RegisterConsoleCommand(
+		TEXT("RemoteControl.EnableLocalForwardingInterceptor"),
+		TEXT("Enables/Disables the local forwarding interceptor debugging mode.\n")
+		TEXT("This allows to emulate remote control interception and processing locally.\n")
+		TEXT("Usage: \"RemoteControl.EnableLocalForwardingInterceptor [true/false] or [0/1]\""),
+		FConsoleCommandWithArgsDelegate::CreateLambda([this](const TArray<FString>& InArgs)
+		{
+			const bool bEnable = InArgs.Num() > 0 ? InArgs[0].ToBool() : true;
+			if (bEnable)
+			{
+				RegisterLocalForwardingInterceptor();
+				UE_LOG(LogRemoteControl, Log, TEXT("Local Forwarding Interception is enabled."));
+			}
+			else
+			{
+				UnregisterLocalForwardingInterceptor();
+				UE_LOG(LogRemoteControl, Log, TEXT("Local Forwarding Interception is disabled."));
+			}
+		}),
+		ECVF_Default
+		));
 }
 
 void FRemoteControlModule::ShutdownModule()
 {
+	UnregisterConsoleObjects();
+
 #if WITH_EDITOR
 	FCoreUObjectDelegates::PreLoadMap.RemoveAll(this);
 	UnregisterEditorDelegates();
@@ -956,6 +981,8 @@ void FRemoteControlModule::ShutdownModule()
 
 	// Unregister the interceptor feature on module shutdown
 	IModularFeatures::Get().UnregisterModularFeature(IRemoteControlInterceptionFeatureProcessor::GetName(), RCIProcessor.Get());
+	
+	UnregisterLocalForwardingInterceptor();
 
 	if (UObjectInitialized())
 	{
@@ -3184,6 +3211,33 @@ FRemoteControlModule::FOngoingChange::FOngoingChange(FRCCallReference InReferenc
 
 #endif
 
+void FRemoteControlModule::RegisterLocalForwardingInterceptor()
+{
+	if (!RCILocalForwardingInterceptor)
+	{
+		RCILocalForwardingInterceptor =  MakeUnique<FRemoteControlInterceptionForwardInterceptor>();
+		IModularFeatures::Get().RegisterModularFeature(IRemoteControlInterceptionFeatureInterceptor::GetName(), RCILocalForwardingInterceptor.Get());
+	}
+}
+
+void FRemoteControlModule::UnregisterLocalForwardingInterceptor()
+{
+	if (RCILocalForwardingInterceptor)
+	{
+		IModularFeatures::Get().UnregisterModularFeature(IRemoteControlInterceptionFeatureInterceptor::GetName(), RCILocalForwardingInterceptor.Get());
+		RCILocalForwardingInterceptor.Reset();
+	}
+}
+
+void FRemoteControlModule::UnregisterConsoleObjects()
+{
+	for (IConsoleObject* ConsoleObject : ConsoleObjects)
+	{
+		IConsoleManager::Get().UnregisterConsoleObject(ConsoleObject);
+	}
+	ConsoleObjects.Empty();
+}
+
 PRAGMA_ENABLE_DEPRECATION_WARNINGS
 
 #undef LOCTEXT_NAMESPACE
diff --git a/Engine/Plugins/VirtualProduction/RemoteControl/Source/RemoteControl/Private/RemoteControlModule.h b/Engine/Plugins/VirtualProduction/RemoteControl/Source/RemoteControl/Private/RemoteControlModule.h
--- a/Engine/Plugins/VirtualProduction/RemoteControl/Source/RemoteControl/Private/RemoteControlModule.h
+++ b/Engine/Plugins/VirtualProduction/RemoteControl/Source/RemoteControl/Private/RemoteControlModule.h
@@ -208,6 +208,15 @@ private:
 
 #endif // WITH_EDITOR
 
+	/** Registers the local forwarding interceptor. */
+	void RegisterLocalForwardingInterceptor();
+
+	/** Unregisters the local forwarding interceptor. */
+	void UnregisterLocalForwardingInterceptor();
+
+	/** Unregisters the console objects from the console manager. */
+	void UnregisterConsoleObjects();
+	
 private:
 	/** Cache of preset names to preset assets */
 	mutable TMap<FName, TArray<FAssetData>> CachedPresetsByName;
@@ -242,6 +251,12 @@ private:
 	/** RC Processor feature instance */
 	TUniquePtr<IRemoteControlInterceptionFeatureProcessor> RCIProcessor;
 
+	/**
+	 * RC Local Forwarding interceptor instance.
+	 * When this is enabled it can locally emulate interception and forwarding to the RCI processor.
+	 */
+	TUniquePtr<IRemoteControlInterceptionFeatureInterceptor> RCILocalForwardingInterceptor;
+
 #if WITH_EDITOR
 	/** Flags for a given RC change. */
 	struct FOngoingChange
@@ -299,6 +314,9 @@ private:
 
 	/** List of functions that can't be called remotely. */
 	TSet<TWeakObjectPtr<UFunction>> FunctionDisallowList;
+
+	/** Keep track of the registered console objects for automatic unregistration on module shutdown. */
+	TArray<IConsoleObject*> ConsoleObjects;
 };
 
 PRAGMA_ENABLE_DEPRECATION_WARNINGS
diff --git a/Engine/Plugins/VirtualProduction/RemoteControl/Source/RemoteControlUI/Private/UI/SRCPanelExposedField.cpp b/Engine/Plugins/VirtualProduction/RemoteControl/Source/RemoteControlUI/Private/UI/SRCPanelExposedField.cpp
--- a/Engine/Plugins/VirtualProduction/RemoteControl/Source/RemoteControlUI/Private/UI/SRCPanelExposedField.cpp
+++ b/Engine/Plugins/VirtualProduction/RemoteControl/Source/RemoteControlUI/Private/UI/SRCPanelExposedField.cpp
@@ -3,6 +3,7 @@
 #include "SRCPanelExposedField.h"
 
 #include "Algo/Transform.h"
+#include "Backends/CborStructSerializerBackend.h"
 #include "Components/MeshComponent.h"
 #include "Components/SceneComponent.h"
 #include "DetailTreeNode.h"
@@ -25,9 +26,9 @@
 #include "RemoteControlPreset.h"
 #include "RemoteControlSettings.h"
 #include "RemoteControlUIModule.h"
-#include "ScopedTransaction.h"
 #include "SRCPanelTreeNode.h"
 #include "SResetToDefaultPropertyEditor.h"
+#include "StructSerializer.h"
 #include "Styling/AppStyle.h"
 #include "Styling/RemoteControlStyles.h"
 #include "Styling/SlateBrush.h"
@@ -876,9 +877,6 @@ TSharedRef<SWidget> SRCPanelExposedField::ConstructCallFunctionButton(bool bIsEn
 
 FReply SRCPanelExposedField::OnClickFunctionButton()
 {
-	FScopedTransaction FunctionTransaction(LOCTEXT("CallExposedFunction", "Called a function through preset."));
-	FEditorScriptExecutionGuard ScriptGuard;
-
 	if (URemoteControlPreset* RCPreset = Preset.Get())
 	{
 		if (TSharedPtr<FRemoteControlFunction> Function = RCPreset->GetExposedEntity<FRemoteControlFunction>(EntityId).Pin())
@@ -887,8 +885,21 @@ FReply SRCPanelExposedField::OnClickFunctionButton()
 			{
 				if (Function->FunctionArguments && Function->FunctionArguments->IsValid())
 				{
-					Object->Modify();
-					Object->ProcessEvent(Function->GetFunction(), Function->FunctionArguments->GetStructMemory());
+					FRCCallReference CallRef;
+					CallRef.Object = Object;
+					CallRef.Function = Function->GetFunction();
+
+					FRCCall Call;
+					Call.CallRef = MoveTemp(CallRef);
+					Call.TransactionMode = ERCTransactionMode::AUTOMATIC;
+					Call.ParamStruct = FStructOnScope(Function->FunctionArguments->GetStruct(), Function->FunctionArguments->GetStructMemory());
+
+					TArray<uint8> FunctionPayload;
+					FMemoryWriter Writer(FunctionPayload);
+					FCborStructSerializerBackend SerializerBackend{ Writer, EStructSerializerBackendFlags::Default };
+					FStructSerializer::Serialize(Function->FunctionArguments->GetStructMemory(), *const_cast<UStruct*>(Function->FunctionArguments->GetStruct()), SerializerBackend);
+					
+					IRemoteControlModule::Get().InvokeCall(Call, ERCPayloadType::Cbor, FunctionPayload);
 				}
 				else
 				{
