Remote Control: fixed issues with deleting/renaming groups not working

diff --git a/Engine/Plugins/VirtualProduction/RemoteControl/Source/RemoteControlUI/Private/UI/SRCPanelFieldGroup.cpp b/Engine/Plugins/VirtualProduction/RemoteControl/Source/RemoteControlUI/Private/UI/SRCPanelFieldGroup.cpp
--- a/Engine/Plugins/VirtualProduction/RemoteControl/Source/RemoteControlUI/Private/UI/SRCPanelFieldGroup.cpp
+++ b/Engine/Plugins/VirtualProduction/RemoteControl/Source/RemoteControlUI/Private/UI/SRCPanelFieldGroup.cpp
@@ -36,6 +36,13 @@ void SRCPanelGroup::Tick(const FGeometry&, const double, const float)
 
 void SRCPanelGroup::Construct(const FArguments& InArgs, URemoteControlPreset* InPreset, FRCColumnSizeData InColumnSizeData)
 {
+	CommandList = MakeShared<FUICommandList>();
+
+	if (InArgs._CommandList.IsValid())
+	{
+		CommandList->Append(InArgs._CommandList.ToSharedRef());
+	}
+
 	Preset = InPreset;
 	
 	Id = InArgs._Id;
