Remote Control: Fixed mem leak with group color.

diff --git a/Engine/Plugins/VirtualProduction/RemoteControl/Source/RemoteControlUI/Private/UI/SRCPanelFieldGroup.cpp b/Engine/Plugins/VirtualProduction/RemoteControl/Source/RemoteControlUI/Private/UI/SRCPanelFieldGroup.cpp
--- a/Engine/Plugins/VirtualProduction/RemoteControl/Source/RemoteControlUI/Private/UI/SRCPanelFieldGroup.cpp
+++ b/Engine/Plugins/VirtualProduction/RemoteControl/Source/RemoteControlUI/Private/UI/SRCPanelFieldGroup.cpp
@@ -65,7 +65,8 @@ void SRCPanelGroup::Construct(const FArguments& InArgs, URemoteControlPreset* In
 		.AutoWidth()
 		[
 			SNew(SBorder)
-			.BorderImage(this, &SRCPanelGroup::HandleGroupColor)
+			.BorderImage(FAppStyle::Get().GetBrush("WhiteBrush"))
+			.BorderBackgroundColor(this, &SRCPanelGroup::GetGroupColor)
 		]
 		// Drag and drop handle
 		+ SHorizontalBox::Slot()
@@ -278,11 +279,11 @@ void SRCPanelGroup::OnLabelCommitted(const FText& InLabel, ETextCommit::Type InC
 	Preset->Layout.RenameGroup(Id, FName(*InLabel.ToString()));
 }
 
-const FSlateBrush* SRCPanelGroup::HandleGroupColor() const
+FSlateColor SRCPanelGroup::GetGroupColor() const
 {
-	check(Preset.IsValid());
-
-	return new FSlateColorBrush(Preset->Layout.GetTagColor(Id));
+	return Preset.IsValid()
+		? Preset->Layout.GetTagColor(Id)
+		: FLinearColor::White;
 }
 
 #undef LOCTEXT_NAMESPACE
diff --git a/Engine/Plugins/VirtualProduction/RemoteControl/Source/RemoteControlUI/Private/UI/SRCPanelFieldGroup.h b/Engine/Plugins/VirtualProduction/RemoteControl/Source/RemoteControlUI/Private/UI/SRCPanelFieldGroup.h
--- a/Engine/Plugins/VirtualProduction/RemoteControl/Source/RemoteControlUI/Private/UI/SRCPanelFieldGroup.h
+++ b/Engine/Plugins/VirtualProduction/RemoteControl/Source/RemoteControlUI/Private/UI/SRCPanelFieldGroup.h
@@ -2,14 +2,15 @@
 
 #pragma once
 
+#include "Brushes/SlateColorBrush.h"
 #include "CoreMinimal.h"
 #include "DragAndDrop/DecoratedDragDropOp.h"
 #include "Input/DragAndDrop.h"
 #include "Misc/Guid.h"
 #include "SRCPanelTreeNode.h"
+#include "UObject/WeakObjectPtrTemplates.h"
 #include "Widgets/Views/SListView.h"
 #include "Widgets/Views/STableRow.h"
-#include "UObject/WeakObjectPtrTemplates.h"
 
 class SRemoteControlPanel;
 struct SRCPanelTreeNode;
@@ -85,7 +86,7 @@ private:
 
 	void OnLabelCommitted(const FText& InLabel, ETextCommit::Type InCommitInfo);
 
-	const FSlateBrush* HandleGroupColor() const;
+	FSlateColor GetGroupColor() const;
 
 private:
 	/** Command list to use for context menu actions */
