Remote Control: String list controller edit window now accepts button clicks while editing fields, discarding the changes in the edited field.

diff --git a/Engine/Plugins/VirtualProduction/RemoteControl/Source/RemoteControlUI/Private/UI/Controller/CustomControllers/SCustomStringListControllerListEditorWidget.cpp b/Engine/Plugins/VirtualProduction/RemoteControl/Source/RemoteControlUI/Private/UI/Controller/CustomControllers/SCustomStringListControllerListEditorWidget.cpp
--- a/Engine/Plugins/VirtualProduction/RemoteControl/Source/RemoteControlUI/Private/UI/Controller/CustomControllers/SCustomStringListControllerListEditorWidget.cpp
+++ b/Engine/Plugins/VirtualProduction/RemoteControl/Source/RemoteControlUI/Private/UI/Controller/CustomControllers/SCustomStringListControllerListEditorWidget.cpp
@@ -64,6 +64,7 @@ class SCustomStringListControllerEnumEditorRowEditBox : public SCompoundWidget
 					.ContentPadding(5.f)
 					.OnClicked(this, &SCustomStringListControllerEnumEditorRowEditBox::OnRemoveClicked)
 					.IsEnabled(this, &SCustomStringListControllerEnumEditorRowEditBox::CanRemove)
+					.ClickMethod(EButtonClickMethod::MouseDown)
 					[
 						SNew(SImage)
 						.Image(FAppStyle::Get().GetBrush("Icons.Delete"))
@@ -162,6 +163,12 @@ private:
 
 	void OnEntryCommitted(const FText& InNew, ETextCommit::Type InCommitType)
 	{
+		// When leaving focus, just ignore changes to the text in the box.
+		if (InCommitType == ETextCommit::OnUserMovedFocus)
+		{
+			return;
+		}
+
 		TSharedPtr<SCustomStringListControllerListEditorWidget> ParentWidget = ParentWidgetWeak.Pin();
 
 		if (!ParentWidget.IsValid())
@@ -381,6 +388,7 @@ void SCustomStringListControllerListEditorWidget::Construct(const FArguments& In
 				SNew(SButton)
 				.OnClicked(this, &SCustomStringListControllerListEditorWidget::OnOkClicked)
 				.ToolTipText(LOCTEXT("OkayToolTip", "Close this window and use these values for the enum."))
+				.ClickMethod(EButtonClickMethod::MouseDown)
 				[
 					SNew(STextBlock)
 					.Text(LOCTEXT("AcceptValues", "Ok"))
@@ -393,6 +401,7 @@ void SCustomStringListControllerListEditorWidget::Construct(const FArguments& In
 				SNew(SButton)
 				.OnClicked(this, &SCustomStringListControllerListEditorWidget::OnCancelClicked)
 				.ToolTipText(LOCTEXT("CancelToolTip", "Renumber the entries sequentially from 0."))
+				.ClickMethod(EButtonClickMethod::MouseDown)
 				[
 					SNew(STextBlock)
 					.Text(LOCTEXT("RejectValues", "Cancel"))
