Remote Control: Enums now correctly bind to and auto-bind from properties.

diff --git a/Engine/Plugins/VirtualProduction/RemoteControl/Source/RemoteControlLogic/Private/Behaviour/Builtin/Bind/RCBehaviourBind.cpp b/Engine/Plugins/VirtualProduction/RemoteControl/Source/RemoteControlLogic/Private/Behaviour/Builtin/Bind/RCBehaviourBind.cpp
--- a/Engine/Plugins/VirtualProduction/RemoteControl/Source/RemoteControlLogic/Private/Behaviour/Builtin/Bind/RCBehaviourBind.cpp
+++ b/Engine/Plugins/VirtualProduction/RemoteControl/Source/RemoteControlLogic/Private/Behaviour/Builtin/Bind/RCBehaviourBind.cpp
@@ -202,17 +202,26 @@ bool URCBehaviourBind::CanHaveActionForField(URCController* Controller, TSharedR
 		return false; // Unbound Property
 	}
 
-	// Enums with base other than 8bits should be ignored, so let's make sure of that
+	const UEnum* Enum = nullptr;
+
 	if (const FEnumProperty* EnumProperty = CastField<FEnumProperty>(RemoteControlProperty))
 	{
-		if (const UEnum* Enum = EnumProperty->GetEnum())
-		{
-			const int64 MaxEnumValue = Enum->GetMaxEnumValue();
-			const uint32 NeededBits = FMath::RoundUpToPowerOfTwo(MaxEnumValue);
+		Enum = EnumProperty->GetEnum();
+	}
+	// BP enum properties are bytes
+	else if (const FByteProperty* ByteProperty = CastField<FByteProperty>(RemoteControlProperty))
+	{
+		Enum = ByteProperty->Enum;
+	}
 
-			// 8 bits enums only
-			return NeededBits <= 256;
-		}
+	// Enums with base other than 8bits should be ignored, so let's make sure of that
+	if (Enum)
+	{
+		const int64 MaxEnumValue = Enum->GetMaxEnumValue();
+		const uint32 NeededBits = FMath::RoundUpToPowerOfTwo(MaxEnumValue);
+
+		// 8 bits enums only
+		return NeededBits <= 256;
 	}
 	
 	// Indirect Binding (related types)
@@ -346,9 +355,10 @@ bool URCBehaviourBind::GetPropertyBagTypeFromFieldProperty(const FProperty* InPr
 		/* Property Bag Type */                           /* Input - Remote Control Property Type */
 		{ EPropertyBagPropertyType::String,  /* --> */    { FStrProperty::StaticClass(),      FTextProperty::StaticClass(),     FNameProperty::StaticClass()}},
 		{ EPropertyBagPropertyType::Int32,   /* --> */    { FIntProperty::StaticClass(),      FInt64Property::StaticClass(),    FInt16Property::StaticClass(),   FUInt32Property::StaticClass(),
-														   FUInt64Property::StaticClass(),   FUInt16Property::StaticClass(),   FEnumProperty::StaticClass(),    FByteProperty::StaticClass() } },
+														   FUInt64Property::StaticClass(),   FUInt16Property::StaticClass(),    FByteProperty::StaticClass() } },
 		{ EPropertyBagPropertyType::Float,   /* --> */    { FFloatProperty::StaticClass(),    FDoubleProperty::StaticClass() } },
 		{ EPropertyBagPropertyType::Bool,    /* --> */    { FBoolProperty::StaticClass() } },
+		{ EPropertyBagPropertyType::Enum,    /* --> */    { FEnumProperty::StaticClass() } },
 		{ EPropertyBagPropertyType::Struct,  /* --> */    { FStructProperty::StaticClass() }, },
 		{ EPropertyBagPropertyType::Object,  /* --> */    { FObjectProperty::StaticClass() } /* this will probably never occur, so we manually set the out bag type later */ }
 	};
@@ -395,6 +405,20 @@ bool URCBehaviourBind::GetPropertyBagTypeFromFieldProperty(const FProperty* InPr
 			}
 		}
 	}
+	// Handle getting the enum from enum properties
+	else if (const FEnumProperty* EnumProperty = CastField<FEnumProperty>(InProperty))
+	{
+		OutStructObject = EnumProperty->GetEnum();
+	}
+	// Handle BP-based enums that are stored as byte propertties
+	else if (const FByteProperty* ByteProperty = CastField<FByteProperty>(InProperty))
+	{
+		if (ByteProperty->Enum)
+		{
+			OutPropertyBagType = EPropertyBagPropertyType::Enum;
+			OutStructObject = ByteProperty->Enum;
+		}
+	}
 
 	if (OutStructObject == UTexture::StaticClass())
 	{
