Remote Control: String controller edit fields will now update the list when tabbing or clicking out.

diff --git a/Engine/Plugins/VirtualProduction/RemoteControl/Source/RemoteControlUI/Private/UI/Controller/CustomControllers/SCustomStringListControllerListEditorWidget.cpp b/Engine/Plugins/VirtualProduction/RemoteControl/Source/RemoteControlUI/Private/UI/Controller/CustomControllers/SCustomStringListControllerListEditorWidget.cpp
--- a/Engine/Plugins/VirtualProduction/RemoteControl/Source/RemoteControlUI/Private/UI/Controller/CustomControllers/SCustomStringListControllerListEditorWidget.cpp
+++ b/Engine/Plugins/VirtualProduction/RemoteControl/Source/RemoteControlUI/Private/UI/Controller/CustomControllers/SCustomStringListControllerListEditorWidget.cpp
@@ -163,12 +163,6 @@ private:
 
 	void OnEntryCommitted(const FText& InNew, ETextCommit::Type InCommitType)
 	{
-		// When leaving focus, just ignore changes to the text in the box.
-		if (InCommitType == ETextCommit::OnUserMovedFocus)
-		{
-			return;
-		}
-
 		TSharedPtr<SCustomStringListControllerListEditorWidget> ParentWidget = ParentWidgetWeak.Pin();
 
 		if (!ParentWidget.IsValid())
