Remote Control: Fixed crash when duplicating categories.

#FailedValidations Preflight

diff --git a/Engine/Plugins/VirtualProduction/RemoteControl/Source/RemoteControlUI/Private/UI/Controller/SRCControllerPanel.cpp b/Engine/Plugins/VirtualProduction/RemoteControl/Source/RemoteControlUI/Private/UI/Controller/SRCControllerPanel.cpp
--- a/Engine/Plugins/VirtualProduction/RemoteControl/Source/RemoteControlUI/Private/UI/Controller/SRCControllerPanel.cpp
+++ b/Engine/Plugins/VirtualProduction/RemoteControl/Source/RemoteControlUI/Private/UI/Controller/SRCControllerPanel.cpp
@@ -163,20 +163,49 @@ void SRCControllerPanel::DeleteSelectedPanelItems()
 	}
 }
 
+bool SRCControllerPanel::CanDuplicateItems() const
+{
+	for (const TSharedPtr<FRCLogicModeBase>& LogicItem : GetSelectedLogicItems())
+	{
+		if (LogicItem->GetModelType() == FRCControllerModel::ControllerModelName)
+		{
+			return true;
+		}
+	}
+
+	return false;
+}
+
 void SRCControllerPanel::DuplicateSelectedPanelItems()
 {
 	for (const TSharedPtr<FRCLogicModeBase>& LogicItem : GetSelectedLogicItems())
 	{
-		if (const TSharedPtr<FRCControllerModel>& ControllerItem = StaticCastSharedPtr<FRCControllerModel>(LogicItem))
+		if (LogicItem->GetModelType() == FRCControllerModel::ControllerModelName)
 		{
-			if (URCController* Controller = Cast<URCController>(ControllerItem->GetVirtualProperty()))
+			if (const TSharedPtr<FRCControllerModel>& ControllerItem = StaticCastSharedPtr<FRCControllerModel>(LogicItem))
 			{
-				DuplicateController(Controller);
+				if (URCController* Controller = Cast<URCController>(ControllerItem->GetVirtualProperty()))
+				{
+					DuplicateController(Controller);
+				}
 			}
 		}
 	}
 }
 
+bool SRCControllerPanel::CanCopyItems() const
+{
+	for (const TSharedPtr<FRCLogicModeBase>& LogicItem : GetSelectedLogicItems())
+	{
+		if (LogicItem->GetModelType() == FRCControllerModel::ControllerModelName)
+		{
+			return true;
+		}
+	}
+
+	return false;
+}
+
 void SRCControllerPanel::CopySelectedPanelItems()
 {
 	if (const TSharedPtr<SRemoteControlPanel>& RemoteControlPanel = GetRemoteControlPanel())
@@ -187,11 +216,14 @@ void SRCControllerPanel::CopySelectedPanelItems()
 
 		for (const TSharedPtr<FRCLogicModeBase>& LogicItem : LogicItems)
 		{
-			if (const TSharedPtr<FRCControllerModel>& ControllerItem = StaticCastSharedPtr<FRCControllerModel>(LogicItem))
+			if (LogicItem->GetModelType() == FRCControllerModel::ControllerModelName)
 			{
-				if (URCController* Controller = Cast<URCController>(ControllerItem->GetVirtualProperty()))
+				if (const TSharedPtr<FRCControllerModel>& ControllerItem = StaticCastSharedPtr<FRCControllerModel>(LogicItem))
 				{
-					ItemsToCopy.Add(Controller);
+					if (URCController* Controller = Cast<URCController>(ControllerItem->GetVirtualProperty()))
+					{
+						ItemsToCopy.Add(Controller);
+					}
 				}
 			}
 		}
diff --git a/Engine/Plugins/VirtualProduction/RemoteControl/Source/RemoteControlUI/Private/UI/Controller/SRCControllerPanel.h b/Engine/Plugins/VirtualProduction/RemoteControl/Source/RemoteControlUI/Private/UI/Controller/SRCControllerPanel.h
--- a/Engine/Plugins/VirtualProduction/RemoteControl/Source/RemoteControlUI/Private/UI/Controller/SRCControllerPanel.h
+++ b/Engine/Plugins/VirtualProduction/RemoteControl/Source/RemoteControlUI/Private/UI/Controller/SRCControllerPanel.h
@@ -74,9 +74,15 @@ public:
 	/** Delete Items UI command implementation for this panel */
 	virtual void DeleteSelectedPanelItems() override;
 
+	/** Checks to make sure selected items are duplicatable */
+	virtual bool CanDuplicateItems() const override;
+
 	/** "Duplicate Items" UI command implementation for Controller panel*/
 	virtual void DuplicateSelectedPanelItems() override;
 
+	/** Checks to make sure selected items are copyable. */
+	virtual bool CanCopyItems() const override;
+
 	/** "Copy Items" UI command implementation for Controller panel*/
 	virtual void CopySelectedPanelItems() override;
 
