Remote Control: fix bit shifting 32 bit value with subsequent expansion to 64 bit warning.

diff --git a/Engine/Plugins/VirtualProduction/RemoteControl/Source/RemoteControlProtocolWidgets/Private/ViewModels/ProtocolRangeViewModel.cpp b/Engine/Plugins/VirtualProduction/RemoteControl/Source/RemoteControlProtocolWidgets/Private/ViewModels/ProtocolRangeViewModel.cpp
--- a/Engine/Plugins/VirtualProduction/RemoteControl/Source/RemoteControlProtocolWidgets/Private/ViewModels/ProtocolRangeViewModel.cpp
+++ b/Engine/Plugins/VirtualProduction/RemoteControl/Source/RemoteControlProtocolWidgets/Private/ViewModels/ProtocolRangeViewModel.cpp
@@ -81,7 +81,7 @@ void FProtocolRangeViewModel::UpdateInputValueRange()
 	if (const FProperty* Property = GetInputProperty())
 	{
 		// If there's a mismatch between the RangePropertySize and Property, clamp to RangePropertySize
-		const int32 NewRangeTypeSize = GetBinding()->GetRemoteControlProtocolEntityPtr()->Get()->GetRangePropertySize();
+		const uint8 NewRangeTypeSize = GetBinding()->GetRemoteControlProtocolEntityPtr()->Get()->GetRangePropertySize();
 		FProperty* PropertyInContainer = InputProxyPropertyContainer->GetValueProperty();
 		int64 NewRangeMaxValue = 0;
 
@@ -100,9 +100,13 @@ void FProtocolRangeViewModel::UpdateInputValueRange()
 						NewRangeMaxValue = FCString::Atoi64(*RangeMaxValueStr);
 					}
 					// Infer from RangeTypeSize
+					else if (NewRangeTypeSize < 4)
+					{
+						NewRangeMaxValue = (static_cast<int64>(1) << (NewRangeTypeSize * 8)) - 1;
+					}
 					else
 					{
-						NewRangeMaxValue = NewRangeTypeSize < 4 ? (1 << (NewRangeTypeSize * 8)) - 1 : TNumericLimits<uint32>::Max();
+						NewRangeMaxValue = TNumericLimits<uint32>::Max();
 					}
 					
 					if(PreviousRangeMaxValue == 0)
