Remote Control: fixed issue with playing a level in PIE with embedded remote control preset and opening another level while PIE is running.

diff --git a/Engine/Plugins/VirtualProduction/RemoteControl/Source/RemoteControlUI/Private/UI/Controller/SRCControllerPanelList.cpp b/Engine/Plugins/VirtualProduction/RemoteControl/Source/RemoteControlUI/Private/UI/Controller/SRCControllerPanelList.cpp
--- a/Engine/Plugins/VirtualProduction/RemoteControl/Source/RemoteControlUI/Private/UI/Controller/SRCControllerPanelList.cpp
+++ b/Engine/Plugins/VirtualProduction/RemoteControl/Source/RemoteControlUI/Private/UI/Controller/SRCControllerPanelList.cpp
@@ -158,10 +158,14 @@ void SRCControllerPanelList::Reset()
 	check(ControllerPanelWeakPtr.IsValid());
 
 	TSharedPtr<SRCControllerPanel> ControllerPanel = ControllerPanelWeakPtr.Pin();
-	URemoteControlPreset* Preset = ControllerPanel->GetPreset();
 	TSharedPtr<SRemoteControlPanel> RemoteControlPanel = ControllerPanel->GetRemoteControlPanel();
-	
-	check(Preset);
+
+	URemoteControlPreset* const Preset = ControllerPanel->GetPreset();
+	if (!Preset)
+	{
+		PropertyRowGenerator->SetStructure(TSharedPtr<FStructOnScope>());
+		return;
+	}
 
 	PropertyRowGenerator->SetStructure(Preset->GetControllerContainerStructOnScope());
 	if (!PropertyRowGenerator->OnFinishedChangingProperties().IsBoundToObject(this))
diff --git a/Engine/Plugins/VirtualProduction/RemoteControl/Source/RemoteControlUI/Private/UI/SRCPanelExposedEntitiesList.cpp b/Engine/Plugins/VirtualProduction/RemoteControl/Source/RemoteControlUI/Private/UI/SRCPanelExposedEntitiesList.cpp
--- a/Engine/Plugins/VirtualProduction/RemoteControl/Source/RemoteControlUI/Private/UI/SRCPanelExposedEntitiesList.cpp
+++ b/Engine/Plugins/VirtualProduction/RemoteControl/Source/RemoteControlUI/Private/UI/SRCPanelExposedEntitiesList.cpp
@@ -184,7 +184,7 @@ void SRCPanelExposedEntitiesList::Construct(const FArguments& InArgs, URemoteCon
 
 	bIsInLiveMode = InArgs._LiveMode;
 	bIsInProtocolsMode = InArgs._ProtocolsMode;
-	Preset = TStrongObjectPtr<URemoteControlPreset>(InPreset);
+	PresetWeak = InPreset;
 	OnEntityListUpdatedDelegate = InArgs._OnEntityListUpdated;
 	WidgetRegistry = MoveTemp(InWidgetRegistry);
 
@@ -214,9 +214,9 @@ void SRCPanelExposedEntitiesList::Construct(const FArguments& InArgs, URemoteCon
 	const TArray<FName>& HiddenColumns = Settings->EntitiesListHiddenColumns.Array();
 
 	// The Default group will be selected by default.
-	if (Preset.IsValid())
+	if (InPreset)
 	{
-		const FRemoteControlPresetGroup& DefaultGroup = Preset->Layout.GetDefaultGroup();
+		const FRemoteControlPresetGroup& DefaultGroup = InPreset->Layout.GetDefaultGroup();
 
 		CurrentlySelectedGroup = DefaultGroup.Id;
 	}
@@ -334,7 +334,7 @@ void SRCPanelExposedEntitiesList::Construct(const FArguments& InArgs, URemoteCon
 			{
 				return bIsInProtocolsMode.Get() ? ActiveProtocol : NAME_None;
 			});
-		const TSharedRef<SWidget> ExtenderWidget = Extender->MakeWidget(Preset.Get(), Args);
+		const TSharedRef<SWidget> ExtenderWidget = Extender->MakeWidget(InPreset, Args);
 
 		ExposedEntitiesPanelContent->AddSlot()
 			.AutoHeight()
@@ -468,7 +468,7 @@ void SRCPanelExposedEntitiesList::Construct(const FArguments& InArgs, URemoteCon
 
 	// Per Preset Protocol Settings
 	const TSharedRef<SRCPanelPerPresetProtocolSettings> PerPresetProtocolSettings =
-		SNew(SRCPanelPerPresetProtocolSettings, Preset.Get())
+		SNew(SRCPanelPerPresetProtocolSettings, InPreset)
 		.Visibility_Lambda([this]()
 			{
 				return bIsInProtocolsMode.Get() ? EVisibility::Visible : EVisibility::Collapsed;
@@ -552,7 +552,7 @@ void SRCPanelExposedEntitiesList::Tick(const FGeometry& AllottedGeometry, const
 
 	if (bUpdateEntityUsage)
 	{
-		if (Preset.IsValid())
+		if (URemoteControlPreset* Preset = GetPreset())
 		{
 			Preset->UpdateEntityUsage();
 		}
@@ -754,7 +754,8 @@ void SRCPanelExposedEntitiesList::ExposedEntitiesNodesRefresh()
 
 void SRCPanelExposedEntitiesList::OnPropertyIdRenamed(const FName InNewId, TSharedPtr<SRCPanelTreeNode> InNode)
 {
-	if (IsEntitySelected(InNode))
+	URemoteControlPreset* const Preset = GetPreset();
+	if (Preset && IsEntitySelected(InNode))
 	{
 		TArray<TSharedPtr<SRCPanelTreeNode>> SelectedEntities = GetSelectedEntities();
 		for (const TSharedPtr<SRCPanelTreeNode>& Entity : SelectedEntities)
@@ -778,6 +779,12 @@ void SRCPanelExposedEntitiesList::OnPropertyIdRenamed(const FName InNewId, TShar
 
 void SRCPanelExposedEntitiesList::OnLabelModified(const FName InOldName, const FName InNewName)
 {
+	URemoteControlPreset* const Preset = GetPreset();
+	if (!Preset)
+	{
+		return;
+	}
+
 	TArray<TSharedPtr<SRCPanelTreeNode>> SelectedEntities = GetSelectedEntities();
 	for (const TSharedPtr<SRCPanelTreeNode>& EntityNode : SelectedEntities)
 	{
@@ -794,7 +801,7 @@ void SRCPanelExposedEntitiesList::OnLabelModified(const FName InOldName, const F
 			ExposedField.Pin()->Rename(InNewName);
 			const FName NewName = ExposedField.Pin()->GetLabel();
 			EntityNode->SetName(NewName);
-			Preset->OnFieldRenamed().Broadcast(Preset.Get(), OldLabel, NewName);
+			Preset->OnFieldRenamed().Broadcast(Preset, OldLabel, NewName);
 		}
 	}
 }
@@ -843,7 +850,7 @@ void SRCPanelExposedEntitiesList::StoreListSettings()
 		Settings.FieldGroupType = CurrentGroupType;
 		Settings.FieldGroupOrder = CurrentGroupSortType;
 
-		(*SettingsPtr)->OnSettingsChanged(Preset.Get(), Settings);
+		(*SettingsPtr)->OnSettingsChanged(GetPreset(), Settings);
 	}
 	else
 	{
@@ -860,7 +867,7 @@ void SRCPanelExposedEntitiesList::RecallListSettings()
 
 	if (SettingsPtr)
 	{
-		const FRCPanelExposedEntitiesListSettingsData Settings = (*SettingsPtr)->GetListSettings(Preset.Get());
+		const FRCPanelExposedEntitiesListSettingsData Settings = (*SettingsPtr)->GetListSettings(GetPreset());
 		CurrentGroupType = Settings.FieldGroupType;
 		CurrentGroupSortType = Settings.FieldGroupOrder;
 	}
@@ -895,12 +902,10 @@ void SRCPanelExposedEntitiesList::SelectSource_Execute()
 
 void SRCPanelExposedEntitiesList::UpdateEntityUsage_Execute()
 {
-	if (!Preset.IsValid())
+	if (URemoteControlPreset* Preset = GetPreset())
 	{
-		return;
+		Preset->UpdateEntityUsage();
 	}
-
-	Preset->UpdateEntityUsage();
 }
 
 void SRCPanelExposedEntitiesList::UpdateEntityUsage()
@@ -938,7 +943,7 @@ void SRCPanelExposedEntitiesList::OnObjectPropertyChange(UObject* InObject, FPro
 			}
 		}
 
-		if (Preset)
+		if (URemoteControlPreset* const Preset = GetPreset())
 		{
 			// If the modified property is a parent of an exposed property, re-enable the edit condition.
 			// This is useful in case we re-add an array element which contains a nested property that is exposed.
@@ -995,7 +1000,8 @@ void SRCPanelExposedEntitiesList::TryRefreshingSearch(const FText& InSearchText,
 {
 	FieldEntities.Reset();
 
-	if (FieldWidgetMap.IsEmpty() || !Preset.IsValid())
+	URemoteControlPreset* const Preset = GetPreset();
+	if (FieldWidgetMap.IsEmpty() || !Preset)
 	{
 		FieldsListView->RequestListRefresh();
 
@@ -1045,6 +1051,12 @@ void SRCPanelExposedEntitiesList::GenerateListWidgets()
 {
 	TRACE_CPUPROFILER_EVENT_SCOPE(SRCPanelExposedEntitiesList::GenerateListWidgets);
 
+	URemoteControlPreset* const Preset = GetPreset();
+	if (!Preset)
+	{
+		return;
+	}
+
 	TArray<FGuid> OrderMap = Preset->Layout.GetDefaultGroupOrder();
 	FieldWidgetMap.Reset();
 	CachedEntityOwners.Reset();
@@ -1063,7 +1075,7 @@ void SRCPanelExposedEntitiesList::GenerateListWidgets()
 
 			FGenerateWidgetArgs Args;
 			Args.Entity = Entity;
-			Args.Preset = Preset.Get();
+			Args.Preset = Preset;
 			Args.WidgetRegistry = WidgetRegistry;
 			Args.ColumnSizeData = ColumnSizeData;
 			Args.bIsInLiveMode = bIsInLiveMode;
@@ -1107,7 +1119,8 @@ void SRCPanelExposedEntitiesList::GenerateListWidgets(const FRemoteControlPreset
 		return;
 	}
 
-	if (Preset->Layout.IsDefaultGroup(FromGroup.Id))
+	URemoteControlPreset* const Preset = GetPreset();
+	if (Preset && Preset->Layout.IsDefaultGroup(FromGroup.Id))
 	{
 		FieldWidgetMap.GenerateValueArray(CachedFieldEntities);
 		CreateGroupsAndSort();
@@ -1123,11 +1136,17 @@ void SRCPanelExposedEntitiesList::GenerateListWidgets(const FRemoteControlPreset
 
 void SRCPanelExposedEntitiesList::RefreshGroups()
 {
+	URemoteControlPreset* const Preset = GetPreset();
+	if (!Preset)
+	{
+		return;
+	}
+
 	FieldGroups.Reset(Preset->Layout.GetGroups().Num());
 
 	for (const FRemoteControlPresetGroup& RCGroup : Preset->Layout.GetGroups())
 	{
-		TSharedRef<SRCPanelGroup> FieldGroup = SNew(SRCPanelGroup, Preset.Get(), ColumnSizeData)
+		TSharedRef<SRCPanelGroup> FieldGroup = SNew(SRCPanelGroup, Preset, ColumnSizeData)
 			.Id(RCGroup.Id)
 			.Name(RCGroup.Name)
 			.OnFieldDropEvent(this, &SRCPanelExposedEntitiesList::OnDropOnGroup)
@@ -1214,7 +1233,8 @@ void SRCPanelExposedEntitiesList::OnGetNodeChildren(TSharedPtr<SRCPanelTreeNode>
 
 void SRCPanelExposedEntitiesList::OnSelectionChanged(TSharedPtr<SRCPanelTreeNode> Node, ESelectInfo::Type SelectInfo)
 {
-	if (!Node || SelectInfo != ESelectInfo::OnMouseClick)
+	URemoteControlPreset* const Preset = GetPreset();
+	if (!Node || SelectInfo != ESelectInfo::OnMouseClick || !Preset)
 	{
 		FieldsListView->ClearSelection();
 
@@ -1247,6 +1267,12 @@ FReply SRCPanelExposedEntitiesList::OnDropOnGroup(const TSharedPtr<FDragDropOper
 {
 	checkSlow(DragTargetGroup);
 
+	URemoteControlPreset* const Preset = GetPreset();
+	if (!Preset)
+	{
+		return FReply::Unhandled();
+	}
+
 	// will be changed, for now as long as you have a grouping type active you will not be able to drag fields to re-order them
 	if (DragDropOperation->IsOfType<FExposedEntityDragDrop>() && CurrentGroupType == ERCFieldGroupType::None)
 	{
@@ -1337,27 +1363,36 @@ FReply SRCPanelExposedEntitiesList::OnDropOnGroup(const TSharedPtr<FDragDropOper
 FGuid SRCPanelExposedEntitiesList::GetGroupId(const FGuid& EntityId)
 {
 	FGuid GroupId;
-	if (FRemoteControlPresetGroup* Group = Preset->Layout.FindGroupFromField(EntityId))
+	if (URemoteControlPreset* Preset = GetPreset())
 	{
-		GroupId = Group->Id;
+		if (FRemoteControlPresetGroup* Group = Preset->Layout.FindGroupFromField(EntityId))
+		{
+			GroupId = Group->Id;
+		}
 	}
-
 	return GroupId;
 }
 
 FReply SRCPanelExposedEntitiesList::OnCreateGroup()
 {
-	FScopedTransaction Transaction(LOCTEXT("CreateGroup", "Create Group"));
-	Preset->Modify();
-	Preset->Layout.CreateGroup();
-	return FReply::Handled();
+	if (URemoteControlPreset* Preset = GetPreset())
+	{
+		FScopedTransaction Transaction(LOCTEXT("CreateGroup", "Create Group"));
+		Preset->Modify();
+		Preset->Layout.CreateGroup();
+		return FReply::Handled();
+	}
+	return FReply::Unhandled();
 }
 
 void SRCPanelExposedEntitiesList::OnDeleteGroup(const FGuid& GroupId)
 {
-	FScopedTransaction Transaction(LOCTEXT("DeleteGroup", "Delete Group"));
-	Preset->Modify();
-	Preset->Layout.DeleteGroup(GroupId);
+	if (URemoteControlPreset* Preset = GetPreset())
+	{
+		FScopedTransaction Transaction(LOCTEXT("DeleteGroup", "Delete Group"));
+		Preset->Modify();
+		Preset->Layout.DeleteGroup(GroupId);
+	}
 }
 
 void SRCPanelExposedEntitiesList::OnEntityRebind(const FGuid& InEntityGuid)
@@ -1458,9 +1493,9 @@ void SRCPanelExposedEntitiesList::OnSearchTextCommitted(const FText& InFilterTex
 
 void SRCPanelExposedEntitiesList::PopulateSearchStrings(const SRCPanelTreeNode& Item, TArray<FString>& OutSearchStrings) const
 {
-	if (Preset.IsValid())
+	if (const URemoteControlPreset* Preset = GetPreset())
 	{
-		if (const TSharedPtr<FRemoteControlEntity> Entity = Preset->GetExposedEntity<FRemoteControlEntity>(Item.GetRCId()).Pin())
+		if (const TSharedPtr<const FRemoteControlEntity> Entity = Preset->GetExposedEntity<FRemoteControlEntity>(Item.GetRCId()).Pin())
 		{
 			OutSearchStrings.Add(Entity->GetLabel().ToString());
 		}
@@ -1575,6 +1610,12 @@ void SRCPanelExposedEntitiesList::OnCreateFieldGroup(ERCFieldGroupType InFieldGr
 
 void SRCPanelExposedEntitiesList::CreateFieldGroup()
 {
+	URemoteControlPreset* const Preset = GetPreset();
+	if (!Preset)
+	{
+		return;
+	}
+
 	FieldEntitiesGroups.Reset();
 	TSet<FName> GroupKeys;
 	if (CurrentGroupType == ERCFieldGroupType::PropertyId)
@@ -1586,7 +1627,7 @@ void SRCPanelExposedEntitiesList::CreateFieldGroup()
 				if (!GroupKeys.Contains(ExposedField->GetPropertyId()))
 				{
 					GroupKeys.Add(ExposedField->GetPropertyId());
-					TSharedRef<SRCPanelExposedEntitiesGroup> GroupWidget = SNew(SRCPanelExposedEntitiesGroup, CurrentGroupType, Preset.Get())
+					TSharedRef<SRCPanelExposedEntitiesGroup> GroupWidget = SNew(SRCPanelExposedEntitiesGroup, CurrentGroupType, Preset)
 						.FieldKey(ExposedField->GetPropertyId())
 						.OnGroupPropertyIdChanged(this, &SRCPanelExposedEntitiesList::CreateGroupsAndSort);
 
@@ -1608,7 +1649,7 @@ void SRCPanelExposedEntitiesList::CreateFieldGroup()
 				{
 					GroupKeys.Add(OwnerPathName);
 
-					TSharedRef<SRCPanelExposedEntitiesGroup> GroupWidget = SNew(SRCPanelExposedEntitiesGroup, CurrentGroupType, Preset.Get())
+					TSharedRef<SRCPanelExposedEntitiesGroup> GroupWidget = SNew(SRCPanelExposedEntitiesGroup, CurrentGroupType, Preset)
 						.FieldKey(OwnerPathName);
 
 					GroupWidget->AssignChildren(CachedFieldEntities);
@@ -1719,6 +1760,12 @@ void SRCPanelExposedEntitiesList::CreateGroupsAndSort()
 
 void SRCPanelExposedEntitiesList::RegisterPresetDelegates()
 {
+	URemoteControlPreset* const Preset = GetPreset();
+	if (!Preset)
+	{
+		return;
+	}
+
 	FRemoteControlPresetLayout& Layout = Preset->Layout;
 	Layout.OnGroupAdded().AddSP(this, &SRCPanelExposedEntitiesList::OnGroupAdded);
 	Layout.OnGroupDeleted().AddSP(this, &SRCPanelExposedEntitiesList::OnGroupDeleted);
@@ -1733,7 +1780,7 @@ void SRCPanelExposedEntitiesList::RegisterPresetDelegates()
 
 void SRCPanelExposedEntitiesList::UnregisterPresetDelegates()
 {
-	if (Preset)
+	if (URemoteControlPreset* Preset = GetPreset())
 	{
 		FRemoteControlPresetLayout& Layout = Preset->Layout;
 		Preset->OnEntitiesUpdated().RemoveAll(this);
@@ -1750,23 +1797,28 @@ void SRCPanelExposedEntitiesList::UnregisterPresetDelegates()
 
 void SRCPanelExposedEntitiesList::OnEntityAdded(const FGuid& InEntityId)
 {
-	auto ExposeEntity = [this, InEntityId](TSharedPtr<SRCPanelTreeNode>&& Node)
+	URemoteControlPreset* const Preset = GetPreset();
+	if (!Preset)
+	{
+		return;
+	}
+
+	auto ExposeEntity = [This=this, Preset, InEntityId](TSharedPtr<SRCPanelTreeNode>&& Node)
 	{
 		if (Node)
 		{
-			FieldWidgetMap.Add(InEntityId, Node);
+			This->FieldWidgetMap.Add(InEntityId, Node);
 
-			if (TSharedPtr<SRCPanelGroup> SRCGroup = FindGroupById(GetGroupId(InEntityId)))
+			if (TSharedPtr<SRCPanelGroup> SRCGroup = This->FindGroupById(This->GetGroupId(InEntityId)))
 			{
 				SRCGroup->GetNodes().Add(MoveTemp(Node));
-
-				GenerateListWidgets(*Preset->Layout.GetGroup(GetGroupId(InEntityId)));
+				This->GenerateListWidgets(*Preset->Layout.GetGroup(This->GetGroupId(InEntityId)));
 			}
 		}
 	};
 	
 	FGenerateWidgetArgs Args;
-	Args.Preset = Preset.Get();
+	Args.Preset = Preset;
 	Args.WidgetRegistry = WidgetRegistry;
 	Args.ColumnSizeData = ColumnSizeData;
 	Args.bIsInLiveMode = bIsInLiveMode;
@@ -1796,13 +1848,22 @@ void SRCPanelExposedEntitiesList::OnEntityRemoved(const FGuid& InGroupId, const
 	}
 
 	FieldWidgetMap.Remove(InEntityId);
-	
-	GenerateListWidgets(*Preset->Layout.GetGroup(InGroupId));
+
+	if (URemoteControlPreset* Preset = GetPreset())
+	{
+		GenerateListWidgets(*Preset->Layout.GetGroup(InGroupId));
+	}
 }
 
 void SRCPanelExposedEntitiesList::OnGroupAdded(const FRemoteControlPresetGroup& Group)
 {
-	TSharedRef<SRCPanelGroup> FieldGroup = SNew(SRCPanelGroup, Preset.Get(), ColumnSizeData)
+	URemoteControlPreset* const Preset = GetPreset();
+	if (!Preset)
+	{
+		return;
+	}
+
+	TSharedRef<SRCPanelGroup> FieldGroup = SNew(SRCPanelGroup, Preset, ColumnSizeData)
 		.Id(Group.Id)
 		.Name(Group.Name)
 		.OnFieldDropEvent(this, &SRCPanelExposedEntitiesList::OnDropOnGroup)
@@ -1907,6 +1968,7 @@ void SRCPanelExposedEntitiesList::OnFieldOrderChanged(const FGuid& GroupId, cons
 
 		if (*Group)
 		{
+			URemoteControlPreset* const Preset = GetPreset();
 			if (Preset && Preset->Layout.IsDefaultGroup((*Group)->GetRCId()))
 			{
 				FieldWidgetMap.ValueSort(
@@ -1954,6 +2016,12 @@ void SRCPanelExposedEntitiesList::OnEntitiesUpdated(URemoteControlPreset*, const
 
 void SRCPanelExposedEntitiesList::ApplyFilters()
 {
+	URemoteControlPreset* const Preset = GetPreset();
+	if (!Preset)
+	{
+		return;
+	}
+
 	if (BackendFilter.HasAnyActiveFilters())
 	{
 		// If we are not actively searching anything then include all the entities.
@@ -2011,7 +2079,8 @@ void SRCPanelExposedEntitiesList::RequestSearchOrFilter()
 
 FReply SRCPanelExposedEntitiesList::RequestDeleteAllEntities()
 {
-	if (!FieldsListView.IsValid() || !Preset.IsValid())
+	URemoteControlPreset* const Preset = GetPreset();
+	if (!FieldsListView.IsValid() || !Preset)
 	{
 		return FReply::Unhandled();
 	}
@@ -2036,7 +2105,8 @@ FReply SRCPanelExposedEntitiesList::RequestDeleteAllEntities()
 
 FReply SRCPanelExposedEntitiesList::RequestDeleteAllGroups()
 {
-	if (!GroupsListView.IsValid() || !Preset.IsValid())
+	URemoteControlPreset* const Preset = GetPreset();
+	if (!GroupsListView.IsValid() || !Preset)
 	{
 		return FReply::Unhandled();
 	}
@@ -2254,7 +2324,7 @@ void SRCPanelExposedEntitiesList::ProcessRefresh()
 
 	RefreshGroups();
 
-	if (Preset.IsValid())
+	if (URemoteControlPreset* const Preset = GetPreset())
 	{
 		constexpr bool bForceMouseClick = true;
 
diff --git a/Engine/Plugins/VirtualProduction/RemoteControl/Source/RemoteControlUI/Private/UI/SRCPanelExposedEntitiesList.h b/Engine/Plugins/VirtualProduction/RemoteControl/Source/RemoteControlUI/Private/UI/SRCPanelExposedEntitiesList.h
--- a/Engine/Plugins/VirtualProduction/RemoteControl/Source/RemoteControlUI/Private/UI/SRCPanelExposedEntitiesList.h
+++ b/Engine/Plugins/VirtualProduction/RemoteControl/Source/RemoteControlUI/Private/UI/SRCPanelExposedEntitiesList.h
@@ -89,6 +89,11 @@ public:
 	virtual void Tick(const FGeometry& AllottedGeometry, const double InCurrentTime, const float InDeltaTime) override;
 	//~ End SWidget Interface
 
+	URemoteControlPreset* GetPreset() const
+	{
+		return PresetWeak.Get();
+	}
+
 	/** Get the currently selected group. */
 	TSharedPtr<SRCPanelTreeNode> GetSelectedGroup() const;
 	
@@ -310,8 +315,8 @@ private:
 	TAttribute<bool> bIsInLiveMode;
 	/** Whether the panel is in protocols mode. */
 	TAttribute<bool> bIsInProtocolsMode;
-	/** Holds the preset asset. */
-	TStrongObjectPtr<URemoteControlPreset> Preset;
+	/** Holds a weak reference to the preset asset. */
+	TWeakObjectPtr<URemoteControlPreset> PresetWeak;
 	/** Handle to the delegate called when an object property change is detected. */
 	FDelegateHandle OnPropertyChangedHandle;
 	/** Handle to the delegate called when a binding is added or removed. */
diff --git a/Engine/Plugins/VirtualProduction/RemoteControl/Source/RemoteControlUI/Private/UI/SRemoteControlPanel.cpp b/Engine/Plugins/VirtualProduction/RemoteControl/Source/RemoteControlUI/Private/UI/SRemoteControlPanel.cpp
--- a/Engine/Plugins/VirtualProduction/RemoteControl/Source/RemoteControlUI/Private/UI/SRemoteControlPanel.cpp
+++ b/Engine/Plugins/VirtualProduction/RemoteControl/Source/RemoteControlUI/Private/UI/SRemoteControlPanel.cpp
@@ -109,12 +109,12 @@ void SRemoteControlPanel::AddReferencedObjects(FReferenceCollector& Collector)
 
 void SRemoteControlPanel::ApplyProtocolBindings()
 {
-	IRemoteControlProtocolModule::Get().ApplyProtocolBindings(Preset.Get());
+	IRemoteControlProtocolModule::Get().ApplyProtocolBindings(GetPreset());
 }
 
 void SRemoteControlPanel::UnapplyProtocolBindings()
 {
-	IRemoteControlProtocolModule::Get().UnapplyProtocolBindings(Preset.Get());
+	IRemoteControlProtocolModule::Get().UnapplyProtocolBindings(GetPreset());
 }
 
 namespace RemoteControlPanelUtils
@@ -144,7 +144,7 @@ namespace RemoteControlPanelUtils
 void SRemoteControlPanel::Construct(const FArguments& InArgs, URemoteControlPreset* InPreset, TSharedPtr<IToolkitHost> InToolkitHost)
 {
 	OnLiveModeChange = InArgs._OnLiveModeChange;
-	Preset = TStrongObjectPtr<URemoteControlPreset>(InPreset);
+	PresetWeak = InPreset;
 	WidgetRegistry = MakeShared<FRCPanelWidgetRegistry>();
 	ToolkitHost = InToolkitHost;
 
@@ -202,7 +202,7 @@ void SRemoteControlPanel::Construct(const FArguments& InArgs, URemoteControlPres
 
 	EntityProtocolDetails = SNew(SBox);
 
-	EntityList = SNew(SRCPanelExposedEntitiesList, Preset.Get(), WidgetRegistry)
+	EntityList = SNew(SRCPanelExposedEntitiesList, InPreset, WidgetRegistry)
 		.ExposeActorsComboButton(CreateExposeActorsButton())
 		.ExposeFunctionsComboButton(CreateExposeFunctionsButton())
 		.LiveMode(this, &SRemoteControlPanel::IsModeActive, ERCPanelMode::Live)
@@ -503,16 +503,21 @@ bool SRemoteControlPanel::IsExposed(const FRCExposesPropertyArgs& InPropertyArgs
 
 bool SRemoteControlPanel::IsAllObjectsExposed(TArray<UObject*> InOuterObjects, const FString& InPath, bool bUsingDuplicatesInPath)
 {
-	TArray<TSharedPtr<FRemoteControlProperty>, TInlineAllocator<1>> PotentialMatches;
-	for (const TWeakPtr<FRemoteControlProperty>& WeakProperty : Preset->GetExposedEntities<FRemoteControlProperty>())
+	const URemoteControlPreset* const Preset = GetPreset();
+	if (!Preset)
 	{
-		if (TSharedPtr<FRemoteControlProperty> Property = WeakProperty.Pin())
+		return false;
+	}
+
+	TArray<TSharedPtr<const FRemoteControlProperty>, TInlineAllocator<1>> PotentialMatches;
+	for (const TWeakPtr<const FRemoteControlProperty>& WeakProperty : Preset->GetExposedEntities<FRemoteControlProperty>())
+	{
+		if (TSharedPtr<const FRemoteControlProperty> Property = WeakProperty.Pin())
 		{
 			// If that was exposed by property path it should be checked by the full path with duplicated like propertypath.propertypath[0]
 			// If that was exposed by the owner object it should be without duplicated in the path, just propertypath[0]
-			const bool Isbound = bUsingDuplicatesInPath ? Property->CheckIsBoundToPropertyPath(InPath) : Property->CheckIsBoundToString(InPath);
-
-			if (Isbound)
+			const bool bIsBound = bUsingDuplicatesInPath ? Property->CheckIsBoundToPropertyPath(InPath) : Property->CheckIsBoundToString(InPath);
+			if (bIsBound)
 			{
 				PotentialMatches.Add(Property);
 			}
@@ -525,7 +530,7 @@ bool SRemoteControlPanel::IsAllObjectsExposed(TArray<UObject*> InOuterObjects, c
 	{
 		bool bFoundPropForObject = false;
 
-		for (const TSharedPtr<FRemoteControlProperty>& Property : PotentialMatches)
+		for (const TSharedPtr<const FRemoteControlProperty>& Property : PotentialMatches)
 		{
 			if (Property->ContainsBoundObjects({ InOuterObjects } ))
 			{
@@ -553,6 +558,12 @@ void SRemoteControlPanel::ExecutePropertyAction(const FRCExposesPropertyArgs& In
 		return;
 	}
 
+	URemoteControlPreset* const Preset = GetPreset();
+	if (!Preset)
+	{
+		return;
+	}
+
 	if (IsExposed(InPropertyArgs))
 	{
 		FScopedTransaction Transaction(LOCTEXT("UnexposeProperty", "Unexpose Property"));
@@ -850,8 +861,10 @@ TSharedRef<SWidget> SRemoteControlPanel::CreateExposeActorsButton()
 			);
 		};
 
+		URemoteControlPreset* const Preset = GetPreset();
+
 		// SObjectPropertyEntryBox does not support non-level editor worlds.
-		if (!Preset.IsValid() || !Preset->IsEmbeddedPreset())
+		if (!Preset || !Preset->IsEmbeddedPreset())
 		{
 			MenuBuilder.AddWidget(
 				SNew(SObjectPropertyEntryBox)
@@ -868,18 +881,18 @@ TSharedRef<SWidget> SRemoteControlPanel::CreateExposeActorsButton()
 			MenuBuilder.AddSubMenu(
 				LOCTEXT("SelectActor", "Select Actor"),
 				LOCTEXT("SelectActorTooltip", "Select an actor to Remote Control."),
-				FNewMenuDelegate::CreateLambda(
-					[this](FMenuBuilder& SubMenuBuilder)
+				FNewMenuDelegate::CreateSPLambda(this,
+					[This=this](FMenuBuilder& SubMenuBuilder)
 					{
 						FSceneOutlinerModule& SceneOutlinerModule = FModuleManager::LoadModuleChecked<FSceneOutlinerModule>("SceneOutliner");
 						FSceneOutlinerInitializationOptions InitOptions;
 						constexpr bool bAllowPIE = false;
-						UWorld* PresetWorld = URemoteControlPreset::GetWorld(Preset.Get(), bAllowPIE);
+						UWorld* PresetWorld = URemoteControlPreset::GetWorld(This->GetPreset(), bAllowPIE);
 
 						SubMenuBuilder.AddWidget(
 							SceneOutlinerModule.CreateActorPicker(
 								InitOptions,
-								FOnActorPicked::CreateSP(this, &SRemoteControlPanel::ExposeActor),
+								FOnActorPicked::CreateSP(This, &SRemoteControlPanel::ExposeActor),
 								PresetWorld
 							),
 							FText::GetEmpty(), true, false
@@ -962,25 +975,25 @@ TSharedRef<SWidget> SRemoteControlPanel::CreateExposeByClassWidget()
 		Options.bShowUnloadedBlueprints = false;
 	}
 
-	TSharedRef<SWidget> Widget = FModuleManager::LoadModuleChecked<FClassViewerModule>("ClassViewer").CreateClassViewer(Options, FOnClassPicked::CreateLambda(
-		[this](UClass* ChosenClass)
+	TSharedRef<SWidget> Widget = FModuleManager::LoadModuleChecked<FClassViewerModule>("ClassViewer").CreateClassViewer(Options, FOnClassPicked::CreateSPLambda(this,
+		[This=this](UClass* ChosenClass)
 		{
 			constexpr bool bAllowPIE = false;
 
-			if (UWorld* World = URemoteControlPreset::GetWorld(Preset.Get(), bAllowPIE))
+			if (UWorld* World = URemoteControlPreset::GetWorld(This->GetPreset(), bAllowPIE))
 			{
 				for (TActorIterator<AActor> It(World, ChosenClass, EActorIteratorFlags::SkipPendingKill); It; ++It)
 				{
 					if (RemoteControlPanelUtils::IsExposableActor(*It))
 					{
-						ExposeActor(*It);
+						This->ExposeActor(*It);
 					}
 				}
 			}
 
-			if (ExposeActorsComboButton)
+			if (This->ExposeActorsComboButton)
 			{
-				ExposeActorsComboButton->SetIsOpen(false);
+				This->ExposeActorsComboButton->SetIsOpen(false);
 			}
 		}));
 
@@ -998,7 +1011,7 @@ void SRemoteControlPanel::CacheLevelClasses()
 	CachedClassesInLevel.Empty();
 	constexpr bool bAllowPIE = false;
 
-	if (UWorld* World = URemoteControlPreset::GetWorld(Preset.Get(), bAllowPIE))
+	if (UWorld* World = URemoteControlPreset::GetWorld(GetPreset(), bAllowPIE))
 	{
 		for (TActorIterator<AActor> It(World, AActor::StaticClass(), EActorIteratorFlags::SkipPendingKill); It; ++It)
 		{
@@ -1072,10 +1085,12 @@ void SRemoteControlPanel::OnMapChange(uint32)
 	// Clear the widget cache on map change to make sure we don't keep widgets around pointing to potentially stale objects.
 	WidgetRegistry->Clear();
 
+	const URemoteControlPreset* const Preset = GetPreset();
+
 	// Update the world to the newer level editor map change
 	// Only do this if the current selected world is now invalid (stale due to map change).
 	// GWorld is used here instead of GEditor->EditorWorld as it's the one set to the newer world prior to the delegate being called (GEditor->EditorWorld set at a later time)
-	if (GWorld && !Preset->SelectedWorld.IsValid())
+	if (GWorld && (!Preset || !Preset->SelectedWorld.IsValid()))
 	{
 		UpdatePanelForWorld(GWorld);
 	}
@@ -1261,7 +1276,8 @@ void SRemoteControlPanel::BindRemoteControlCommands()
 
 void SRemoteControlPanel::OnObjectReplaced(const TMap<UObject*, UObject*>& InObjectReplaced)
 {
-	if (!WidgetRegistry.IsValid() || !Preset.IsValid() || !Preset->SelectedWorld.IsValid())
+	URemoteControlPreset* const Preset = GetPreset();
+	if (!WidgetRegistry.IsValid() || !Preset || !Preset->SelectedWorld.IsValid())
 	{
 		return;
 	}
@@ -1284,8 +1300,10 @@ void SRemoteControlPanel::OnObjectReplaced(const TMap<UObject*, UObject*>& InObj
 
 void SRemoteControlPanel::PostPIEStarted(const bool bInIsSimulating)
 {
+	URemoteControlPreset* const Preset = GetPreset();
+
 	// Don't do anything when it's Simulate In Editor. This panel will only react when it's a proper Play In Editor
-	if (!Preset.IsValid() || bInIsSimulating)
+	if (!Preset || bInIsSimulating)
 	{
 		return;
 	}
@@ -1304,26 +1322,41 @@ void SRemoteControlPanel::PostPIEStarted(const bool bInIsSimulating)
 	}
 }
 
-void SRemoteControlPanel::OnEndPIE(const bool bInIsSimulating)
+void SRemoteControlPanel::PostLoadMapWithWorld(UWorld* InWorld)
 {
-	// Don't do anything when it's Simulate In Editor. This panel will only react when it's a proper Play In Editor
-	if (!Preset.IsValid() || bInIsSimulating)
+	if (!InWorld || InWorld->WorldType != EWorldType::PIE)
 	{
 		return;
 	}
 
-	// If the currently selected preset is not a PIE world when ending, no need to override what's there. Return early
-	const UWorld* PresetWorld = Preset->SelectedWorld.Get();
-	if (PresetWorld && PresetWorld->WorldType != EWorldType::PIE)
+	URemoteControlPreset* const Preset = GetPreset();
+
+	const UWorld* const SelectedWorld = Preset ? Preset->SelectedWorld.Get() : nullptr;
+	if (SelectedWorld && SelectedWorld->WorldType == EWorldType::PIE)
+	{
+		UpdatePanelForWorld(InWorld);
+	}
+}
+
+void SRemoteControlPanel::OnEndPIE(const bool bInIsSimulating)
+{
+	// Don't do anything when it's Simulate In Editor. This panel will only react when it's a proper Play In Editor
+	if (bInIsSimulating)
 	{
 		return;
 	}
 
-	if (Preset->IsEmbeddedPreset())
+	if (URemoteControlPreset* const Preset = GetPreset())
 	{
-		OpenEditorEmbeddedPreset();
+		// If the currently selected preset is not a PIE world when ending, no need to override what's there. Return early
+		const UWorld* PresetWorld = Preset->SelectedWorld.Get();
+		if (PresetWorld && PresetWorld->WorldType != EWorldType::PIE)
+		{
+			return;
+		}
 	}
-	else if (GEditor)
+
+	if (GEditor)
 	{
 		UpdatePanelForWorld(GEditor->EditorWorld);
 	}
@@ -1345,20 +1378,27 @@ void SRemoteControlPanel::RegisterEvents()
 		GEngine->OnLevelActorDeleted().AddSP(this, &SRemoteControlPanel::OnLevelActorsRemoved);
 	}
 
-	Preset->OnEntityExposed().AddSP(this, &SRemoteControlPanel::OnEntityExposed);
-	Preset->OnEntityUnexposed().AddSP(this, &SRemoteControlPanel::OnEntityUnexposed);
+	if (URemoteControlPreset* Preset = GetPreset())
+	{	
+		Preset->OnEntityExposed().AddSP(this, &SRemoteControlPanel::OnEntityExposed);
+		Preset->OnEntityUnexposed().AddSP(this, &SRemoteControlPanel::OnEntityUnexposed);
+	}
 
 	UMaterial::OnMaterialCompilationFinished().AddSP(this, &SRemoteControlPanel::OnMaterialCompiled);
 	FCoreUObjectDelegates::OnObjectsReplaced.AddSP(this, &SRemoteControlPanel::OnObjectReplaced);
 
+	FCoreUObjectDelegates::PostLoadMapWithWorld.AddSP(this, &SRemoteControlPanel::PostLoadMapWithWorld);
 	FEditorDelegates::PostPIEStarted.AddSP(this, &SRemoteControlPanel::PostPIEStarted);
 	FEditorDelegates::EndPIE.AddSP(this, &SRemoteControlPanel::OnEndPIE);
 }
 
 void SRemoteControlPanel::UnregisterEvents()
 {
-	Preset->OnEntityExposed().RemoveAll(this);
-	Preset->OnEntityUnexposed().RemoveAll(this);
+	if (URemoteControlPreset* Preset = GetPreset())
+	{
+		Preset->OnEntityExposed().RemoveAll(this);
+		Preset->OnEntityUnexposed().RemoveAll(this);
+	}
 
 	if (GEngine)
 	{
@@ -1376,6 +1416,7 @@ void SRemoteControlPanel::UnregisterEvents()
 	UMaterial::OnMaterialCompilationFinished().RemoveAll(this);
 	FCoreUObjectDelegates::OnObjectsReplaced.RemoveAll(this);
 
+	FCoreUObjectDelegates::PostLoadMapWithWorld.RemoveAll(this);
 	FEditorDelegates::PostPIEStarted.RemoveAll(this);
 	FEditorDelegates::EndPIE.RemoveAll(this);
 }
@@ -1422,7 +1463,13 @@ void SRemoteControlPanel::Unexpose(const FRCExposesPropertyArgs& InPropertyArgs)
 		return;
 	}
 
-	auto CheckAndUnexpose = [Preset=Preset](TArray<UObject*> InOuterObjects, const FString& InPath, bool bInUsingDuplicatesInPath)
+	URemoteControlPreset* const Preset = GetPreset();
+	if (!Preset)
+	{
+		return;
+	}
+
+	auto CheckAndUnexpose = [Preset](TArray<UObject*> InOuterObjects, const FString& InPath, bool bInUsingDuplicatesInPath)
 	{
 		// Find an exposed property with the same path.
 		TArray<TSharedPtr<FRemoteControlProperty>, TInlineAllocator<1>> PotentialMatches;
@@ -1480,7 +1527,8 @@ void SRemoteControlPanel::OnBlueprintReinstanced()
 
 void SRemoteControlPanel::ExposeProperty(UObject* Object, FRCFieldPathInfo Path, FString InDesiredName)
 {
-	if (Path.Resolve(Object))
+	URemoteControlPreset* const Preset = GetPreset();
+	if (Preset && Path.Resolve(Object))
 	{
 		FRemoteControlPresetExposeArgs Args;
 		Args.Label = InDesiredName;
@@ -1491,6 +1539,12 @@ void SRemoteControlPanel::ExposeProperty(UObject* Object, FRCFieldPathInfo Path,
 
 void SRemoteControlPanel::ExposeFunction(UObject* Object, UFunction* Function)
 {
+	URemoteControlPreset* const Preset = GetPreset();
+	if (!Preset)
+	{
+		return;
+	}
+
 	if (ExposeFunctionsComboButton)
 	{
 		ExposeFunctionsComboButton->SetIsOpen(false);
@@ -1511,16 +1565,24 @@ void SRemoteControlPanel::OnExposeActor(const FAssetData& AssetData)
 
 void SRemoteControlPanel::ExposeActor(AActor* Actor)
 {
-	if (Actor)
+	if (!Actor)
 	{
-		FScopedTransaction Transaction(LOCTEXT("ExposeActor", "Expose Actor"));
-		Preset->Modify();
-
-		FRemoteControlPresetExposeArgs Args;
-		Args.GroupId = GetSelectedGroup();
+		return;
+	}
 
-		Preset->ExposeActor(Actor, Args);
+	URemoteControlPreset* const Preset = GetPreset();
+	if (!Preset)
+	{
+		return;
 	}
+
+	FScopedTransaction Transaction(LOCTEXT("ExposeActor", "Expose Actor"));
+	Preset->Modify();
+
+	FRemoteControlPresetExposeArgs Args;
+	Args.GroupId = GetSelectedGroup();
+
+	Preset->ExposeActor(Actor, Args);
 }
 
 TSharedRef<SWidget> SRemoteControlPanel::CreateEntityDetailsView()
@@ -1572,6 +1634,12 @@ TSharedRef<SWidget> SRemoteControlPanel::CreateEntityDetailsView()
 
 void SRemoteControlPanel::UpdateEntityDetailsView(const TSharedPtr<SRCPanelTreeNode>& SelectedNode)
 {
+	URemoteControlPreset* const Preset = GetPreset();
+	if (!Preset)
+	{
+		return;
+	}
+
 	TSharedPtr<FStructOnScope> SelectedEntityPtr;
 
 	LastSelectedEntity = SelectedNode;
@@ -1609,14 +1677,14 @@ void SRemoteControlPanel::UpdateEntityDetailsView(const TSharedPtr<SRCPanelTreeN
 	}
 
 	static const FName ProtocolWidgetsModuleName = "RemoteControlProtocolWidgets";
-	if(LastSelectedEntity.IsValid() && FModuleManager::Get().IsModuleLoaded(ProtocolWidgetsModuleName) && ensure(Preset.IsValid()))
+	if(LastSelectedEntity.IsValid() && FModuleManager::Get().IsModuleLoaded(ProtocolWidgetsModuleName))
 	{
 		if (const TSharedPtr<FRemoteControlEntity> RCEntity = Preset->GetExposedEntity(LastSelectedEntity->GetRCId()).Pin())
 		{
 			if(RCEntity->IsBound())
 			{
 				IRemoteControlProtocolWidgetsModule& ProtocolWidgetsModule = FModuleManager::LoadModuleChecked<IRemoteControlProtocolWidgetsModule>(ProtocolWidgetsModuleName);
-				EntityProtocolDetails->SetContent(ProtocolWidgetsModule.GenerateDetailsForEntity(Preset.Get(), RCEntity->GetId()));
+				EntityProtocolDetails->SetContent(ProtocolWidgetsModule.GenerateDetailsForEntity(Preset, RCEntity->GetId()));
 			}
 			else
 			{
@@ -1638,9 +1706,9 @@ void SRemoteControlPanel::UpdateEntityDetailsView(const TSharedPtr<SRCPanelTreeN
 
 void SRemoteControlPanel::UpdateRebindButtonVisibility()
 {
-	if (URemoteControlPreset* PresetPtr = Preset.Get())
+	if (URemoteControlPreset* Preset = GetPreset())
 	{
-		for (TWeakPtr<FRemoteControlEntity> WeakEntity : PresetPtr->GetExposedEntities<FRemoteControlEntity>())
+		for (TWeakPtr<FRemoteControlEntity> WeakEntity : Preset->GetExposedEntities<FRemoteControlEntity>())
 		{
 			if (TSharedPtr<FRemoteControlEntity> Entity = WeakEntity.Pin())
 			{
@@ -1658,9 +1726,9 @@ void SRemoteControlPanel::UpdateRebindButtonVisibility()
 
 FReply SRemoteControlPanel::OnClickRebindAllButton()
 {
-	if (URemoteControlPreset* PresetPtr = Preset.Get())
+	if (URemoteControlPreset* Preset = GetPreset())
 	{
-		PresetPtr->RebindUnboundEntities();
+		Preset->RebindUnboundEntities();
 
 		UpdateRebindButtonVisibility();
 	}
@@ -1806,12 +1874,13 @@ void SRemoteControlPanel::GenerateToolbar()
 		MiscWidgets = MiscWidgetsHBox;
 	}
 
-	if (!Preset->SelectedWorld.IsValid())
+	URemoteControlPreset* const Preset = GetPreset();
+	if (Preset && !Preset->SelectedWorld.IsValid())
 	{
 		Preset->SelectedWorld = Preset->GetWorld(true);
 	}
 
-	if (Preset->SelectedWorld.IsValid())
+	if (Preset && Preset->SelectedWorld.IsValid())
 	{
 		SelectedWorldName = Preset->SelectedWorld->StreamingLevelsPrefix + Preset->SelectedWorld->GetName();
 	}
@@ -1993,7 +2062,7 @@ void SRemoteControlPanel::GenerateAuxiliaryToolbar()
 
 FText SRemoteControlPanel::HandlePresetName() const
 {
-	if (Preset)
+	if (const URemoteControlPreset* const Preset = GetPreset())
 	{
 		return FText::FromString(Preset->GetName());
 	}
@@ -2003,19 +2072,20 @@ FText SRemoteControlPanel::HandlePresetName() const
 
 bool SRemoteControlPanel::CanSaveAsset() const
 {
-	return Preset.IsValid() && Preset->IsAsset();
+	const URemoteControlPreset* const Preset = GetPreset();
+	return Preset && Preset->IsAsset();
 }
 
 void SRemoteControlPanel::SaveAsset() const
 {
-	if (Preset.IsValid())
+	if (const URemoteControlPreset* const Preset = GetPreset())
 	{
 		TArray<UPackage*> PackagesToSave;
 
 		if (!Preset->IsAsset())
 		{
 			// Log an invalid object but don't try to save it
-			UE_LOG(LogRemoteControl, Log, TEXT("Invalid object to save: %s"), (Preset.IsValid()) ? *Preset->GetFullName() : TEXT("Null Object"));
+			UE_LOG(LogRemoteControl, Log, TEXT("Invalid object to save: %s. Preset is not an asset"), *Preset->GetFullName());
 		}
 		else
 		{
@@ -2031,14 +2101,14 @@ void SRemoteControlPanel::SaveAsset() const
 
 bool SRemoteControlPanel::CanFindInContentBrowser() const
 {
-	return Preset.IsValid();
+	return !!GetPreset();
 }
 
 void SRemoteControlPanel::FindInContentBrowser() const
 {
-	if (URemoteControlPreset* ResolvedPreset = Preset.Get())
+	if (const URemoteControlPreset* Preset = GetPreset())
 	{
-		GEditor->SyncBrowserToObject(ResolvedPreset);
+		GEditor->SyncBrowserToObject(Preset);
 	}
 }
 
@@ -2106,6 +2176,12 @@ void SRemoteControlPanel::DeleteEntity_Execute()
 		return;
 	}
 
+	URemoteControlPreset* const Preset = GetPreset();
+	if (!Preset)
+	{
+		return;
+	}
+
 	if (LastSelectedEntity->GetRCType() == SRCPanelTreeNode::Group)
 	{
 		FScopedTransaction Transaction(LOCTEXT("DeleteGroup", "Delete Group"));
@@ -2160,7 +2236,8 @@ bool SRemoteControlPanel::CanDeleteEntity() const
 		return !ActiveLogicPanel->GetSelectedLogicItems().IsEmpty(); // User has focus on a logic panel
 	}
 
-	if (LastSelectedEntity.IsValid() && Preset.IsValid())
+	const URemoteControlPreset* const Preset = GetPreset();
+	if (LastSelectedEntity.IsValid() && Preset)
 	{
 		// Do not allow default group to be deleted.
 		return !Preset->Layout.IsDefaultGroup(LastSelectedEntity->GetRCId());
@@ -2209,7 +2286,8 @@ bool SRemoteControlPanel::CanRenameEntity() const
 		return true;
 	}
 
-	if (LastSelectedEntity.IsValid() && Preset.IsValid())
+	const URemoteControlPreset* const Preset = GetPreset();
+	if (LastSelectedEntity.IsValid() && Preset)
 	{
 		// Do not allow default group to be renamed.
 		return !Preset->Layout.IsDefaultGroup(LastSelectedEntity->GetRCId());
@@ -2381,8 +2459,10 @@ void SRemoteControlPanel::LoadSettings(const FGuid& InInstanceId) const
 
 void SRemoteControlPanel::SaveSettings()
 {
+	URemoteControlPreset* const Preset = GetPreset();
+
 	const TSharedPtr<SRCPanelFilter> FilterPtr = EntityList->GetFilterPtr();
-	if (Preset.IsValid() && FilterPtr.IsValid())
+	if (Preset && FilterPtr.IsValid())
 	{
 		const FString SettingsString = Preset->GetPresetId().ToString();
 
@@ -2421,12 +2501,13 @@ TSharedRef<SWidget> SRemoteControlPanel::OnGetSelectedWorldButtonContent()
 
 void SRemoteControlPanel::UpdatePanelForWorld(const UWorld* InWorld)
 {
-	if (!IsValid(InWorld) || !Preset.IsValid() || InWorld == Preset->SelectedWorld.Get())
+	URemoteControlPreset* const Preset = GetPreset();
+	if (!IsValid(InWorld) || (Preset && InWorld == Preset->SelectedWorld.Get()))
 	{
 		return;
 	}
 
-	if (Preset->IsEmbeddedPreset())
+	if (!Preset || Preset->IsEmbeddedPreset())
 	{
 		OpenPanelForEmbeddedPreset(InWorld);
 	}
@@ -2464,59 +2545,32 @@ void SRemoteControlPanel::UpdatePanelForWorld(const UWorld* InWorld)
 	}
 }
 
-void SRemoteControlPanel::OpenEmbeddedPreset(const FSoftObjectPath& InPresetToOpenPath)
+void SRemoteControlPanel::OpenPanelForEmbeddedPreset(const UWorld* InWorld)
 {
-	URemoteControlPreset* EditorRCPreset = nullptr;
-	if (UObject* EditorPreset = InPresetToOpenPath.ResolveObject())
+	UAssetEditorSubsystem* const AssetSubsystem = GEditor ? GEditor->GetEditorSubsystem<UAssetEditorSubsystem>() : nullptr;
+	if (!AssetSubsystem)
 	{
-		EditorRCPreset = Cast<URemoteControlPreset>(EditorPreset);
-	}
-	else
-	{
-		EditorRCPreset = Cast<URemoteControlPreset>(InPresetToOpenPath.TryLoad());
+		return;
 	}
 
-	if (IsValid(EditorRCPreset))
-	{
-		UAssetEditorSubsystem* AssetSubsystem = GEditor ? GEditor->GetEditorSubsystem<UAssetEditorSubsystem>() : nullptr;
-		if (AssetSubsystem)
-		{
-			AssetSubsystem->OpenEditorForAsset(EditorRCPreset);
-		}
-	}
-}
+	IRemoteControlModule& RCModule = IRemoteControlModule::Get();
 
-void SRemoteControlPanel::OpenEditorEmbeddedPreset()
-{
-	if (GEditor && GEditor->EditorWorld)
-	{
-		if (const UPackage* EditorWorldPackage = GEditor->EditorWorld.GetPackage())
-		{
-			const FSoftObjectPath RemotePreset = Preset.Get();
-			// Fix path for the editor world
-			const FSoftObjectPath EditorPresetPath = FSoftObjectPath(FTopLevelAssetPath(EditorWorldPackage->GetFName(), RemotePreset.GetAssetFName()), RemotePreset.GetSubPathString());
-			OpenEmbeddedPreset(EditorPresetPath);
-		}
-	}
-}
+	TArray<TWeakObjectPtr<URemoteControlPreset>> EmbeddedPresets;
+	RCModule.GetEmbeddedPresets(EmbeddedPresets);
 
-void SRemoteControlPanel::OpenPanelForEmbeddedPreset(const UWorld* InWorld)
-{
-	if (const UPackage* Package = InWorld->GetPackage())
+	URemoteControlPreset* SelectedEmbeddedPreset = nullptr;
+
+	for (TWeakObjectPtr<URemoteControlPreset> EmbeddedPresetWeak : EmbeddedPresets)
 	{
-		const int32 PIEInstanceId = Package->GetPIEInstanceID();
-		if (PIEInstanceId != INDEX_NONE)
+		URemoteControlPreset* const EmbeddedPreset = EmbeddedPresetWeak.Get();
+		if (EmbeddedPreset && EmbeddedPreset->GetEmbeddedWorld() == InWorld)
 		{
-			FSoftObjectPath CurrentPresetSoftPath = Preset.Get();
-			// Fix path for the current selected world
-			CurrentPresetSoftPath = FSoftObjectPath(FTopLevelAssetPath(Package->GetFName(), CurrentPresetSoftPath.GetAssetFName()), CurrentPresetSoftPath.GetSubPathString());
-			OpenEmbeddedPreset(CurrentPresetSoftPath);
-		}
-		else if (InWorld->WorldType == EWorldType::Editor)
-		{
-			OpenEditorEmbeddedPreset();
+			SelectedEmbeddedPreset = EmbeddedPreset;
+			break;
 		}
 	}
+
+	AssetSubsystem->OpenEditorForAsset(SelectedEmbeddedPreset);
 }
 
 void SRemoteControlPanel::CreateTargetWorldButtonDynamicEntries(UToolMenu* InMenu)
diff --git a/Engine/Plugins/VirtualProduction/RemoteControl/Source/RemoteControlUI/Private/UI/SRemoteControlPanel.h b/Engine/Plugins/VirtualProduction/RemoteControl/Source/RemoteControlUI/Private/UI/SRemoteControlPanel.h
--- a/Engine/Plugins/VirtualProduction/RemoteControl/Source/RemoteControlUI/Private/UI/SRemoteControlPanel.h
+++ b/Engine/Plugins/VirtualProduction/RemoteControl/Source/RemoteControlUI/Private/UI/SRemoteControlPanel.h
@@ -103,12 +103,18 @@ public:
 	/**
 	 * @return The preset represented by the panel.
 	 */
-	URemoteControlPreset* GetPreset() { return Preset.Get(); }
+	URemoteControlPreset* GetPreset()
+	{
+		return PresetWeak.Get();
+	}
 
 	/**
 	 * @return The preset represented by the panel.
 	 */
-	const URemoteControlPreset* GetPreset() const { return Preset.Get(); }
+	const URemoteControlPreset* GetPreset() const
+	{
+		return PresetWeak.Get();
+	}
 
 	/**
 	 * @param InArgs extension arguments
@@ -207,6 +213,9 @@ private:
 	/** Called when PIE begins */
 	void PostPIEStarted(const bool bInIsSimulating);
 
+	/** Called at the end of LoadMap to keep the remote control preset panel up-to-date with current PIE world */
+	void PostLoadMapWithWorld(UWorld* InWorld);
+
 	/** Called when PIE ends */
 	void OnEndPIE(const bool bInIsSimulating);
 
@@ -387,12 +396,6 @@ private:
 	/** Update the panel for the new World */
 	void UpdatePanelForWorld(const UWorld* InWorld);
 
-	/** try to Open the given preset path */
-	void OpenEmbeddedPreset(const FSoftObjectPath& InPresetToOpenPath);
-
-	/** try to open the editor preset of the current Preset */
-	void OpenEditorEmbeddedPreset();
-
 	/** Handle the TargetWorld change for embedded presets */
 	void OpenPanelForEmbeddedPreset(const UWorld* World);
 
@@ -403,8 +406,8 @@ private:
 	static const FName DefaultRemoteControlPanelToolBarName;
 	static const FName AuxiliaryRemoteControlPanelToolBarName;
 	static const FName TargetWorldRemoteControlPanelMenuName;
-	/** Holds the preset asset. */
-	TStrongObjectPtr<URemoteControlPreset> Preset;
+	/** Holds a weak reference to the preset asset. */
+	TWeakObjectPtr<URemoteControlPreset> PresetWeak;
 	/** Command list of this panel */
 	TSharedPtr<FUICommandList> CommandList;
 	/** Delegate called when the live mode changes. */
