Remote Control: Fix ChangePropId label to be more 'Change Property ID'

diff --git a/Engine/Plugins/VirtualProduction/RemoteControl/Source/RemoteControlUI/Private/Commands/RemoteControlCommands.cpp b/Engine/Plugins/VirtualProduction/RemoteControl/Source/RemoteControlUI/Private/Commands/RemoteControlCommands.cpp
--- a/Engine/Plugins/VirtualProduction/RemoteControl/Source/RemoteControlUI/Private/Commands/RemoteControlCommands.cpp
+++ b/Engine/Plugins/VirtualProduction/RemoteControl/Source/RemoteControlUI/Private/Commands/RemoteControlCommands.cpp
@@ -72,7 +72,7 @@ void FRemoteControlCommands::RegisterCommands()
 
 	// Modify Entity Prop Id
 	UI_COMMAND(ChangePropId,
-		"ChangePropid", "Change the selected property Ids",
+		"Change Property ID", "Change the selected property Ids",
 		EUserInterfaceActionType::Button,
 		FInputChord(EKeys::F3, EModifierKey::Control));
 
