Remote Control: Path behavior controller inputs are now able to be added if a controller with the specified name doesn't exist (where previously it would disable the button if a name was specified)

diff --git a/Engine/Plugins/VirtualProduction/RemoteControl/Source/RemoteControlUI/Private/UI/Customizations/RCAssetPathElementCustomization.cpp b/Engine/Plugins/VirtualProduction/RemoteControl/Source/RemoteControlUI/Private/UI/Customizations/RCAssetPathElementCustomization.cpp
--- a/Engine/Plugins/VirtualProduction/RemoteControl/Source/RemoteControlUI/Private/UI/Customizations/RCAssetPathElementCustomization.cpp
+++ b/Engine/Plugins/VirtualProduction/RemoteControl/Source/RemoteControlUI/Private/UI/Customizations/RCAssetPathElementCustomization.cpp
@@ -264,7 +264,26 @@ bool FRCAssetPathElementCustomization::IsCreateControllerButtonEnabled() const
 	FString Value;
 	PathHandle->GetValue(Value);
 
-	return Value.IsEmpty();
+	if (Value.IsEmpty())
+	{
+		return true;
+	}
+
+	TSharedPtr<FRCSetAssetByPathBehaviorModelNew> PathBehaviorModelNew = PathBehaviorModelNewWeak.Pin();
+
+	if (!PathBehaviorModelNew.IsValid())
+	{
+		return false;
+	}
+
+	URemoteControlPreset* Preset = PathBehaviorModelNew->GetPreset();
+
+	if (!Preset)
+	{
+		return false;
+	}
+
+	return !Preset->GetControllerByDisplayName(*Value);
 }
 
 FReply FRCAssetPathElementCustomization::OnCreateControllerButtonClicked()
