Remote Control: Added indicators for valid path elements in the path behavior.
- Added transaction to flipping between Internal and External paths.

diff --git a/Engine/Plugins/VirtualProduction/RemoteControl/Source/RemoteControlLogic/Private/Behaviour/Builtin/Path/RCSetAssetByPathBehaviorNew.cpp b/Engine/Plugins/VirtualProduction/RemoteControl/Source/RemoteControlLogic/Private/Behaviour/Builtin/Path/RCSetAssetByPathBehaviorNew.cpp
--- a/Engine/Plugins/VirtualProduction/RemoteControl/Source/RemoteControlLogic/Private/Behaviour/Builtin/Path/RCSetAssetByPathBehaviorNew.cpp
+++ b/Engine/Plugins/VirtualProduction/RemoteControl/Source/RemoteControlLogic/Private/Behaviour/Builtin/Path/RCSetAssetByPathBehaviorNew.cpp
@@ -4,6 +4,7 @@
 
 #include "Action/Path/RCSetAssetByPathActionNew.h"
 #include "Controller/RCController.h"
+#include "IRemoteControlModule.h"
 #include "RemoteControlField.h"
 #include "RemoteControlPreset.h"
 #include "RCVirtualProperty.h"
@@ -188,6 +189,102 @@ FString URCSetAssetByPathBehaviorNew::GetCurrentPath() const
 	return CurrentPath.Replace(TEXT("//"), TEXT("/"));
 }
 
+TArray<FRCSetAssetByPathBehaviorNewPathElement> URCSetAssetByPathBehaviorNew::GetPathElements() const
+{
+	TArray<FRCSetAssetByPathBehaviorNewPathElement> PathElements;
+
+	FString CurrentPath;
+	URCController* OwningController = ControllerWeakPtr.Get();
+	if (!OwningController)
+	{
+		UE_LOG(LogRemoteControl, Error, TEXT("Owning controller is invalid."));
+		return PathElements;
+	}
+
+	FString ControllerString;
+	OwningController->GetValueString(ControllerString);
+
+	using namespace UE::RemoteControl::Logic;
+
+	// Add Path String Concat
+	CurrentPath = bInternal ? SetAssetByPathNew::ContentFolder : FString("");
+
+	for (const FRCAssetPathElementNew& PathPart : PathStruct.AssetPath)
+	{
+		FRCSetAssetByPathBehaviorNewPathElement PathElement;
+
+		if (PathPart.bIsInput)
+		{
+			PathElement.bIsController = true;
+			PathElement.Controller = nullptr;
+
+			if (PathPart.Path.IsEmpty())
+			{
+				PathElement.Controller = OwningController;
+			}
+			else if (const URemoteControlPreset* RemoteControlPreset = OwningController->PresetWeakPtr.Get())
+			{
+				PathElement.Controller = RemoteControlPreset->GetControllerByDisplayName(*PathPart.Path);
+
+				if (!PathElement.Controller)
+				{
+					UE_LOG(LogRemoteControl, Error, TEXT("Unable to find controller. [{0}]"), *PathPart.Path);
+					PathElement.Path = TEXT("InvalidControllerName");
+					PathElement.bIsError = true;
+				}
+			}
+			else
+			{
+				UE_LOG(LogRemoteControl, Error, TEXT("Unable to find remote control preset."));
+				PathElement.Path = TEXT("InvalidPreset");
+				PathElement.bIsError = true;
+			}
+
+			// All control paths without a controller already have a path and error set.
+			if (!PathElement.Controller)
+			{
+				continue;
+			}
+
+			FString ControllerValue;
+
+			if (PathElement.Controller->GetValueString(ControllerValue))
+			{
+				PathElement.Path = ControllerValue;
+			}
+			else
+			{
+				if (PathPart.Path.IsEmpty())
+				{
+					UE_LOG(LogRemoteControl, Error, TEXT("Failed to get owning controller value as string."));
+				}
+				else
+				{
+					UE_LOG(LogRemoteControl, Error, TEXT("Failed to get controller value as string. [{0}]"), *PathPart.Path);
+				}
+
+				PathElement.Path = TEXT("UnableToGetControllerValue");
+				PathElement.bIsError = true;
+			}
+		}
+		else
+		{
+			PathElement.Path = PathPart.Path;
+			PathElement.Controller = nullptr;
+			PathElement.bIsController = false;
+		}
+
+		PathElement.Path.ReplaceCharInline(TCHAR('\\'), TCHAR('/'), ESearchCase::IgnoreCase);
+		PathElement.Path.ReplaceInline(TEXT("//"), TEXT("/"), ESearchCase::IgnoreCase);
+		
+		CurrentPath += PathElement.Path;
+
+		PathElements.Add(PathElement);
+	}
+
+	return PathElements;
+}
+
 void URCSetAssetByPathBehaviorNew::RefreshPathArray()
 {
 	if (PathStruct.AssetPath.Num() < 1)
diff --git a/Engine/Plugins/VirtualProduction/RemoteControl/Source/RemoteControlLogic/Public/Behaviour/Builtin/Path/RCSetAssetByPathBehaviorNew.h b/Engine/Plugins/VirtualProduction/RemoteControl/Source/RemoteControlLogic/Public/Behaviour/Builtin/Path/RCSetAssetByPathBehaviorNew.h
--- a/Engine/Plugins/VirtualProduction/RemoteControl/Source/RemoteControlLogic/Public/Behaviour/Builtin/Path/RCSetAssetByPathBehaviorNew.h
+++ b/Engine/Plugins/VirtualProduction/RemoteControl/Source/RemoteControlLogic/Public/Behaviour/Builtin/Path/RCSetAssetByPathBehaviorNew.h
@@ -12,6 +12,7 @@
 
 class URCSetAssetByPathActionNew;
 class URCUserDefinedStruct;
+class URCVirtualPropertyBase;
 class URCVirtualPropertyContainerBase;
 class UTexture;
 struct FRemoteControlProperty;
@@ -69,6 +70,14 @@ struct FRCSetAssetPathNew
 	TArray<FRCAssetPathElementNew> AssetPath;
 };
 
+struct FRCSetAssetByPathBehaviorNewPathElement
+{
+	FString Path;
+	bool bIsController = false;
+	URCVirtualPropertyBase* Controller = nullptr;
+	bool bIsError = false;
+};
+
 /**
  * Custom behaviour for Set Asset By Path
  */
@@ -93,6 +102,12 @@ public:
 	/** Returns the Path as the completed version, including Token Inputs. */
 	REMOTECONTROLLOGIC_API FString GetCurrentPath() const;
 
+	/** 
+	 * Returns an array of path elements, indicating whether the element was sourced from a controller and 
+	 * the controller itself, if it is still valid.
+	 */
+	REMOTECONTROLLOGIC_API TArray<FRCSetAssetByPathBehaviorNewPathElement> GetPathElements() const;
+
 	/** Called whenever a change has occured in the Slate */
 	REMOTECONTROLLOGIC_API void RefreshPathArray();
 
diff --git a/Engine/Plugins/VirtualProduction/RemoteControl/Source/RemoteControlUI/Private/UI/Behaviour/Builtin/Path/RCBehaviorSetAssetByPathModelNew.cpp b/Engine/Plugins/VirtualProduction/RemoteControl/Source/RemoteControlUI/Private/UI/Behaviour/Builtin/Path/RCBehaviorSetAssetByPathModelNew.cpp
--- a/Engine/Plugins/VirtualProduction/RemoteControl/Source/RemoteControlUI/Private/UI/Behaviour/Builtin/Path/RCBehaviorSetAssetByPathModelNew.cpp
+++ b/Engine/Plugins/VirtualProduction/RemoteControl/Source/RemoteControlUI/Private/UI/Behaviour/Builtin/Path/RCBehaviorSetAssetByPathModelNew.cpp
@@ -30,8 +30,11 @@
 
 namespace UE::RemoteControl::UI::Private
 {
-	constexpr const TCHAR* ErrorImage = TEXT("SourceControl.StatusIcon.Error");
-	constexpr const TCHAR* ValidImage = TEXT("SourceControl.StatusIcon.On");
+	namespace SetAssetByPathBehaviorModelNew
+	{
+		constexpr const TCHAR* ErrorImage = TEXT("SourceControl.StatusIcon.Error");
+		constexpr const TCHAR* ValidImage = TEXT("SourceControl.StatusIcon.On");
+	}
 }
 
 FRCSetAssetByPathBehaviorModelNew::FRCSetAssetByPathBehaviorModelNew(URCSetAssetByPathBehaviorNew* InSetAssetByPathBehavior, const TSharedPtr<SRemoteControlPanel> InRemoteControlPanel)
@@ -42,10 +45,7 @@ FRCSetAssetByPathBehaviorModelNew::FRCSetAssetByPathBehaviorModelNew(URCSetAsset
 
 void FRCSetAssetByPathBehaviorModelNew::Initialize()
 {
-	PreviewPathWidget = SNew(STextBlock)
-		.AutoWrapText(true)
-		.Font(DEFAULT_FONT("Regular", 10));
-
+	PreviewPathWidget = SNew(SHorizontalBox);
 	PathArrayWidget = SNew(SBox);
 
 	ElementsView = SNew(SListView<TSharedPtr<FRCPathBehaviorElementRow>>)
@@ -154,7 +154,6 @@ TSharedRef<SWidget> FRCSetAssetByPathBehaviorModelNew::GetPropertyWidget()
 				PreviewPathWidget->AsShared()
 			]
 		];
-
 	
 	return 	SNew(SHorizontalBox)
 		.Clipping(EWidgetClipping::OnDemand)
@@ -171,6 +170,7 @@ void FRCSetAssetByPathBehaviorModelNew::NotifyControllerValueChanged(TSharedPtr<
 {
 	// Update the preview when a controller change to always get the correct path
 	RefreshPreview();
+	RefreshValidity();
 }
 
 TSharedPtr<SRCLogicPanelListBase> FRCSetAssetByPathBehaviorModelNew::GetActionsListWidget(TSharedRef<SRCActionPanel> InActionPanel)
@@ -180,33 +180,118 @@ TSharedPtr<SRCLogicPanelListBase> FRCSetAssetByPathBehaviorModelNew::GetActionsL
 
 void FRCSetAssetByPathBehaviorModelNew::RefreshPreview() const
 {
-	FString PreviewText = LOCTEXT("InvalidPath", "Invalid Path!").ToString();
-	bool bInternal = true;
+	PreviewPathWidget->ClearChildren();
 
-	if (URCSetAssetByPathBehaviorNew* SetAssetByPathBehavior = SetAssetByPathBehaviorWeakPtr.Get())
+	URCSetAssetByPathBehaviorNew* SetAssetByPathBehavior = SetAssetByPathBehaviorWeakPtr.Get();
+
+	if (!SetAssetByPathBehavior)
 	{
-		PreviewText = SetAssetByPathBehavior->GetCurrentPath();
-		bInternal = SetAssetByPathBehavior->bInternal;
-	}
+		PreviewPathWidget->AddSlot()
+			.AutoWidth()
+			[
+				SNew(STextBlock)
+				.AutoWrapText(true)
+				.Font(DEFAULT_FONT("Regular", 10))
+				.Text(LOCTEXT("InvalidPath", "Invalid Path!"))
+			];
+
+		ValidationImage->SetImage(FAppStyle::Get().GetBrush(UE::RemoteControl::UI::Private::SetAssetByPathBehaviorModelNew::ErrorImage));
+		ValidationImage->SetToolTipText(LOCTEXT("InvalidImageTooltip", "Path does not point to a valid asset!"));
 
-	PreviewPathWidget->SetText(FText::FromString(PreviewText));
+		return;
+	}
 
 	bool bIsValidPath = false;
 
+	const TArray<FRCSetAssetByPathBehaviorNewPathElement> PathElements = SetAssetByPathBehavior->GetPathElements();
+	const bool bInternal = SetAssetByPathBehavior->bInternal;
+
+	FString Path;
+	
+	if (bInternal)
+	{
+		PreviewPathWidget->AddSlot()
+			.AutoWidth()
+			[
+				SNew(STextBlock)
+				.AutoWrapText(true)
+				.Font(DEFAULT_FONT("Regular", 10))
+				.Text(FText::FromString(UE::RemoteControl::Logic::SetAssetByPathNew::ContentFolder))
+			];
+	}
+
+	TArray<TTuple<FName, FString>> ControllerValues;
+
+	for (const FRCSetAssetByPathBehaviorNewPathElement& PathElement : PathElements)
+	{
+		if (PathElement.bIsController)
+		{
+			if (PathElement.Controller)
+			{
+				ControllerValues.Add({PathElement.Controller->DisplayName, PathElement.Path});
+			}
+			else
+			{
+				ControllerValues.Add({TEXT("{InvalidController}"), PathElement.Path});
+			}			
+		}
+
+		if (PathElement.bIsError)
+		{
+			if (PathElement.Path.IsEmpty())
+			{
+				PreviewPathWidget->AddSlot()
+					.AutoWidth()
+					[
+						SNew(STextBlock)
+						.AutoWrapText(true)
+						.Font(DEFAULT_FONT("Regular", 10))
+						.Text(LOCTEXT("NoControllerValue", "{Missing_Value}"))
+						.ColorAndOpacity(FStyleColors::AccentGray.GetSpecifiedColor())
+					];
+			}
+			else
+			{
+				PreviewPathWidget->AddSlot()
+					.AutoWidth()
+					[
+						SNew(STextBlock)
+						.AutoWrapText(true)
+						.Font(DEFAULT_FONT("Regular", 10))
+						.Text(FText::FromString(FString(TEXT("{")) + PathElement.Path + FString(TEXT("}"))))
+						.ColorAndOpacity(FStyleColors::AccentGray.GetSpecifiedColor())
+					];
+			}
+		}
+		else
+		{
+			Path += PathElement.Path;
+
+			PreviewPathWidget->AddSlot()
+				.AutoWidth()
+				[
+					SNew(STextBlock)
+					.AutoWrapText(true)
+					.Font(DEFAULT_FONT("Regular", 10))
+					.Text(FText::FromString(PathElement.Path))
+				];
+		}
+	}
+
 	if (!bInternal)
 	{
-		bIsValidPath = FPaths::FileExists(PreviewText);
+		bIsValidPath = FPaths::FileExists(Path);
 	}
 	else
 	{
 		// Specifying a path without the internal package name works, but cannot be found by the asset registry.
 		// Double up the last path element if it does not contain a dot to complete the path.
 		FAssetData AssetData;
-		FString CorrectedPath = PreviewText;
+		FString CorrectedPath = UE::RemoteControl::Logic::SetAssetByPathNew::ContentFolder + Path;
 
 		if (!CorrectedPath.EndsWith(TEXT("/")))
 		{
-			const FString Filename = FPaths::GetCleanFilename(PreviewText);
+			const FString Filename = FPaths::GetCleanFilename(CorrectedPath);
 			int32 DotIndex;
 
 			if (!Filename.FindChar(TEXT('.'), DotIndex))
@@ -217,18 +302,68 @@ void FRCSetAssetByPathBehaviorModelNew::RefreshPreview() const
 
 		bIsValidPath = IAssetRegistry::Get()->TryGetAssetByObjectPath(CorrectedPath, AssetData) == UE::AssetRegistry::EExists::Exists;
 	}
-	
+
 	ValidationImage->SetImage(FAppStyle::Get().GetBrush(
 		bIsValidPath
-			? UE::RemoteControl::UI::Private::ValidImage
-			: UE::RemoteControl::UI::Private::ErrorImage
+			? UE::RemoteControl::UI::Private::SetAssetByPathBehaviorModelNew::ValidImage
+			: UE::RemoteControl::UI::Private::SetAssetByPathBehaviorModelNew::ErrorImage
 	));
 
-	ValidationImage->SetToolTipText(
-		bIsValidPath
-			? LOCTEXT("ValidImageTooltip", "Path points to a valid asset.")
-			: LOCTEXT("InvalidImageTooltip", "Path does not point to a valid asset!")
-	);
+	FText Tooltip = bIsValidPath
+		? LOCTEXT("ValidImageTooltip", "Path points to a valid asset.")
+		: LOCTEXT("InvalidImageTooltip", "Path does not point to a valid asset!");
+
+	if (!ControllerValues.IsEmpty())
+	{
+		const FText ControllerValueFormat = LOCTEXT("ControllerValueFormat", "- {0}: {1}");
+		TArray<FText> ControllerValueTexts;
+
+		for (const TTuple<FName, FString>& ControllerValue : ControllerValues)
+		{
+			ControllerValueTexts.Add(
+				FText::Format(ControllerValueFormat, FText::FromName(ControllerValue.Key), FText::FromString(ControllerValue.Value))
+			);
+		}
+
+		Tooltip = FText::Format(
+			LOCTEXT("ControllerValuesTooltipFormat", "{0}\n\nController Values:\n{1}"),
+			Tooltip, 
+			FText::Join(INVTEXT("\n"), ControllerValueTexts)
+		);
+	}
+
+	ValidationImage->SetToolTipText(Tooltip);
+}
+
+void FRCSetAssetByPathBehaviorModelNew::RefreshValidity()
+{
+	bool bInternal = true;
+	TArray<FRCSetAssetByPathBehaviorNewPathElement> PathStack;
+
+	if (URCSetAssetByPathBehaviorNew* SetAssetByPathBehavior = SetAssetByPathBehaviorWeakPtr.Get())
+	{
+		bInternal = SetAssetByPathBehavior->bInternal;
+		PathStack = SetAssetByPathBehavior->GetPathElements();
+	}
+
+	for (const TSharedPtr<FRCPathBehaviorElementRow>& ElementItem : ElementItems)
+	{
+		ElementItem->Validity = ERCPathBehaviorElementValidty::Unchecked;
+	}
+
+	if (PathStack.Num() != ElementItems.Num())
+	{
+		return;
+	}
+
+	if (bInternal)
+	{
+		RefreshValidity_Internal(PathStack);
+	}
+	else
+	{
+		RefreshValidity_External(PathStack);
+	}
 }
 
 void FRCSetAssetByPathBehaviorModelNew::NotifyPostChange(const FPropertyChangedEvent& InPropertyChangedEvent, FEditPropertyChain* InPropertyThatChanged)
@@ -353,9 +488,11 @@ void FRCSetAssetByPathBehaviorModelNew::OnAssetPathFinishedChangingProperties(co
 					}
 				}
 			}
+		
 		}
 		// Just refresh the preview without reconstructing the widget if we just set values
 		RefreshPreview();
+		RefreshValidity();
 	}
 	else
 	{
@@ -462,6 +599,8 @@ void FRCSetAssetByPathBehaviorModelNew::RefreshPathAndPreview()
 
 	RefreshPreview();
 	PathArrayWidget->SetContent(Content);
+
+	RefreshValidity();
 }
 
 TSharedRef<ITableRow> FRCSetAssetByPathBehaviorModelNew::OnGenerateWidgetForList(TSharedPtr<FRCPathBehaviorElementRow> InItem, const TSharedRef<STableViewBase>& InOwnerTable)
@@ -469,4 +608,133 @@ TSharedRef<ITableRow> FRCSetAssetByPathBehaviorModelNew::OnGenerateWidgetForList
 	return SNew(SRCBehaviorSetAssetByPathNewElementRow, InOwnerTable, SharedThis(this), InItem);
 }
 
+void FRCSetAssetByPathBehaviorModelNew::RefreshValidity_Internal(TConstArrayView<FRCSetAssetByPathBehaviorNewPathElement> InPathStack)
+{
+	IAssetRegistry* AssetRegistry = IAssetRegistry::Get();
+	FAssetData AssetData;
+	FString Path;
+
+	TArray<FString> FullPaths;
+
+	for (const FRCSetAssetByPathBehaviorNewPathElement& PathElement : InPathStack)
+	{
+		if (FullPaths.IsEmpty())
+		{
+			FullPaths.Add(UE::RemoteControl::Logic::SetAssetByPathNew::ContentFolder + PathElement.Path);
+		}
+		else
+		{
+			FullPaths.Add(FullPaths.Last() + PathElement.Path);
+		}
+	}
+
+	const int32 LastElement = InPathStack.Num() - 1;
+
+	for (int32 Index = LastElement; Index >= 0; --Index)
+	{
+		Path = FullPaths[Index];
+
+		bool bValid = false;
+
+		// Only the last index can point to an asset.
+		if (Index == LastElement && !Path.EndsWith(TEXT("/")))
+		{
+			// Specifying a path without the internal package name works, but cannot be found by the asset registry.
+			// Double up the last path element if it does not contain a dot to complete the path.
+			const FString Filename = FPaths::GetCleanFilename(Path);
+			int32 DotIndex;
+
+			if (!Filename.FindChar(TEXT('.'), DotIndex))
+			{
+				Path += TEXT(".") + Filename;
+			}
+
+			bValid = AssetRegistry->TryGetAssetByObjectPath(Path, AssetData) == UE::AssetRegistry::EExists::Exists;
+		}
+
+		if (!bValid)
+		{
+			bValid = AssetRegistry->PathExists(Path);
+		}
+
+		// Invalid path, mark as invalid and check the previous one.
+		if (!bValid)
+		{
+			ElementItems[Index]->Validity = ERCPathBehaviorElementValidty::Invalid;
+			continue;
+		}
+
+		// Found a valid path, so all paths here and previous must be valid.
+		for (int32 UpdateIndex = 0; UpdateIndex <= Index; ++UpdateIndex)
+		{
+			if (!InPathStack[UpdateIndex].Path.IsEmpty())
+			{
+				ElementItems[UpdateIndex]->Validity = ERCPathBehaviorElementValidty::Valid;
+			}
+			else
+			{
+				ElementItems[UpdateIndex]->Validity = ERCPathBehaviorElementValidty::Invalid;
+			}
+		}
+
+		break;
+	}
+}
+
+void FRCSetAssetByPathBehaviorModelNew::RefreshValidity_External(TConstArrayView<FRCSetAssetByPathBehaviorNewPathElement> InPathStack)
+{
+	FString Path;
+
+	TArray<FString> FullPaths;
+
+	for (const FRCSetAssetByPathBehaviorNewPathElement& PathElement : InPathStack)
+	{
+		if (FullPaths.IsEmpty())
+		{
+			FullPaths.Add(PathElement.Path);
+		}
+		else
+		{
+			FullPaths.Add(FullPaths.Last() + PathElement.Path);
+		}
+	}
+
+	const int32 LastElement = InPathStack.Num() - 1;
+
+	for (int32 Index = LastElement; Index >= 0; --Index)
+	{
+		Path = FullPaths[Index];
+
+		// Only the last index can point to a file.
+		bool bValid = (Index == LastElement) && FPaths::FileExists(Path);
+
+		if (!bValid)
+		{
+			bValid = FPaths::DirectoryExists(Path);
+		}
+
+		// Invalid path, mark as invalid and check the previous one.
+		if (!bValid)
+		{
+			ElementItems[Index]->Validity = ERCPathBehaviorElementValidty::Invalid;
+			continue;
+		}
+
+		// Found a valid path, so all paths here and previous must be valid.
+		for (int32 UpdateIndex = 0; UpdateIndex <= Index; ++UpdateIndex)
+		{
+			if (!InPathStack[UpdateIndex].Path.IsEmpty())
+			{
+				ElementItems[UpdateIndex]->Validity = ERCPathBehaviorElementValidty::Valid;
+			}
+			else
+			{
+				ElementItems[UpdateIndex]->Validity = ERCPathBehaviorElementValidty::Invalid;
+			}
+		}
+
+		break;
+	}
+}
+
 #undef LOCTEXT_NAMESPACE
diff --git a/Engine/Plugins/VirtualProduction/RemoteControl/Source/RemoteControlUI/Private/UI/Behaviour/Builtin/Path/RCBehaviorSetAssetByPathModelNew.h b/Engine/Plugins/VirtualProduction/RemoteControl/Source/RemoteControlUI/Private/UI/Behaviour/Builtin/Path/RCBehaviorSetAssetByPathModelNew.h
--- a/Engine/Plugins/VirtualProduction/RemoteControl/Source/RemoteControlUI/Private/UI/Behaviour/Builtin/Path/RCBehaviorSetAssetByPathModelNew.h
+++ b/Engine/Plugins/VirtualProduction/RemoteControl/Source/RemoteControlUI/Private/UI/Behaviour/Builtin/Path/RCBehaviorSetAssetByPathModelNew.h
@@ -6,14 +6,17 @@
 #include "RemoteControlEntity.h"
 #include "UI/Behaviour/RCBehaviourModel.h"
 
+class FString;
 class IDetailTreeNode;
 class ITableRow;
 class SBox;
+class SHorizontalBox;
 class SImage;
 class STableViewBase;
 class URCSetAssetByPathBehaviorNew;
 enum class EItemDropZone;
 struct FRCPathBehaviorElementRow;
+struct FRCSetAssetByPathBehaviorNewPathElement;
 template<typename T> class SListView;
 
 /*
@@ -52,6 +55,9 @@ public:
 	/** Refresh the PreviewPath */
 	void RefreshPreview() const;
 
+	/** Refresh the validity state of the path elements. */
+	void RefreshValidity();
+
 	virtual void NotifyPostChange(const FPropertyChangedEvent& InPropertyChangedEvent, class FEditPropertyChain* InPropertyThatChanged) override;
 
 	void ReorderElementItems(int32 InDroppedOnIndex, EItemDropZone InDropZone, int32 InDroppedIndex);
@@ -69,6 +75,10 @@ private:
 	/** Creates a row for the elements view. */
 	TSharedRef<ITableRow> OnGenerateWidgetForList(TSharedPtr<FRCPathBehaviorElementRow> InItem, const TSharedRef<STableViewBase>& InOwnerTable);
 
+	void RefreshValidity_Internal(TConstArrayView<FRCSetAssetByPathBehaviorNewPathElement> InPathStack);
+	
+	void RefreshValidity_External(TConstArrayView<FRCSetAssetByPathBehaviorNewPathElement> InPathStack);
+
 private:
 	/** The SetAssetByPath Behaviour associated with this Model */
 	TWeakObjectPtr<URCSetAssetByPathBehaviorNew> SetAssetByPathBehaviorWeakPtr;
@@ -83,7 +93,7 @@ private:
 	TArray<TSharedPtr<FRCPathBehaviorElementRow>> ElementItems;
 
 	/** Pointer to the Preview Text Widget for the Path */
-	TSharedPtr<STextBlock> PreviewPathWidget;
+	TSharedPtr<SHorizontalBox> PreviewPathWidget;
 
 	/** Image showing whether the path is valid or not. */
 	TSharedPtr<SImage> ValidationImage;
diff --git a/Engine/Plugins/VirtualProduction/RemoteControl/Source/RemoteControlUI/Private/UI/Behaviour/Builtin/Path/SRCBehaviorSetAssetByPathNew.cpp b/Engine/Plugins/VirtualProduction/RemoteControl/Source/RemoteControlUI/Private/UI/Behaviour/Builtin/Path/SRCBehaviorSetAssetByPathNew.cpp
--- a/Engine/Plugins/VirtualProduction/RemoteControl/Source/RemoteControlUI/Private/UI/Behaviour/Builtin/Path/SRCBehaviorSetAssetByPathNew.cpp
+++ b/Engine/Plugins/VirtualProduction/RemoteControl/Source/RemoteControlUI/Private/UI/Behaviour/Builtin/Path/SRCBehaviorSetAssetByPathNew.cpp
@@ -5,6 +5,7 @@
 #include "Behaviour/Builtin/Path/RCSetAssetByPathBehaviorNew.h"
 #include "IDetailTreeNode.h"
 #include "RemoteControlPreset.h"
+#include "ScopedTransaction.h"
 #include "SlateOptMacros.h"
 #include "SPositiveActionButton.h"
 #include "Styling/StyleColors.h"
@@ -115,11 +116,14 @@ void SRCBehaviorSetAssetByPathNew::OnInternalExternalSwitchStateChanged(ECheckBo
 		return;
 	}
 
+	FScopedTransaction Transaction(LOCTEXT("PathBehaviorInternalToggle", "Toggled Path Behavior Internal/External."));
+	PathBehaviour->Modify();
 	PathBehaviour->bInternal = bInIsInternal;
 
-	if (const TSharedPtr<const FRCSetAssetByPathBehaviorModelNew>& SetAssetPath = SetAssetByPathWeakPtr.Pin())
+	if (const TSharedPtr<FRCSetAssetByPathBehaviorModelNew>& SetAssetPath = SetAssetByPathWeakPtr.Pin())
 	{
 		SetAssetPath->RefreshPreview();
+		SetAssetPath->RefreshValidity();
 	}
 }
 
diff --git a/Engine/Plugins/VirtualProduction/RemoteControl/Source/RemoteControlUI/Private/UI/Behaviour/Builtin/Path/SRCBehaviorSetAssetByPathNew.h b/Engine/Plugins/VirtualProduction/RemoteControl/Source/RemoteControlUI/Private/UI/Behaviour/Builtin/Path/SRCBehaviorSetAssetByPathNew.h
--- a/Engine/Plugins/VirtualProduction/RemoteControl/Source/RemoteControlUI/Private/UI/Behaviour/Builtin/Path/SRCBehaviorSetAssetByPathNew.h
+++ b/Engine/Plugins/VirtualProduction/RemoteControl/Source/RemoteControlUI/Private/UI/Behaviour/Builtin/Path/SRCBehaviorSetAssetByPathNew.h
@@ -31,7 +31,7 @@ private:
 
 private:
 	/** The Behaviour (UI model) associated with us */
-	TWeakPtr<const FRCSetAssetByPathBehaviorModelNew> SetAssetByPathWeakPtr;
+	TWeakPtr<FRCSetAssetByPathBehaviorModelNew> SetAssetByPathWeakPtr;
 
 	/** The PathBehaviour associated with us */
 	TWeakObjectPtr<URCSetAssetByPathBehaviorNew> PathBehaviour;
diff --git a/Engine/Plugins/VirtualProduction/RemoteControl/Source/RemoteControlUI/Private/UI/Behaviour/Builtin/Path/SRCBehaviorSetAssetByPathNewElementRow.cpp b/Engine/Plugins/VirtualProduction/RemoteControl/Source/RemoteControlUI/Private/UI/Behaviour/Builtin/Path/SRCBehaviorSetAssetByPathNewElementRow.cpp
--- a/Engine/Plugins/VirtualProduction/RemoteControl/Source/RemoteControlUI/Private/UI/Behaviour/Builtin/Path/SRCBehaviorSetAssetByPathNewElementRow.cpp
+++ b/Engine/Plugins/VirtualProduction/RemoteControl/Source/RemoteControlUI/Private/UI/Behaviour/Builtin/Path/SRCBehaviorSetAssetByPathNewElementRow.cpp
@@ -10,6 +10,15 @@
 
 #define LOCTEXT_NAMESPACE "SRCBehaviorPanelRow"
 
+namespace UE::RemoteControl::UI::Private
+{
+	namespace BehaviorSetAssetByPathNewElementRow
+	{
+		constexpr const TCHAR* ErrorImage = TEXT("SourceControl.StatusIcon.Error");
+		constexpr const TCHAR* ValidImage = TEXT("SourceControl.StatusIcon.On");
+	}
+}
+
 /*
 * ~ FRCBehaviorSetAssetByPathNewElementDragDrop ~
 * Facilitates drag-drop operation
@@ -66,6 +75,20 @@ void SRCBehaviorSetAssetByPathNewElementRow::Construct(const FArguments& InArgs,
 			[
 				InElementItem->Widget
 			]
+
+			+ SHorizontalBox::Slot()
+			.AutoWidth()	
+			.Padding(0.f, 0.f, 4.f, 0.f)		
+			[
+				SNew(SBox)
+				.VAlign(EVerticalAlignment::VAlign_Center)
+				[
+					SNew(SImage)
+					.Image(this, &SRCBehaviorSetAssetByPathNewElementRow::GetValidityImage)
+					.ToolTipText(this, &SRCBehaviorSetAssetByPathNewElementRow::GetValidityToolTip)
+					.DesiredSizeOverride(FVector2D(12))
+				]
+			]
 		],
 		InOwnerTableView
 	);
@@ -147,4 +170,36 @@ FReply SRCBehaviorSetAssetByPathNewElementRow::OnRowDragDetected(const FGeometry
 	return FReply::Handled().BeginDragDrop(PathBehaviorElementDragDrop);
 }
 
+const FSlateBrush* SRCBehaviorSetAssetByPathNewElementRow::GetValidityImage() const
+{
+	using namespace UE::RemoteControl::UI::Private::BehaviorSetAssetByPathNewElementRow;
+
+	if (ElementItem.IsValid() && ElementItem->Validity == ERCPathBehaviorElementValidty::Valid)
+	{
+		return FAppStyle::Get().GetBrush(ValidImage);
+	}
+
+	return FAppStyle::Get().GetBrush(ErrorImage);
+}
+
+FText SRCBehaviorSetAssetByPathNewElementRow::GetValidityToolTip() const
+{
+	if (!ElementItem.IsValid())
+	{
+		return LOCTEXT("Error", "Error");
+	}
+
+	switch (ElementItem->Validity)
+	{
+		case ERCPathBehaviorElementValidty::Valid:
+			return LOCTEXT("ValidPathElement", "Valid path element.");
+
+		case ERCPathBehaviorElementValidty::Invalid:
+			return LOCTEXT("InvalidPathElement", "Invalid path element!");
+
+		default:
+			return LOCTEXT("UnknownPathElement", "Unable to determine status.");
+	}
+}
+
 #undef LOCTEXT_NAMESPACE
diff --git a/Engine/Plugins/VirtualProduction/RemoteControl/Source/RemoteControlUI/Private/UI/Behaviour/Builtin/Path/SRCBehaviorSetAssetByPathNewElementRow.h b/Engine/Plugins/VirtualProduction/RemoteControl/Source/RemoteControlUI/Private/UI/Behaviour/Builtin/Path/SRCBehaviorSetAssetByPathNewElementRow.h
--- a/Engine/Plugins/VirtualProduction/RemoteControl/Source/RemoteControlUI/Private/UI/Behaviour/Builtin/Path/SRCBehaviorSetAssetByPathNewElementRow.h
+++ b/Engine/Plugins/VirtualProduction/RemoteControl/Source/RemoteControlUI/Private/UI/Behaviour/Builtin/Path/SRCBehaviorSetAssetByPathNewElementRow.h
@@ -7,16 +7,26 @@
 #include "Widgets/Views/SListView.h"
 #include "Widgets/Views/STableRow.h"
 
+class FRCSetAssetByPathBehaviorModelNew;
+struct FRCPathBehaviorElementRow;
+struct FSlateBrush;
+
+using RCBehaviorSetAssetByPathNewElementListItem = TSharedPtr<FRCPathBehaviorElementRow>;
+
+enum ERCPathBehaviorElementValidty
+{
+	Valid,
+	Invalid,
+	Unchecked
+};
+
 struct FRCPathBehaviorElementRow
 {
 	int32 Index;
 	TSharedRef<SWidget> Widget;
+	ERCPathBehaviorElementValidty Validity = ERCPathBehaviorElementValidty::Unchecked;
 };
 
-class FRCSetAssetByPathBehaviorModelNew;
-
-using RCBehaviorSetAssetByPathNewElementListItem = TSharedPtr<FRCPathBehaviorElementRow>;
-
 class SRCBehaviorSetAssetByPathNewElementRow : public STableRow<RCBehaviorSetAssetByPathNewElementListItem>
 {
 	using FSuperRowType = STableRow<RCBehaviorSetAssetByPathNewElementListItem>;
@@ -40,6 +50,10 @@ private:
 
 	FReply OnRowDragDetected(const FGeometry& InMyGeometry, const FPointerEvent& InMouseEvent);
 
+	const FSlateBrush* GetValidityImage() const;
+
+	FText GetValidityToolTip() const;
+
 	TWeakPtr<FRCSetAssetByPathBehaviorModelNew> SetAssetByPathBehaviorModelNewWeak;
 	RCBehaviorSetAssetByPathNewElementListItem ElementItem;
 };
\ No newline at end of file
diff --git a/Engine/Plugins/VirtualProduction/RemoteControl/Source/RemoteControlUI/Private/UI/Customizations/RCAssetPathElementCustomization.cpp b/Engine/Plugins/VirtualProduction/RemoteControl/Source/RemoteControlUI/Private/UI/Customizations/RCAssetPathElementCustomization.cpp
--- a/Engine/Plugins/VirtualProduction/RemoteControl/Source/RemoteControlUI/Private/UI/Customizations/RCAssetPathElementCustomization.cpp
+++ b/Engine/Plugins/VirtualProduction/RemoteControl/Source/RemoteControlUI/Private/UI/Customizations/RCAssetPathElementCustomization.cpp
@@ -163,7 +163,7 @@ void FRCAssetPathElementCustomization::CustomizeHeader(TSharedRef<IPropertyHandl
 			.IsFocusable(false)
 			[
 				SNew(STextBlock)
-				.Text(LOCTEXT("RCInputButtonAssetPath", "RCInput"))
+				.Text(LOCTEXT("RCInputButtonAssetPath", "Controller"))
 				.Font(IDetailLayoutBuilder::GetDetailFont())
 			]
 		]
