Remote Control: Small fix to path behavior for invalid controllers

diff --git a/Engine/Plugins/VirtualProduction/RemoteControl/Source/RemoteControlLogic/Private/Behaviour/Builtin/Path/RCSetAssetByPathBehaviorNew.cpp b/Engine/Plugins/VirtualProduction/RemoteControl/Source/RemoteControlLogic/Private/Behaviour/Builtin/Path/RCSetAssetByPathBehaviorNew.cpp
--- a/Engine/Plugins/VirtualProduction/RemoteControl/Source/RemoteControlLogic/Private/Behaviour/Builtin/Path/RCSetAssetByPathBehaviorNew.cpp
+++ b/Engine/Plugins/VirtualProduction/RemoteControl/Source/RemoteControlLogic/Private/Behaviour/Builtin/Path/RCSetAssetByPathBehaviorNew.cpp
@@ -228,7 +228,7 @@ TArray<FRCSetAssetByPathBehaviorNewPathElement> URCSetAssetByPathBehaviorNew::Ge
 
 				if (!PathElement.Controller)
 				{
-					UE_LOG(LogRemoteControl, Error, TEXT("Unable to find controller. [{0}]"), *PathPart.Path);
+					UE_LOG(LogRemoteControl, Error, TEXT("Unable to find controller. [%s]"), *PathPart.Path);
 					PathElement.Path = TEXT("InvalidControllerName");
 					PathElement.bIsError = true;
 				}
@@ -241,30 +241,28 @@ TArray<FRCSetAssetByPathBehaviorNewPathElement> URCSetAssetByPathBehaviorNew::Ge
 			}
 
 			// All control paths without a controller already have a path and error set.
-			if (!PathElement.Controller)
+			if (PathElement.Controller)
 			{
-				continue;
-			}
-
-			FString ControllerValue;
+				FString ControllerValue;
 
-			if (PathElement.Controller->GetValueString(ControllerValue))
-			{
-				PathElement.Path = ControllerValue;
-			}
-			else
-			{
-				if (PathPart.Path.IsEmpty())
+				if (PathElement.Controller->GetValueString(ControllerValue))
 				{
-					UE_LOG(LogRemoteControl, Error, TEXT("Failed to get owning controller value as string."));
+					PathElement.Path = ControllerValue;
 				}
 				else
 				{
-					UE_LOG(LogRemoteControl, Error, TEXT("Failed to get controller value as string. [{0}]"), *PathPart.Path);
+					if (PathPart.Path.IsEmpty())
+					{
+						UE_LOG(LogRemoteControl, Error, TEXT("Failed to get owning controller value as string."));
+					}
+					else
+					{
+						UE_LOG(LogRemoteControl, Error, TEXT("Failed to get controller value as string. [%s]"), *PathPart.Path);
+					}
+
+					PathElement.Path = TEXT("UnableToGetControllerValue");
+					PathElement.bIsError = true;
 				}
-
-				PathElement.Path = TEXT("UnableToGetControllerValue");
-				PathElement.bIsError = true;
 			}
 		}
 		else
