Remote Control: Redux for the new path behavior visual indication update.

diff --git a/Engine/Plugins/VirtualProduction/RemoteControl/Source/RemoteControlLogic/Private/Behaviour/Builtin/Path/RCSetAssetByPathBehaviorNew.cpp b/Engine/Plugins/VirtualProduction/RemoteControl/Source/RemoteControlLogic/Private/Behaviour/Builtin/Path/RCSetAssetByPathBehaviorNew.cpp
--- a/Engine/Plugins/VirtualProduction/RemoteControl/Source/RemoteControlLogic/Private/Behaviour/Builtin/Path/RCSetAssetByPathBehaviorNew.cpp
+++ b/Engine/Plugins/VirtualProduction/RemoteControl/Source/RemoteControlLogic/Private/Behaviour/Builtin/Path/RCSetAssetByPathBehaviorNew.cpp
@@ -215,8 +215,8 @@ TArray<FRCSetAssetByPathBehaviorNewPathElement> URCSetAssetByPathBehaviorNew::Ge
 
 		if (PathPart.bIsInput)
 		{
-			PathElement.bIsController = true;
 			PathElement.Controller = nullptr;
+			PathElement.ControllerName = *PathPart.Path;
 
 			if (PathPart.Path.IsEmpty())
 			{
@@ -269,7 +269,6 @@ TArray<FRCSetAssetByPathBehaviorNewPathElement> URCSetAssetByPathBehaviorNew::Ge
 		{
 			PathElement.Path = PathPart.Path;
 			PathElement.Controller = nullptr;
-			PathElement.bIsController = false;
 		}
 
 		PathElement.Path.ReplaceCharInline(TCHAR('\\'), TCHAR('/'), ESearchCase::IgnoreCase);
diff --git a/Engine/Plugins/VirtualProduction/RemoteControl/Source/RemoteControlLogic/Public/Behaviour/Builtin/Path/RCSetAssetByPathBehaviorNew.h b/Engine/Plugins/VirtualProduction/RemoteControl/Source/RemoteControlLogic/Public/Behaviour/Builtin/Path/RCSetAssetByPathBehaviorNew.h
--- a/Engine/Plugins/VirtualProduction/RemoteControl/Source/RemoteControlLogic/Public/Behaviour/Builtin/Path/RCSetAssetByPathBehaviorNew.h
+++ b/Engine/Plugins/VirtualProduction/RemoteControl/Source/RemoteControlLogic/Public/Behaviour/Builtin/Path/RCSetAssetByPathBehaviorNew.h
@@ -73,8 +73,8 @@ struct FRCSetAssetPathNew
 struct FRCSetAssetByPathBehaviorNewPathElement
 {
 	FString Path;
-	bool bIsController = false;
 	URCVirtualPropertyBase* Controller = nullptr;
+	FName ControllerName = NAME_None;
 	bool bIsError = false;
 };
 
diff --git a/Engine/Plugins/VirtualProduction/RemoteControl/Source/RemoteControlUI/Private/UI/Behaviour/Builtin/Path/RCBehaviorSetAssetByPathModelNew.cpp b/Engine/Plugins/VirtualProduction/RemoteControl/Source/RemoteControlUI/Private/UI/Behaviour/Builtin/Path/RCBehaviorSetAssetByPathModelNew.cpp
--- a/Engine/Plugins/VirtualProduction/RemoteControl/Source/RemoteControlUI/Private/UI/Behaviour/Builtin/Path/RCBehaviorSetAssetByPathModelNew.cpp
+++ b/Engine/Plugins/VirtualProduction/RemoteControl/Source/RemoteControlUI/Private/UI/Behaviour/Builtin/Path/RCBehaviorSetAssetByPathModelNew.cpp
@@ -224,7 +224,7 @@ void FRCSetAssetByPathBehaviorModelNew::RefreshPreview() const
 
 	for (const FRCSetAssetByPathBehaviorNewPathElement& PathElement : PathElements)
 	{
-		if (PathElement.bIsController)
+		if (!PathElement.ControllerName.IsNone())
 		{
 			if (PathElement.Controller)
 			{
@@ -232,7 +232,7 @@ void FRCSetAssetByPathBehaviorModelNew::RefreshPreview() const
 			}
 			else
 			{
-				ControllerValues.Add({TEXT("{InvalidController}"), PathElement.Path});
+				ControllerValues.Add({PathElement.ControllerName, PathElement.Path});
 			}			
 		}
 
