Remote Control: Protocol bindings now correctly mark the preset modified when changing override mask.

diff --git a/Engine/Plugins/VirtualProduction/RemoteControl/Source/RemoteControlProtocolWidgets/Private/Widgets/Masks/SRCProtocolBindingMask.cpp b/Engine/Plugins/VirtualProduction/RemoteControl/Source/RemoteControlProtocolWidgets/Private/Widgets/Masks/SRCProtocolBindingMask.cpp
--- a/Engine/Plugins/VirtualProduction/RemoteControl/Source/RemoteControlProtocolWidgets/Private/Widgets/Masks/SRCProtocolBindingMask.cpp
+++ b/Engine/Plugins/VirtualProduction/RemoteControl/Source/RemoteControlProtocolWidgets/Private/Widgets/Masks/SRCProtocolBindingMask.cpp
@@ -2,6 +2,7 @@
 
 #include "SRCProtocolBindingMask.h"
 #include "IRemoteControlModule.h"
+#include "RemoteControlPreset.h"
 #include "RemoteControlProtocolBinding.h"
 #include "ViewModels/ProtocolBindingViewModel.h"
 
@@ -94,6 +95,11 @@ void SRCProtocolBindingMask::SetMaskEnabled(ECheckBoxState NewState, ERCMask New
 {
 	if (ViewModel.IsValid())
 	{
+		if (URemoteControlPreset* Preset = ViewModel->GetPreset().Get())
+		{
+			Preset->Modify();
+		}
+
 		if (const FRemoteControlProtocolBinding* RCProtocolBinding = ViewModel->GetBinding())
 		{
 			if (TSharedPtr<TStructOnScope<FRemoteControlProtocolEntity>> RCProtocolEntity = RCProtocolBinding->GetRemoteControlProtocolEntityPtr())
