Remote Control: Small improvement to the icons and tooltips in the path behavior.

diff --git a/Engine/Plugins/VirtualProduction/RemoteControl/Source/RemoteControlUI/Private/UI/Behaviour/Builtin/Path/RCBehaviorSetAssetByPathModelNew.cpp b/Engine/Plugins/VirtualProduction/RemoteControl/Source/RemoteControlUI/Private/UI/Behaviour/Builtin/Path/RCBehaviorSetAssetByPathModelNew.cpp
--- a/Engine/Plugins/VirtualProduction/RemoteControl/Source/RemoteControlUI/Private/UI/Behaviour/Builtin/Path/RCBehaviorSetAssetByPathModelNew.cpp
+++ b/Engine/Plugins/VirtualProduction/RemoteControl/Source/RemoteControlUI/Private/UI/Behaviour/Builtin/Path/RCBehaviorSetAssetByPathModelNew.cpp
@@ -32,8 +32,8 @@ namespace UE::RemoteControl::UI::Private
 {
 	namespace SetAssetByPathBehaviorModelNew
 	{
-		constexpr const TCHAR* ErrorImage = TEXT("SourceControl.StatusIcon.Error");
 		constexpr const TCHAR* ValidImage = TEXT("SourceControl.StatusIcon.On");
+		constexpr const TCHAR* WarningImage = TEXT("SourceControl.StatusIcon.Error");
 	}
 }
 
@@ -195,7 +195,7 @@ void FRCSetAssetByPathBehaviorModelNew::RefreshPreview() const
 				.Text(LOCTEXT("InvalidPath", "Invalid Path!"))
 			];
 
-		ValidationImage->SetImage(FAppStyle::Get().GetBrush(UE::RemoteControl::UI::Private::SetAssetByPathBehaviorModelNew::ErrorImage));
+		ValidationImage->SetImage(FAppStyle::Get().GetBrush(UE::RemoteControl::UI::Private::SetAssetByPathBehaviorModelNew::WarningImage));
 		ValidationImage->SetToolTipText(LOCTEXT("InvalidImageTooltip", "Path does not point to a valid asset!"));
 
 		return;
@@ -306,7 +306,7 @@ void FRCSetAssetByPathBehaviorModelNew::RefreshPreview() const
 	ValidationImage->SetImage(FAppStyle::Get().GetBrush(
 		bIsValidPath
 			? UE::RemoteControl::UI::Private::SetAssetByPathBehaviorModelNew::ValidImage
-			: UE::RemoteControl::UI::Private::SetAssetByPathBehaviorModelNew::ErrorImage
+			: UE::RemoteControl::UI::Private::SetAssetByPathBehaviorModelNew::WarningImage
 	));
 
 	FText Tooltip = bIsValidPath
@@ -661,6 +661,12 @@ void FRCSetAssetByPathBehaviorModelNew::RefreshValidity_Internal(TConstArrayView
 		if (!bValid)
 		{
 			ElementItems[Index]->Validity = ERCPathBehaviorElementValidty::Invalid;
+
+			if (Index < LastElement)
+			{
+				ElementItems[Index + 1]->Validity = ERCPathBehaviorElementValidty::Unknown;
+			}
+
 			continue;
 		}
 
@@ -717,6 +723,12 @@ void FRCSetAssetByPathBehaviorModelNew::RefreshValidity_External(TConstArrayView
 		if (!bValid)
 		{
 			ElementItems[Index]->Validity = ERCPathBehaviorElementValidty::Invalid;
+
+			if (Index < LastElement)
+			{
+				ElementItems[Index + 1]->Validity = ERCPathBehaviorElementValidty::Unknown;
+			}
+
 			continue;
 		}
 
diff --git a/Engine/Plugins/VirtualProduction/RemoteControl/Source/RemoteControlUI/Private/UI/Behaviour/Builtin/Path/SRCBehaviorSetAssetByPathNewElementRow.cpp b/Engine/Plugins/VirtualProduction/RemoteControl/Source/RemoteControlUI/Private/UI/Behaviour/Builtin/Path/SRCBehaviorSetAssetByPathNewElementRow.cpp
--- a/Engine/Plugins/VirtualProduction/RemoteControl/Source/RemoteControlUI/Private/UI/Behaviour/Builtin/Path/SRCBehaviorSetAssetByPathNewElementRow.cpp
+++ b/Engine/Plugins/VirtualProduction/RemoteControl/Source/RemoteControlUI/Private/UI/Behaviour/Builtin/Path/SRCBehaviorSetAssetByPathNewElementRow.cpp
@@ -14,8 +14,10 @@ namespace UE::RemoteControl::UI::Private
 {
 	namespace BehaviorSetAssetByPathNewElementRow
 	{
-		constexpr const TCHAR* ErrorImage = TEXT("SourceControl.StatusIcon.Error");
 		constexpr const TCHAR* ValidImage = TEXT("SourceControl.StatusIcon.On");
+		constexpr const TCHAR* WarningImage = TEXT("SourceControl.StatusIcon.Error");
+		constexpr const TCHAR* ErrorImage = TEXT("Icons.Error.Solid");
+		constexpr const TCHAR* NoImage = TEXT("Sequencer.Empty");
 	}
 }
 
@@ -174,9 +176,22 @@ const FSlateBrush* SRCBehaviorSetAssetByPathNewElementRow::GetValidityImage() co
 {
 	using namespace UE::RemoteControl::UI::Private::BehaviorSetAssetByPathNewElementRow;
 
-	if (ElementItem.IsValid() && ElementItem->Validity == ERCPathBehaviorElementValidty::Valid)
+	if (ElementItem.IsValid())
 	{
-		return FAppStyle::Get().GetBrush(ValidImage);
+		switch (ElementItem->Validity)
+		{
+			case ERCPathBehaviorElementValidty::Valid:
+				return FAppStyle::Get().GetBrush(ValidImage);
+
+			case ERCPathBehaviorElementValidty::Invalid:
+				return FAppStyle::Get().GetBrush(WarningImage);
+
+			case ERCPathBehaviorElementValidty::Unknown:
+				return FAppStyle::Get().GetBrush(NoImage);
+
+			case ERCPathBehaviorElementValidty::Unchecked:
+				return FAppStyle::Get().GetBrush(WarningImage);
+		}
 	}
 
 	return FAppStyle::Get().GetBrush(ErrorImage);
@@ -192,10 +207,10 @@ FText SRCBehaviorSetAssetByPathNewElementRow::GetValidityToolTip() const
 	switch (ElementItem->Validity)
 	{
 		case ERCPathBehaviorElementValidty::Valid:
-			return LOCTEXT("ValidPathElement", "Valid path element.");
+			return LOCTEXT("ValidPathElement", "Points to a valid folder/asset.");
 
 		case ERCPathBehaviorElementValidty::Invalid:
-			return LOCTEXT("InvalidPathElement", "Invalid path element!");
+			return LOCTEXT("InvalidPathElement", "Does not point to a valid folder/asset!");
 
 		default:
 			return LOCTEXT("UnknownPathElement", "Unable to determine status.");
diff --git a/Engine/Plugins/VirtualProduction/RemoteControl/Source/RemoteControlUI/Private/UI/Behaviour/Builtin/Path/SRCBehaviorSetAssetByPathNewElementRow.h b/Engine/Plugins/VirtualProduction/RemoteControl/Source/RemoteControlUI/Private/UI/Behaviour/Builtin/Path/SRCBehaviorSetAssetByPathNewElementRow.h
--- a/Engine/Plugins/VirtualProduction/RemoteControl/Source/RemoteControlUI/Private/UI/Behaviour/Builtin/Path/SRCBehaviorSetAssetByPathNewElementRow.h
+++ b/Engine/Plugins/VirtualProduction/RemoteControl/Source/RemoteControlUI/Private/UI/Behaviour/Builtin/Path/SRCBehaviorSetAssetByPathNewElementRow.h
@@ -17,6 +17,7 @@ enum ERCPathBehaviorElementValidty
 {
 	Valid,
 	Invalid,
+	Unknown,
 	Unchecked
 };
 
