Remote Control: Fixed RC Protocol Filter localization

diff --git a/Engine/Plugins/VirtualProduction/RemoteControl/Source/RemoteControlUI/Private/UI/Filters/CustomRCFilterData.cpp b/Engine/Plugins/VirtualProduction/RemoteControl/Source/RemoteControlUI/Private/UI/Filters/CustomRCFilterData.cpp
--- a/Engine/Plugins/VirtualProduction/RemoteControl/Source/RemoteControlUI/Private/UI/Filters/CustomRCFilterData.cpp
+++ b/Engine/Plugins/VirtualProduction/RemoteControl/Source/RemoteControlUI/Private/UI/Filters/CustomRCFilterData.cpp
@@ -96,6 +96,11 @@ FLinearColor FCustomTypeFilterData::GetColor() const
 
 FText FCustomTypeFilterData::GetName() const
 {
+	if (!FilterDisplayName.IsEmpty())
+	{
+		return FilterDisplayName;
+	}
+
 	if (TSharedPtr<IAssetTypeActions> AssetTypeActionsPin = AssetTypeActions.Pin())
 	{
 		return AssetTypeActionsPin->GetName();
diff --git a/Engine/Plugins/VirtualProduction/RemoteControl/Source/RemoteControlUI/Private/UI/Filters/SRCFilterBar.h b/Engine/Plugins/VirtualProduction/RemoteControl/Source/RemoteControlUI/Private/UI/Filters/SRCFilterBar.h
--- a/Engine/Plugins/VirtualProduction/RemoteControl/Source/RemoteControlUI/Private/UI/Filters/SRCFilterBar.h
+++ b/Engine/Plugins/VirtualProduction/RemoteControl/Source/RemoteControlUI/Private/UI/Filters/SRCFilterBar.h
@@ -6,6 +6,7 @@
 
 #include "AssetToolsModule.h"
 #include "AssetTypeCategories.h"
+#include "EdGraphSchema_K2.h"
 
 #include "Filters/SBasicFilterBar.h"
 
@@ -297,32 +298,39 @@ protected:
 			, MakeShareable(new FFilterCategory(LOCTEXT("TypeFilter_Primary", "All"), LOCTEXT("PrimaryFilterTooltip", "Filter in entities of unique types.")))
 		);
 
+		auto GetTypeDisplayText =
+			[](FName InCategory, FName InSubCategory = NAME_None)->FText
+			{
+				constexpr bool bForMenu = true;
+				return UEdGraphSchema_K2::GetCategoryText(InCategory, InSubCategory, bForMenu);
+			};
+
 		// Core Type Filters.
-		PopulateFilterData(EEntityTypeCategories::Core, FBoolProperty::StaticClass());
-		PopulateFilterData(EEntityTypeCategories::Core, FByteProperty::StaticClass());
-		PopulateFilterData(EEntityTypeCategories::Core, FEnumProperty::StaticClass());
-		PopulateFilterData(EEntityTypeCategories::Core, FFloatProperty::StaticClass());
-		PopulateFilterData(EEntityTypeCategories::Core, FIntProperty::StaticClass());
-		PopulateFilterData(EEntityTypeCategories::Core, FInt64Property::StaticClass());
-		
+		PopulateFilterData(EEntityTypeCategories::Core, FBoolProperty::StaticClass() , GetTypeDisplayText(UEdGraphSchema_K2::PC_Boolean));
+		PopulateFilterData(EEntityTypeCategories::Core, FByteProperty::StaticClass() , GetTypeDisplayText(UEdGraphSchema_K2::PC_Byte));
+		PopulateFilterData(EEntityTypeCategories::Core, FEnumProperty::StaticClass() , GetTypeDisplayText(UEdGraphSchema_K2::PC_Enum));
+		PopulateFilterData(EEntityTypeCategories::Core, FFloatProperty::StaticClass(), GetTypeDisplayText(UEdGraphSchema_K2::PC_Real, UEdGraphSchema_K2::PC_Float));
+		PopulateFilterData(EEntityTypeCategories::Core, FIntProperty::StaticClass()  , GetTypeDisplayText(UEdGraphSchema_K2::PC_Int));
+		PopulateFilterData(EEntityTypeCategories::Core, FInt64Property::StaticClass(), GetTypeDisplayText(UEdGraphSchema_K2::PC_Int64));
+
 		// String Type Filters.
-		PopulateFilterData(EEntityTypeCategories::Strings, FNameProperty::StaticClass());
-		PopulateFilterData(EEntityTypeCategories::Strings, FStrProperty::StaticClass());
-		PopulateFilterData(EEntityTypeCategories::Strings, FTextProperty::StaticClass());
-		
+		PopulateFilterData(EEntityTypeCategories::Strings, FNameProperty::StaticClass(), GetTypeDisplayText(UEdGraphSchema_K2::PC_Name));
+		PopulateFilterData(EEntityTypeCategories::Strings, FStrProperty::StaticClass() , GetTypeDisplayText(UEdGraphSchema_K2::PC_String));
+		PopulateFilterData(EEntityTypeCategories::Strings, FTextProperty::StaticClass(), GetTypeDisplayText(UEdGraphSchema_K2::PC_Text));
+
 		// Primary Type Filters.
-		PopulateFilterData(EEntityTypeCategories::Primary, FStructProperty::StaticClass(), NAME_Vector);
-		PopulateFilterData(EEntityTypeCategories::Primary, FStructProperty::StaticClass(), NAME_Rotator);
-		PopulateFilterData(EEntityTypeCategories::Primary, FStructProperty::StaticClass(), NAME_Color);
-		PopulateFilterData(EEntityTypeCategories::Primary, FStructProperty::StaticClass(), NAME_Transform);
-		PopulateFilterData(EEntityTypeCategories::Primary, FProperty::StaticClass(), RemoteControlTypes::NAME_RCActors);
-		PopulateFilterData(EEntityTypeCategories::Primary, FProperty::StaticClass(), RemoteControlTypes::NAME_RCFunctions);
-		
+		PopulateFilterData(EEntityTypeCategories::Primary, FStructProperty::StaticClass(), TBaseStructure<FVector>::Get()->GetDisplayNameText(), NAME_Vector);
+		PopulateFilterData(EEntityTypeCategories::Primary, FStructProperty::StaticClass(), TBaseStructure<FRotator>::Get()->GetDisplayNameText(), NAME_Rotator);
+		PopulateFilterData(EEntityTypeCategories::Primary, FStructProperty::StaticClass(), TBaseStructure<FColor>::Get()->GetDisplayNameText(), NAME_Color);
+		PopulateFilterData(EEntityTypeCategories::Primary, FStructProperty::StaticClass(), TBaseStructure<FTransform>::Get()->GetDisplayNameText(), NAME_Transform);
+		PopulateFilterData(EEntityTypeCategories::Primary, FProperty::StaticClass(), LOCTEXT("RCActorsDisplayName", "Exposed Actors"), RemoteControlTypes::NAME_RCActors);
+		PopulateFilterData(EEntityTypeCategories::Primary, FProperty::StaticClass(), LOCTEXT("RCFunctionsDisplayName", "Exposed Functions"), RemoteControlTypes::NAME_RCFunctions);
+
 		// Container Type Filters.
-		PopulateFilterData(EEntityTypeCategories::Containers, FArrayProperty::StaticClass());
-		PopulateFilterData(EEntityTypeCategories::Containers, FMapProperty::StaticClass());
-		PopulateFilterData(EEntityTypeCategories::Containers, FSetProperty::StaticClass());
-		
+		PopulateFilterData(EEntityTypeCategories::Containers, FArrayProperty::StaticClass(), LOCTEXT("ArrayDisplayName", "Array"));
+		PopulateFilterData(EEntityTypeCategories::Containers, FMapProperty::StaticClass()  , LOCTEXT("MapDisplayName", "Map"));
+		PopulateFilterData(EEntityTypeCategories::Containers, FSetProperty::StaticClass()  , LOCTEXT("SetDisplayName", "Set"));
+
 		// TODO : Add support to Custom Classes.
 		FAssetToolsModule& AssetToolsModule = FModuleManager::LoadModuleChecked<FAssetToolsModule>(TEXT("AssetTools"));
 
@@ -744,9 +752,9 @@ protected:
 private:
 
 	/** Populates CustomTypeFilters from Entity Type and its Category. */
-	void PopulateFilterData(EEntityTypeCategories::Type InEntityTypeCategory, FFieldClass* InEntityType, const FName& InCustomTypeName = NAME_None)
+	void PopulateFilterData(EEntityTypeCategories::Type InEntityTypeCategory, FFieldClass* InEntityType, FText InFilterDisplayName, const FName& InCustomTypeName = NAME_None)
 	{
-		TSharedRef<FCustomTypeFilterData> CustomTypeFilterData = MakeShared<FCustomTypeFilterData>(InEntityTypeCategory, InEntityType, InCustomTypeName);
+		TSharedRef<FCustomTypeFilterData> CustomTypeFilterData = MakeShared<FCustomTypeFilterData>(InEntityTypeCategory, InEntityType, InFilterDisplayName, InCustomTypeName);
 		CustomTypeFilters.Add(CustomTypeFilterData);
 	}
 	
diff --git a/Engine/Plugins/VirtualProduction/RemoteControl/Source/RemoteControlUI/Public/UI/Filters/CustomRCFilterData.h b/Engine/Plugins/VirtualProduction/RemoteControl/Source/RemoteControlUI/Public/UI/Filters/CustomRCFilterData.h
--- a/Engine/Plugins/VirtualProduction/RemoteControl/Source/RemoteControlUI/Public/UI/Filters/CustomRCFilterData.h
+++ b/Engine/Plugins/VirtualProduction/RemoteControl/Source/RemoteControlUI/Public/UI/Filters/CustomRCFilterData.h
@@ -19,20 +19,22 @@ class FCustomTypeFilterData
 {
 public:
 	
-	FCustomTypeFilterData(EEntityTypeCategories::Type InEntityTypeCategory, FFieldClass* InEntityType, const FName& InCustomTypeName = NAME_None)
+	FCustomTypeFilterData(EEntityTypeCategories::Type InEntityTypeCategory, FFieldClass* InEntityType, FText InFilterDisplayName, const FName& InCustomTypeName = NAME_None)
 		: Class(nullptr)
 		, EntityType(InEntityType)
 		, EntityTypeCategory(InEntityTypeCategory)
+		, FilterDisplayName(InFilterDisplayName)
 		, Color(FLinearColor::White)
 		, CustomTypeName(InCustomTypeName)
 		, AssetTypeActions(nullptr)
 	{
 	};
 	
-	FCustomTypeFilterData(EEntityTypeCategories::Type InEntityTypeCategory, FFieldClass* InEntityType, TSharedPtr<FFilterCategory> InCategory, const FName& InCustomTypeName = NAME_None)
+	FCustomTypeFilterData(EEntityTypeCategories::Type InEntityTypeCategory, FFieldClass* InEntityType, FText InEntityTypeDisplayName, TSharedPtr<FFilterCategory> InCategory, const FName& InCustomTypeName = NAME_None)
 		: Class(nullptr)
 		, EntityType(InEntityType)
 		, EntityTypeCategory(InEntityTypeCategory)
+		, FilterDisplayName(InEntityTypeDisplayName)
 		, Color(FLinearColor::White)
 		, CustomTypeName(InCustomTypeName)
 		, AssetTypeActions(nullptr)
@@ -118,6 +120,9 @@ private:
 	/** The Entity Type the filter is associated with. */
 	EEntityTypeCategories::Type EntityTypeCategory;
 
+	/** The entity type display name */
+	FText FilterDisplayName;
+
 	/** The Categories the filter will show up under */
 	TArray<TSharedPtr<FFilterCategory>> Categories;
 
