Remote Control: Protocol OverrideMasks is now (actually) saved when saved.

diff --git a/Engine/Plugins/VirtualProduction/RemoteControl/Source/RemoteControl/Private/RemoteControlProtocolBinding.cpp b/Engine/Plugins/VirtualProduction/RemoteControl/Source/RemoteControl/Private/RemoteControlProtocolBinding.cpp
--- a/Engine/Plugins/VirtualProduction/RemoteControl/Source/RemoteControl/Private/RemoteControlProtocolBinding.cpp
+++ b/Engine/Plugins/VirtualProduction/RemoteControl/Source/RemoteControl/Private/RemoteControlProtocolBinding.cpp
@@ -572,13 +572,13 @@ bool FRemoteControlProtocolEntity::ApplyProtocolValueToProperty(const double InP
 		IRemoteControlModule::Get().ResolveObjectProperty(ObjectRef.Access, Object, ObjectRef.PropertyPathInfo, ObjectRef);
 
 		TSharedPtr<FRCMaskingOperation> MaskingOperation = MakeShared<FRCMaskingOperation>(ObjectRef.PropertyPathInfo, Object);
-		if (OverridenMasks == ERCMask::NoMask)
+		if (OverridenMasks == static_cast<uint8>(ERCMask::NoMask))
 		{
 			MaskingOperation->Masks = RemoteControlProperty->GetActiveMasks();
 		}
 		else
 		{
-			MaskingOperation->Masks = OverridenMasks;
+			MaskingOperation->Masks = static_cast<ERCMask>(OverridenMasks);
 		}
 
 		// Cache the values.
@@ -624,17 +624,17 @@ void FRemoteControlProtocolEntity::ResetDefaultBindingState()
 
 void FRemoteControlProtocolEntity::ClearMask(ERCMask InMaskBit)
 {
-	OverridenMasks &= ~InMaskBit;
+	OverridenMasks &= ~static_cast<uint8>(InMaskBit);
 }
 
 void FRemoteControlProtocolEntity::EnableMask(ERCMask InMaskBit)
 {
-	OverridenMasks |= InMaskBit;
+	OverridenMasks |= static_cast<uint8>(InMaskBit);
 }
 
 bool FRemoteControlProtocolEntity::HasMask(ERCMask InMaskBit) const
 {
-	return (OverridenMasks & InMaskBit) != ERCMask::NoMask;
+	return (OverridenMasks & static_cast<uint8>(InMaskBit)) != static_cast<uint8>(ERCMask::NoMask);
 }
 
 #if WITH_EDITOR
diff --git a/Engine/Plugins/VirtualProduction/RemoteControl/Source/RemoteControl/Public/RemoteControlProtocolBinding.h b/Engine/Plugins/VirtualProduction/RemoteControl/Source/RemoteControl/Public/RemoteControlProtocolBinding.h
--- a/Engine/Plugins/VirtualProduction/RemoteControl/Source/RemoteControl/Public/RemoteControlProtocolBinding.h
+++ b/Engine/Plugins/VirtualProduction/RemoteControl/Source/RemoteControl/Public/RemoteControlProtocolBinding.h
@@ -458,7 +458,7 @@ public:
 	virtual bool HasMask(ERCMask InMaskBit) const;
 
 	/** Returns the overriden masks. */
-	ERCMask GetOverridenMask() const { return OverridenMasks; }
+	ERCMask GetOverridenMask() const { return static_cast<ERCMask>(OverridenMasks); }
 
 #if WITH_EDITOR
 	/** Retrieves the name of property corresponding to the given column name. */
@@ -576,7 +576,7 @@ protected:
 
 	/** Holds the overriden masks. */
 	UPROPERTY()
-	ERCMask OverridenMasks = ERCMask::NoMask;
+	uint8 OverridenMasks = static_cast<uint8>(ERCMask::NoMask);
 
 #if WITH_EDITOR
 
