Remote Control: Updated the tooltips and icons for the path behavior again.

diff --git a/Engine/Plugins/VirtualProduction/RemoteControl/Source/RemoteControlLogic/Private/Behaviour/Builtin/Path/RCSetAssetByPathBehaviorNew.cpp b/Engine/Plugins/VirtualProduction/RemoteControl/Source/RemoteControlLogic/Private/Behaviour/Builtin/Path/RCSetAssetByPathBehaviorNew.cpp
--- a/Engine/Plugins/VirtualProduction/RemoteControl/Source/RemoteControlLogic/Private/Behaviour/Builtin/Path/RCSetAssetByPathBehaviorNew.cpp
+++ b/Engine/Plugins/VirtualProduction/RemoteControl/Source/RemoteControlLogic/Private/Behaviour/Builtin/Path/RCSetAssetByPathBehaviorNew.cpp
@@ -92,7 +92,7 @@ bool URCSetAssetByPathBehaviorNew::CanHaveActionForField(const TSharedPtr<FRemot
 		return false;
 	}
 
-	FProperty* Property = PropertyPath->GetPropertyInfo(0).Property.Get();
+	FProperty* Property = PropertyPath->GetPropertyInfo(PropertyPath->GetNumProperties() - 1).Property.Get();
 
 	if (!Property)
 	{
diff --git a/Engine/Plugins/VirtualProduction/RemoteControl/Source/RemoteControlUI/Private/UI/Behaviour/Builtin/Path/RCBehaviorSetAssetByPathModelNew.cpp b/Engine/Plugins/VirtualProduction/RemoteControl/Source/RemoteControlUI/Private/UI/Behaviour/Builtin/Path/RCBehaviorSetAssetByPathModelNew.cpp
--- a/Engine/Plugins/VirtualProduction/RemoteControl/Source/RemoteControlUI/Private/UI/Behaviour/Builtin/Path/RCBehaviorSetAssetByPathModelNew.cpp
+++ b/Engine/Plugins/VirtualProduction/RemoteControl/Source/RemoteControlUI/Private/UI/Behaviour/Builtin/Path/RCBehaviorSetAssetByPathModelNew.cpp
@@ -33,7 +33,7 @@ namespace UE::RemoteControl::UI::Private
 	namespace SetAssetByPathBehaviorModelNew
 	{
 		constexpr const TCHAR* ValidImage = TEXT("SourceControl.StatusIcon.On");
-		constexpr const TCHAR* WarningImage = TEXT("SourceControl.StatusIcon.Error");
+		constexpr const TCHAR* ErrorImage = TEXT("Icons.Error.Solid");
 	}
 }
 
@@ -195,7 +195,7 @@ void FRCSetAssetByPathBehaviorModelNew::RefreshPreview() const
 				.Text(LOCTEXT("InvalidPath", "Invalid Path!"))
 			];
 
-		ValidationImage->SetImage(FAppStyle::Get().GetBrush(UE::RemoteControl::UI::Private::SetAssetByPathBehaviorModelNew::WarningImage));
+		ValidationImage->SetImage(FAppStyle::Get().GetBrush(UE::RemoteControl::UI::Private::SetAssetByPathBehaviorModelNew::ErrorImage));
 		ValidationImage->SetToolTipText(LOCTEXT("InvalidImageTooltip", "Path does not point to a valid asset!"));
 
 		return;
@@ -306,7 +306,7 @@ void FRCSetAssetByPathBehaviorModelNew::RefreshPreview() const
 	ValidationImage->SetImage(FAppStyle::Get().GetBrush(
 		bIsValidPath
 			? UE::RemoteControl::UI::Private::SetAssetByPathBehaviorModelNew::ValidImage
-			: UE::RemoteControl::UI::Private::SetAssetByPathBehaviorModelNew::WarningImage
+			: UE::RemoteControl::UI::Private::SetAssetByPathBehaviorModelNew::ErrorImage
 	));
 
 	FText Tooltip = bIsValidPath
@@ -339,6 +339,7 @@ void FRCSetAssetByPathBehaviorModelNew::RefreshValidity()
 {
 	bool bInternal = true;
 	TArray<FRCSetAssetByPathBehaviorNewPathElement> PathStack;
+	URemoteControlPreset* Preset = nullptr;
 
 	if (URCSetAssetByPathBehaviorNew* SetAssetByPathBehavior = SetAssetByPathBehaviorWeakPtr.Get())
 	{
@@ -346,6 +347,11 @@ void FRCSetAssetByPathBehaviorModelNew::RefreshValidity()
 		PathStack = SetAssetByPathBehavior->GetPathElements();
 	}
 
+	if (const TSharedPtr<SRemoteControlPanel>& RCPanel = GetRemoteControlPanel())
+	{
+		Preset = RCPanel->GetPreset();
+	}
+
 	for (const TSharedPtr<FRCPathBehaviorElementRow>& ElementItem : ElementItems)
 	{
 		ElementItem->Validity = ERCPathBehaviorElementValidty::Unchecked;
@@ -364,6 +370,28 @@ void FRCSetAssetByPathBehaviorModelNew::RefreshValidity()
 	{
 		RefreshValidity_External(PathStack);
 	}
+
+	for (int32 Index = 0; Index < PathStack.Num(); ++Index)
+	{
+		if (!PathStack[Index].ControllerName.IsNone() && !PathStack[Index].Controller)
+		{
+			ElementItems[Index]->Validity = ERCPathBehaviorElementValidty::InvalidController;
+		}		
+		else if (PathStack[Index].Controller)
+		{
+			FString Value;
+			const bool bIsStringController = PathStack[Index].Controller->GetValueString(Value);
+			
+			if (!bIsStringController)
+			{
+				ElementItems[Index]->Validity = ERCPathBehaviorElementValidty::InvalidController;
+			}
+			else if (Value.IsEmpty())
+			{
+				ElementItems[Index]->Validity = ERCPathBehaviorElementValidty::EmptyControllerValue;
+			}
+		}
+	}
 }
 
 void FRCSetAssetByPathBehaviorModelNew::NotifyPostChange(const FPropertyChangedEvent& InPropertyChangedEvent, FEditPropertyChain* InPropertyThatChanged)
@@ -636,8 +664,7 @@ void FRCSetAssetByPathBehaviorModelNew::RefreshValidity_Internal(TConstArrayView
 
 		bool bValid = false;
 
-		// Only the last index can point to an asset.
-		if (Index == LastElement && !Path.EndsWith(TEXT("/")))
+		if (!Path.EndsWith(TEXT("/")))
 		{
 			// Specifying a path without the internal package name works, but cannot be found by the asset registry.
 			// Double up the last path element if it does not contain a dot to complete the path.
@@ -649,21 +676,32 @@ void FRCSetAssetByPathBehaviorModelNew::RefreshValidity_Internal(TConstArrayView
 				Path += TEXT(".") + Filename;
 			}
 
-			bValid = AssetRegistry->TryGetAssetByObjectPath(Path, AssetData) == UE::AssetRegistry::EExists::Exists;
+			if (AssetRegistry->TryGetAssetByObjectPath(Path, AssetData) == UE::AssetRegistry::EExists::Exists)
+			{
+				bValid = true;
+				ElementItems[Index]->Validity = ERCPathBehaviorElementValidty::ValidAsset;
+			}
 		}
 
 		if (!bValid)
 		{
-			bValid = AssetRegistry->PathExists(Path);
+			if (AssetRegistry->PathExists(Path))
+			{
+				bValid = true;
+				ElementItems[Index]->Validity = ERCPathBehaviorElementValidty::ValidPath;
+			}
 		}
 
 		// Invalid path, mark as invalid and check the previous one.
 		if (!bValid)
 		{
-			ElementItems[Index]->Validity = ERCPathBehaviorElementValidty::Invalid;
-
-			if (Index < LastElement)
+			if (Index == LastElement)
 			{
+				ElementItems[Index]->Validity = ERCPathBehaviorElementValidty::InvalidAsset;
+			}
+			else
+			{
+				ElementItems[Index]->Validity = ERCPathBehaviorElementValidty::InvalidPath;
 				ElementItems[Index + 1]->Validity = ERCPathBehaviorElementValidty::Unknown;
 			}
 
@@ -671,15 +709,15 @@ void FRCSetAssetByPathBehaviorModelNew::RefreshValidity_Internal(TConstArrayView
 		}
 
 		// Found a valid path, so all paths here and previous must be valid.
-		for (int32 UpdateIndex = 0; UpdateIndex <= Index; ++UpdateIndex)
+		for (int32 UpdateIndex = 0; UpdateIndex < Index; ++UpdateIndex)
 		{
-			if (!InPathStack[UpdateIndex].Path.IsEmpty())
+			if (UpdateIndex == LastElement)
 			{
-				ElementItems[UpdateIndex]->Validity = ERCPathBehaviorElementValidty::Valid;
+				ElementItems[UpdateIndex]->Validity = ERCPathBehaviorElementValidty::ValidAsset;
 			}
 			else
 			{
-				ElementItems[UpdateIndex]->Validity = ERCPathBehaviorElementValidty::Invalid;
+				ElementItems[UpdateIndex]->Validity = ERCPathBehaviorElementValidty::ValidPath;
 			}
 		}
 
@@ -722,10 +760,13 @@ void FRCSetAssetByPathBehaviorModelNew::RefreshValidity_External(TConstArrayView
 		// Invalid path, mark as invalid and check the previous one.
 		if (!bValid)
 		{
-			ElementItems[Index]->Validity = ERCPathBehaviorElementValidty::Invalid;
-
-			if (Index < LastElement)
+			if (Index == LastElement)
+			{
+				ElementItems[Index]->Validity = ERCPathBehaviorElementValidty::InvalidAsset;
+			}
+			else
 			{
+				ElementItems[Index]->Validity = ERCPathBehaviorElementValidty::InvalidPath;
 				ElementItems[Index + 1]->Validity = ERCPathBehaviorElementValidty::Unknown;
 			}
 
@@ -735,13 +776,13 @@ void FRCSetAssetByPathBehaviorModelNew::RefreshValidity_External(TConstArrayView
 		// Found a valid path, so all paths here and previous must be valid.
 		for (int32 UpdateIndex = 0; UpdateIndex <= Index; ++UpdateIndex)
 		{
-			if (!InPathStack[UpdateIndex].Path.IsEmpty())
+			if (UpdateIndex == LastElement)
 			{
-				ElementItems[UpdateIndex]->Validity = ERCPathBehaviorElementValidty::Valid;
+				ElementItems[UpdateIndex]->Validity = ERCPathBehaviorElementValidty::ValidAsset;
 			}
 			else
 			{
-				ElementItems[UpdateIndex]->Validity = ERCPathBehaviorElementValidty::Invalid;
+				ElementItems[UpdateIndex]->Validity = ERCPathBehaviorElementValidty::ValidPath;
 			}
 		}
 
diff --git a/Engine/Plugins/VirtualProduction/RemoteControl/Source/RemoteControlUI/Private/UI/Behaviour/Builtin/Path/SRCBehaviorSetAssetByPathNewElementRow.cpp b/Engine/Plugins/VirtualProduction/RemoteControl/Source/RemoteControlUI/Private/UI/Behaviour/Builtin/Path/SRCBehaviorSetAssetByPathNewElementRow.cpp
--- a/Engine/Plugins/VirtualProduction/RemoteControl/Source/RemoteControlUI/Private/UI/Behaviour/Builtin/Path/SRCBehaviorSetAssetByPathNewElementRow.cpp
+++ b/Engine/Plugins/VirtualProduction/RemoteControl/Source/RemoteControlUI/Private/UI/Behaviour/Builtin/Path/SRCBehaviorSetAssetByPathNewElementRow.cpp
@@ -180,16 +180,20 @@ const FSlateBrush* SRCBehaviorSetAssetByPathNewElementRow::GetValidityImage() co
 	{
 		switch (ElementItem->Validity)
 		{
-			case ERCPathBehaviorElementValidty::Valid:
+			case ERCPathBehaviorElementValidty::ValidPath:
+			case ERCPathBehaviorElementValidty::ValidAsset:
 				return FAppStyle::Get().GetBrush(ValidImage);
 
-			case ERCPathBehaviorElementValidty::Invalid:
-				return FAppStyle::Get().GetBrush(WarningImage);
+			case ERCPathBehaviorElementValidty::InvalidPath:
+			case ERCPathBehaviorElementValidty::InvalidAsset:
+			case ERCPathBehaviorElementValidty::InvalidController:
+				return FAppStyle::Get().GetBrush(ErrorImage);
 
 			case ERCPathBehaviorElementValidty::Unknown:
 				return FAppStyle::Get().GetBrush(NoImage);
 
 			case ERCPathBehaviorElementValidty::Unchecked:
+			case ERCPathBehaviorElementValidty::EmptyControllerValue:
 				return FAppStyle::Get().GetBrush(WarningImage);
 		}
 	}
@@ -206,14 +210,26 @@ FText SRCBehaviorSetAssetByPathNewElementRow::GetValidityToolTip() const
 
 	switch (ElementItem->Validity)
 	{
-		case ERCPathBehaviorElementValidty::Valid:
-			return LOCTEXT("ValidPathElement", "Points to a valid folder/asset.");
+		case ERCPathBehaviorElementValidty::ValidPath:
+			return LOCTEXT("ValidPathElement", "Points to a valid folder.");
+
+		case ERCPathBehaviorElementValidty::ValidAsset:
+			return LOCTEXT("ValidAssetElement", "Points to a valid asset.");
+
+		case ERCPathBehaviorElementValidty::InvalidPath:
+			return LOCTEXT("InvalidPathElement", "Does not point to a valid folder!");
+
+		case ERCPathBehaviorElementValidty::InvalidAsset:
+			return LOCTEXT("InvalidAssetElement", "Does not point to a valid asset!");
+
+		case ERCPathBehaviorElementValidty::InvalidController:
+			return LOCTEXT("InvalidControllerElement", "Invalid controller!");
 
-		case ERCPathBehaviorElementValidty::Invalid:
-			return LOCTEXT("InvalidPathElement", "Does not point to a valid folder/asset!");
+		case ERCPathBehaviorElementValidty::EmptyControllerValue:
+			return LOCTEXT("EmptyControllerValueElement", "Empty controller value!");
 
 		default:
-			return LOCTEXT("UnknownPathElement", "Unable to determine status.");
+			return LOCTEXT("UnknownPathElement", "Unknown.");
 	}
 }
 
diff --git a/Engine/Plugins/VirtualProduction/RemoteControl/Source/RemoteControlUI/Private/UI/Behaviour/Builtin/Path/SRCBehaviorSetAssetByPathNewElementRow.h b/Engine/Plugins/VirtualProduction/RemoteControl/Source/RemoteControlUI/Private/UI/Behaviour/Builtin/Path/SRCBehaviorSetAssetByPathNewElementRow.h
--- a/Engine/Plugins/VirtualProduction/RemoteControl/Source/RemoteControlUI/Private/UI/Behaviour/Builtin/Path/SRCBehaviorSetAssetByPathNewElementRow.h
+++ b/Engine/Plugins/VirtualProduction/RemoteControl/Source/RemoteControlUI/Private/UI/Behaviour/Builtin/Path/SRCBehaviorSetAssetByPathNewElementRow.h
@@ -15,8 +15,12 @@ using RCBehaviorSetAssetByPathNewElementListItem = TSharedPtr<FRCPathBehaviorEle
 
 enum ERCPathBehaviorElementValidty
 {
-	Valid,
-	Invalid,
+	ValidPath,
+	ValidAsset,
+	InvalidPath,
+	InvalidAsset,
+	InvalidController,
+	EmptyControllerValue,
 	Unknown,
 	Unchecked
 };
