Remote Control: fix the wrong preset getting opened when multiple presets are in a world. The previous condition checked for Preset->GetEmbeddedWorld == CurrentWorld, but would sometimes give the a different preset than the expected when multiple levels were present in a world (each level with their own preset). The fix changes the condition to Preset->GetLevelOuter == CurrentLevel.

diff --git a/Engine/Plugins/VirtualProduction/RemoteControl/Source/RemoteControlUI/Private/UI/SRemoteControlPanel.cpp b/Engine/Plugins/VirtualProduction/RemoteControl/Source/RemoteControlUI/Private/UI/SRemoteControlPanel.cpp
--- a/Engine/Plugins/VirtualProduction/RemoteControl/Source/RemoteControlUI/Private/UI/SRemoteControlPanel.cpp
+++ b/Engine/Plugins/VirtualProduction/RemoteControl/Source/RemoteControlUI/Private/UI/SRemoteControlPanel.cpp
@@ -2560,13 +2560,18 @@ void SRemoteControlPanel::OpenPanelForEmbeddedPreset(const UWorld* InWorld)
 
 	URemoteControlPreset* SelectedEmbeddedPreset = nullptr;
 
-	for (TWeakObjectPtr<URemoteControlPreset> EmbeddedPresetWeak : EmbeddedPresets)
+	const ULevel* const SelectedLevel = InWorld ? InWorld->GetCurrentLevel() : nullptr;
+
+	if (SelectedLevel)
 	{
-		URemoteControlPreset* const EmbeddedPreset = EmbeddedPresetWeak.Get();
-		if (EmbeddedPreset && EmbeddedPreset->GetEmbeddedWorld() == InWorld)
+		for (TWeakObjectPtr<URemoteControlPreset> EmbeddedPresetWeak : EmbeddedPresets)
 		{
-			SelectedEmbeddedPreset = EmbeddedPreset;
-			break;
+			URemoteControlPreset* const EmbeddedPreset = EmbeddedPresetWeak.Get();
+			if (EmbeddedPreset && EmbeddedPreset->GetTypedOuter<ULevel>() == SelectedLevel)
+			{
+				SelectedEmbeddedPreset = EmbeddedPreset;
+				break;
+			}
 		}
 	}
 
