Remote Control: added PythonScriptSettings to the classes not allowed when remote control EnableRemotePythonExecution is false.

diff --git a/Engine/Plugins/VirtualProduction/RemoteControl/Source/RemoteControl/Private/RemoteControlModule.cpp b/Engine/Plugins/VirtualProduction/RemoteControl/Source/RemoteControl/Private/RemoteControlModule.cpp
--- a/Engine/Plugins/VirtualProduction/RemoteControl/Source/RemoteControl/Private/RemoteControlModule.cpp
+++ b/Engine/Plugins/VirtualProduction/RemoteControl/Source/RemoteControl/Private/RemoteControlModule.cpp
@@ -2564,8 +2564,11 @@ bool FRemoteControlModule::CanBeAccessedRemotely(UObject* Object) const
 		
 	if (!GetDefault<URemoteControlSettings>()->bEnableRemotePythonExecution)
 	{
-		static const FName PythonScriptName = "PythonScriptLibrary"; 
-		if (Object->GetClass()->GetFName() == PythonScriptName)
+		static const FName PythonScriptName = "PythonScriptLibrary";
+		static const FName PythonScriptSettingsName = "PythonScriptPluginSettings";
+
+		const FName ObjectClassName = Object->GetClass()->GetFName();
+		if (ObjectClassName == PythonScriptName || ObjectClassName == PythonScriptSettingsName)
 		{
 			return false;
 		}
