Prevent RemoteControlUI from forcing a refresh on all detail panel implementations.

diff --git a/Engine/Plugins/VirtualProduction/RemoteControl/Source/RemoteControlUI/Private/RemoteControlUIModule.cpp b/Engine/Plugins/VirtualProduction/RemoteControl/Source/RemoteControlUI/Private/RemoteControlUIModule.cpp
--- a/Engine/Plugins/VirtualProduction/RemoteControl/Source/RemoteControlUI/Private/RemoteControlUIModule.cpp
+++ b/Engine/Plugins/VirtualProduction/RemoteControl/Source/RemoteControlUI/Private/RemoteControlUIModule.cpp
@@ -712,6 +712,7 @@ void FRemoteControlUIModule::HandleCreatePropertyRowExtension(const FOnGenerateG
 {
 	const FRCExposesPropertyArgs ExposesPropertyArgs(InArgs);
 
+	WeakDetailsTreeNode.Reset();
 	if (ExposesPropertyArgs.IsValid())
 	{
 		// Extend context row menu
@@ -740,9 +741,9 @@ void FRemoteControlUIModule::HandleCreatePropertyRowExtension(const FOnGenerateG
 			FGetActionCheckState(),
 			FIsActionButtonVisible::CreateRaw(this, &FRemoteControlUIModule::IsStaticOrSkeletalMaterialProperty, ExposesPropertyArgs)
 		);
-	}
 
-	WeakDetailsTreeNode = InArgs.OwnerTreeNode;
+		WeakDetailsTreeNode = InArgs.OwnerTreeNode;
+	}
 }
 
 TSharedPtr<SRemoteControlPanel> FRemoteControlUIModule::GetPanelForObject(const UObject* Object) const
