Fix duplicate symbol linker errors

#Synced-CL 50445837
#rnx

diff --git a/Engine/Plugins/ChaosClothAssetUsdDataflowNodes/Source/ChaosClothAssetUsdDataflowNodes/Private/ChaosClothAsset/UsdPrimAttributeAccessor.h b/Engine/Plugins/ChaosClothAssetUsdDataflowNodes/Source/ChaosClothAssetUsdDataflowNodes/Private/ChaosClothAsset/UsdPrimAttributeAccessor.h
--- a/Engine/Plugins/ChaosClothAssetUsdDataflowNodes/Source/ChaosClothAssetUsdDataflowNodes/Private/ChaosClothAsset/UsdPrimAttributeAccessor.h
+++ b/Engine/Plugins/ChaosClothAssetUsdDataflowNodes/Source/ChaosClothAssetUsdDataflowNodes/Private/ChaosClothAsset/UsdPrimAttributeAccessor.h
@@ -81,19 +81,19 @@ namespace UE::Chaos::ClothAsset
 	};
 
 	template<>
-	uint32 FUsdPrimAttributeAccessor::GetValueAs<uint32>(const FVtValue& VtValue, const uint32 DefaultValue) const
+	inline uint32 FUsdPrimAttributeAccessor::GetValueAs<uint32>(const FVtValue& VtValue, const uint32 DefaultValue) const
 	{
 		return VtValue.GetTypeName() == TEXT("unsigned int") ? UsdUtils::GetUnderlyingValue<uint32>(VtValue).Get(DefaultValue) : DefaultValue;  // Note: FUsdAttribute.GetTypeName() would return "uint" instead!
 	}
 
 	template<>
-	float FUsdPrimAttributeAccessor::GetValueAs<float>(const FVtValue& VtValue, const float DefaultValue) const
+	inline float FUsdPrimAttributeAccessor::GetValueAs<float>(const FVtValue& VtValue, const float DefaultValue) const
 	{
 		return VtValue.GetTypeName() == TEXT("float") ? UsdUtils::GetUnderlyingValue<float>(VtValue).Get(DefaultValue) : DefaultValue;
 	}
 
 	template<>
-	FVector3f FUsdPrimAttributeAccessor::GetValueAs<FVector3f>(const FVtValue& VtValue, const FVector3f DefaultValue) const
+	inline FVector3f FUsdPrimAttributeAccessor::GetValueAs<FVector3f>(const FVtValue& VtValue, const FVector3f DefaultValue) const
 	{
 		if (VtValue.GetTypeName() == TEXT("GfVec3f"))  // Note: FUsdAttribute.GetTypeName() would return "float3" instead!
 		{
@@ -110,7 +110,7 @@ namespace UE::Chaos::ClothAsset
 	}
 
 	template<>
-	TArray<float> FUsdPrimAttributeAccessor::GetArrayAs<float>(const FVtValue& VtValue) const
+	inline TArray<float> FUsdPrimAttributeAccessor::GetArrayAs<float>(const FVtValue& VtValue) const
 	{
 		TArray<float> Array;
 		if (VtValue.GetTypeName() == TEXT("VtArray<float>"))  // Note: FUsdAttribute.GetTypeName() would return "float[]" instead!
diff --git a/Engine/Plugins/Developer/Concert/ConcertSync/ConcertSyncTest/Source/ConcertSyncTest/Private/HistoryEdition/HistoryTestUtil.h b/Engine/Plugins/Developer/Concert/ConcertSync/ConcertSyncTest/Source/ConcertSyncTest/Private/HistoryEdition/HistoryTestUtil.h
--- a/Engine/Plugins/Developer/Concert/ConcertSync/ConcertSyncTest/Source/ConcertSyncTest/Private/HistoryEdition/HistoryTestUtil.h
+++ b/Engine/Plugins/Developer/Concert/ConcertSync/ConcertSyncTest/Source/ConcertSyncTest/Private/HistoryEdition/HistoryTestUtil.h
@@ -8,7 +8,7 @@
 
 namespace UE::ConcertSyncTests
 {
-	TArray<ConcertSyncCore::FActivityNodeID> ValidateEachActivityHasNode(
+	inline TArray<ConcertSyncCore::FActivityNodeID> ValidateEachActivityHasNode(
 		FAutomationTestBase& Test,
 		const TArray<int64>& ActivityMappings,
 		const ConcertSyncCore::FActivityDependencyGraph& Graph,
diff --git a/Engine/Plugins/Experimental/ChaosCachingUSD/Source/ChaosCachingUSD/Public/ChaosCachingUSD/Operations.h b/Engine/Plugins/Experimental/ChaosCachingUSD/Source/ChaosCachingUSD/Public/ChaosCachingUSD/Operations.h
--- a/Engine/Plugins/Experimental/ChaosCachingUSD/Source/ChaosCachingUSD/Public/ChaosCachingUSD/Operations.h
+++ b/Engine/Plugins/Experimental/ChaosCachingUSD/Source/ChaosCachingUSD/Public/ChaosCachingUSD/Operations.h
@@ -18,7 +18,7 @@
 namespace UE::ChaosCachingUSD
 {
 
-	bool ValuesDiffer(const pxr::VtArray<pxr::GfVec3f>& A, const pxr::VtArray<pxr::GfVec3f>& B, const float Tolerance = 1.0e-8, const uint64 stride = 1)
+	inline bool ValuesDiffer(const pxr::VtArray<pxr::GfVec3f>& A, const pxr::VtArray<pxr::GfVec3f>& B, const float Tolerance = 1.0e-8, const uint64 stride = 1)
 	{
 		if (A.size() != B.size()) return true;
 		for (size_t i = 0; i < A.size(); i += stride)
diff --git a/Engine/Plugins/Experimental/ChaosFlesh/Source/ChaosFleshEditor/Private/ChaosFlesh/ChaosFleshEditorPlugin.cpp b/Engine/Plugins/Experimental/ChaosFlesh/Source/ChaosFleshEditor/Private/ChaosFlesh/ChaosFleshEditorPlugin.cpp
--- a/Engine/Plugins/Experimental/ChaosFlesh/Source/ChaosFleshEditor/Private/ChaosFlesh/ChaosFleshEditorPlugin.cpp
+++ b/Engine/Plugins/Experimental/ChaosFlesh/Source/ChaosFleshEditor/Private/ChaosFlesh/ChaosFleshEditorPlugin.cpp
@@ -20,6 +20,8 @@
 
 #define LOCTEXT_NAMESPACE "FleshEditor"
 
+TOptional<FChaosFleshEditorStyle> FChaosFleshEditorStyle::Singleton;
+
 void IChaosFleshEditorPlugin::StartupModule()
 {
 	FChaosFleshEditorStyle::Get();
diff --git a/Engine/Plugins/Experimental/ChaosFlesh/Source/ChaosFleshEditor/Private/ChaosFlesh/Editor/FleshEditorStyle.h b/Engine/Plugins/Experimental/ChaosFlesh/Source/ChaosFleshEditor/Private/ChaosFlesh/Editor/FleshEditorStyle.h
--- a/Engine/Plugins/Experimental/ChaosFlesh/Source/ChaosFleshEditor/Private/ChaosFlesh/Editor/FleshEditorStyle.h
+++ b/Engine/Plugins/Experimental/ChaosFlesh/Source/ChaosFleshEditor/Private/ChaosFlesh/Editor/FleshEditorStyle.h
@@ -46,5 +46,3 @@ public:
 private:
 	static TOptional<FChaosFleshEditorStyle> Singleton;
 };
-
-TOptional<FChaosFleshEditorStyle> FChaosFleshEditorStyle::Singleton;
diff --git a/Engine/Plugins/Experimental/ChaosSolverPlugin/Source/ChaosSolverEditor/Private/Chaos/ChaosSolverEditorPlugin.cpp b/Engine/Plugins/Experimental/ChaosSolverPlugin/Source/ChaosSolverEditor/Private/Chaos/ChaosSolverEditorPlugin.cpp
--- a/Engine/Plugins/Experimental/ChaosSolverPlugin/Source/ChaosSolverEditor/Private/Chaos/ChaosSolverEditorPlugin.cpp
+++ b/Engine/Plugins/Experimental/ChaosSolverPlugin/Source/ChaosSolverEditor/Private/Chaos/ChaosSolverEditorPlugin.cpp
@@ -8,7 +8,7 @@
 #include "PropertyEditorModule.h"
 #include "ChaosSolverEditorDetails.h"
 
-
+TOptional<FChaosSolverEditorStyle> FChaosSolverEditorStyle::Singleton;
 
 IMPLEMENT_MODULE( IChaosSolverEditorPlugin, ChaosSolverEditor )
 
diff --git a/Engine/Plugins/Experimental/ChaosSolverPlugin/Source/ChaosSolverEditor/Private/Chaos/ChaosSolverEditorStyle.h b/Engine/Plugins/Experimental/ChaosSolverPlugin/Source/ChaosSolverEditor/Private/Chaos/ChaosSolverEditorStyle.h
--- a/Engine/Plugins/Experimental/ChaosSolverPlugin/Source/ChaosSolverEditor/Private/Chaos/ChaosSolverEditorStyle.h
+++ b/Engine/Plugins/Experimental/ChaosSolverPlugin/Source/ChaosSolverEditor/Private/Chaos/ChaosSolverEditorStyle.h
@@ -49,4 +49,3 @@ private:
 	static TOptional<FChaosSolverEditorStyle> Singleton;
 };
 
-TOptional<FChaosSolverEditorStyle> FChaosSolverEditorStyle::Singleton;
diff --git a/Engine/Plugins/Experimental/ChaosVehiclesPlugin/Source/ChaosVehiclesEditor/Private/ChaosVehiclesEditorPlugin.cpp b/Engine/Plugins/Experimental/ChaosVehiclesPlugin/Source/ChaosVehiclesEditor/Private/ChaosVehiclesEditorPlugin.cpp
--- a/Engine/Plugins/Experimental/ChaosVehiclesPlugin/Source/ChaosVehiclesEditor/Private/ChaosVehiclesEditorPlugin.cpp
+++ b/Engine/Plugins/Experimental/ChaosVehiclesPlugin/Source/ChaosVehiclesEditor/Private/ChaosVehiclesEditorPlugin.cpp
@@ -15,9 +15,9 @@
 //#include "ChaosVehiclesEditorDetails.h"
 #include "ChaosVehicleManager.h"
 
-IMPLEMENT_MODULE( IChaosVehiclesEditorPlugin, ChaosVehiclesEditor )
-
+TOptional<FChaosVehiclesEditorStyle> FChaosVehiclesEditorStyle::Singleton;
 
+IMPLEMENT_MODULE( IChaosVehiclesEditorPlugin, ChaosVehiclesEditor )
 
 void IChaosVehiclesEditorPlugin::StartupModule()
 {
diff --git a/Engine/Plugins/Experimental/ChaosVehiclesPlugin/Source/ChaosVehiclesEditor/Private/ChaosVehiclesEditorStyle.h b/Engine/Plugins/Experimental/ChaosVehiclesPlugin/Source/ChaosVehiclesEditor/Private/ChaosVehiclesEditorStyle.h
--- a/Engine/Plugins/Experimental/ChaosVehiclesPlugin/Source/ChaosVehiclesEditor/Private/ChaosVehiclesEditorStyle.h
+++ b/Engine/Plugins/Experimental/ChaosVehiclesPlugin/Source/ChaosVehiclesEditor/Private/ChaosVehiclesEditorStyle.h
@@ -47,5 +47,3 @@ public:
 private:
 	static TOptional<FChaosVehiclesEditorStyle> Singleton;
 };
-
-TOptional<FChaosVehiclesEditorStyle> FChaosVehiclesEditorStyle::Singleton;
\ No newline at end of file
diff --git a/Engine/Plugins/Experimental/FieldSystemPlugin/Source/FieldSyStemEditor/Private/Field/FieldSystemEditorModule.cpp b/Engine/Plugins/Experimental/FieldSystemPlugin/Source/FieldSyStemEditor/Private/Field/FieldSystemEditorModule.cpp
--- a/Engine/Plugins/Experimental/FieldSystemPlugin/Source/FieldSyStemEditor/Private/Field/FieldSystemEditorModule.cpp
+++ b/Engine/Plugins/Experimental/FieldSystemPlugin/Source/FieldSyStemEditor/Private/Field/FieldSystemEditorModule.cpp
@@ -6,6 +6,8 @@
 #include "Field/AssetTypeActions_FieldSystem.h"
 #include "Field/FieldSystemEditorStyle.h"
 
+TOptional<FFieldSystemEditorStyle> FFieldSystemEditorStyle::Singleton;
+
 IMPLEMENT_MODULE( IFieldSystemEditorModule, FieldSystemEditor )
 
 void IFieldSystemEditorModule::StartupModule()
diff --git a/Engine/Plugins/Experimental/FieldSystemPlugin/Source/FieldSyStemEditor/Private/Field/FieldSystemEditorStyle.h b/Engine/Plugins/Experimental/FieldSystemPlugin/Source/FieldSyStemEditor/Private/Field/FieldSystemEditorStyle.h
--- a/Engine/Plugins/Experimental/FieldSystemPlugin/Source/FieldSyStemEditor/Private/Field/FieldSystemEditorStyle.h
+++ b/Engine/Plugins/Experimental/FieldSystemPlugin/Source/FieldSyStemEditor/Private/Field/FieldSystemEditorStyle.h
@@ -49,4 +49,3 @@ private:
 	static TOptional<FFieldSystemEditorStyle> Singleton;
 };
 
-TOptional<FFieldSystemEditorStyle> FFieldSystemEditorStyle::Singleton;
diff --git a/Engine/Plugins/Experimental/GeometryCollectionPlugin/Source/GeometryCollectionEditor/Private/GeometryCollection/GeometryCollectionEditorPlugin.cpp b/Engine/Plugins/Experimental/GeometryCollectionPlugin/Source/GeometryCollectionEditor/Private/GeometryCollection/GeometryCollectionEditorPlugin.cpp
--- a/Engine/Plugins/Experimental/GeometryCollectionPlugin/Source/GeometryCollectionEditor/Private/GeometryCollection/GeometryCollectionEditorPlugin.cpp
+++ b/Engine/Plugins/Experimental/GeometryCollectionPlugin/Source/GeometryCollectionEditor/Private/GeometryCollection/GeometryCollectionEditorPlugin.cpp
@@ -32,6 +32,8 @@
 #include "Toolkits/AssetEditorToolkit.h"
 #include "Toolkits/IToolkitHost.h"
 
+TOptional<FGeometryCollectionEditorStyle> FGeometryCollectionEditorStyle::Singleton;
+
 IMPLEMENT_MODULE( IGeometryCollectionEditorPlugin, GeometryCollectionEditor )
 
 #define BOX_BRUSH(StyleSet, RelativePath, ...) FSlateBoxBrush(StyleSet->RootToContentDir(RelativePath, TEXT(".png")), __VA_ARGS__)
diff --git a/Engine/Plugins/Experimental/GeometryCollectionPlugin/Source/GeometryCollectionEditor/Private/GeometryCollection/GeometryCollectionEditorStyle.h b/Engine/Plugins/Experimental/GeometryCollectionPlugin/Source/GeometryCollectionEditor/Private/GeometryCollection/GeometryCollectionEditorStyle.h
--- a/Engine/Plugins/Experimental/GeometryCollectionPlugin/Source/GeometryCollectionEditor/Private/GeometryCollection/GeometryCollectionEditorStyle.h
+++ b/Engine/Plugins/Experimental/GeometryCollectionPlugin/Source/GeometryCollectionEditor/Private/GeometryCollection/GeometryCollectionEditorStyle.h
@@ -44,5 +44,3 @@ public:
 private:
 	static TOptional<FGeometryCollectionEditorStyle> Singleton;
 };
-
-TOptional<FGeometryCollectionEditorStyle> FGeometryCollectionEditorStyle::Singleton;
\ No newline at end of file
diff --git a/Engine/Plugins/Experimental/LiveLinkControlRig/Source/LiveLinkControlRig/Private/LiveLinkControlRig.cpp b/Engine/Plugins/Experimental/LiveLinkControlRig/Source/LiveLinkControlRig/Private/LiveLinkControlRig.cpp
--- a/Engine/Plugins/Experimental/LiveLinkControlRig/Source/LiveLinkControlRig/Private/LiveLinkControlRig.cpp
+++ b/Engine/Plugins/Experimental/LiveLinkControlRig/Source/LiveLinkControlRig/Private/LiveLinkControlRig.cpp
@@ -1,6 +1,8 @@
 // Copyright Epic Games, Inc. All Rights Reserved.
 #include "LiveLinkControlRig.h"
 
+IMPLEMENT_MODULE(FLiveLinkControlRigModule, LiveLinkControlRig)
+
 void FLiveLinkControlRigModule::StartupModule()
 {
 
diff --git a/Engine/Plugins/Experimental/LiveLinkControlRig/Source/LiveLinkControlRig/Public/LiveLinkControlRig.h b/Engine/Plugins/Experimental/LiveLinkControlRig/Source/LiveLinkControlRig/Public/LiveLinkControlRig.h
--- a/Engine/Plugins/Experimental/LiveLinkControlRig/Source/LiveLinkControlRig/Public/LiveLinkControlRig.h
+++ b/Engine/Plugins/Experimental/LiveLinkControlRig/Source/LiveLinkControlRig/Public/LiveLinkControlRig.h
@@ -9,6 +9,4 @@ class FLiveLinkControlRigModule : public IModuleInterface
 	// IModuleInterface interface
 	virtual void StartupModule() override;
 	virtual void ShutdownModule() override;
-};
-
-IMPLEMENT_MODULE(FLiveLinkControlRigModule, LiveLinkControlRig)
\ No newline at end of file
+};
\ No newline at end of file
diff --git a/Engine/Plugins/Experimental/QuicMessaging/Source/QuicMessagingTransport/Private/QuicCertificate.h b/Engine/Plugins/Experimental/QuicMessaging/Source/QuicMessagingTransport/Private/QuicCertificate.h
--- a/Engine/Plugins/Experimental/QuicMessaging/Source/QuicMessagingTransport/Private/QuicCertificate.h
+++ b/Engine/Plugins/Experimental/QuicMessaging/Source/QuicMessagingTransport/Private/QuicCertificate.h
@@ -30,7 +30,7 @@ namespace QuicCertificateUtils
 	 *
 	 * @param PrivateKey The private key struct
 	 */
-	bool GeneratePrivateKey(EVP_PKEY* PrivateKey)
+	inline bool GeneratePrivateKey(EVP_PKEY* PrivateKey)
 	{
 		if (!PrivateKey)
 		{
@@ -62,7 +62,7 @@ namespace QuicCertificateUtils
 	 *
 	 * @param PrivateKey The private key struct
 	 */
-	char* PrivateKeyToString(EVP_PKEY* PrivateKey)
+	inline char* PrivateKeyToString(EVP_PKEY* PrivateKey)
 	{
 		BIO* Bio = BIO_new(BIO_s_mem());
 
@@ -93,7 +93,7 @@ namespace QuicCertificateUtils
 	 * @param PrivateKey The private key struct
 	 * @param x509 The certificate struct
 	 */
-	bool GenerateCertificate(EVP_PKEY* PrivateKey, X509* x509)
+	inline bool GenerateCertificate(EVP_PKEY* PrivateKey, X509* x509)
 	{
 		if (!x509)
 		{
@@ -138,7 +138,7 @@ namespace QuicCertificateUtils
 	 *
 	 * @param x509 The certificate struct
 	 */
-	char* X509ToString(X509* x509)
+	inline char* X509ToString(X509* x509)
 	{
 		BIO* Bio = BIO_new(BIO_s_mem());
 		int Ret = PEM_write_bio_X509(Bio, x509);
@@ -157,7 +157,7 @@ namespace QuicCertificateUtils
 	/**
 	 * Get the paths for self-signed certificate and key.
 	 */
-	TTuple<FString, FString> GetSelfSignedPaths()
+	inline TTuple<FString, FString> GetSelfSignedPaths()
 	{
 		IPlatformFile& FileManager = FPlatformFileManager::Get().GetPlatformFile();
 
@@ -176,7 +176,7 @@ namespace QuicCertificateUtils
 	/**
 	 * Generates a self-signed certificate and returns paths to (certificate, private key).
 	 */
-	bool CreateSelfSigned()
+	inline bool CreateSelfSigned()
 	{
 		/* Allocate memory for the EVP_PKEY structure. */
 		EVP_PKEY* PrivateKey = EVP_PKEY_new();
diff --git a/Engine/Plugins/Media/ElectraUtil/Source/ElectraSamples/Public/ElectraTextureSampleUtils.h b/Engine/Plugins/Media/ElectraUtil/Source/ElectraSamples/Public/ElectraTextureSampleUtils.h
--- a/Engine/Plugins/Media/ElectraUtil/Source/ElectraSamples/Public/ElectraTextureSampleUtils.h
+++ b/Engine/Plugins/Media/ElectraUtil/Source/ElectraSamples/Public/ElectraTextureSampleUtils.h
@@ -12,7 +12,7 @@ namespace ElectraColorimetryUtils
 	static constexpr uint8 DefaultMPEGMatrixCoefficients = 2;
 	static constexpr uint8 DefaultMPEGTransferCharacteristics = 2;
 
-	UE::Color::EColorSpace TranslateMPEGColorPrimaries(uint8 InPrimaries)
+	inline UE::Color::EColorSpace TranslateMPEGColorPrimaries(uint8 InPrimaries)
 	{
 		switch (InPrimaries)
 		{
@@ -38,7 +38,7 @@ namespace ElectraColorimetryUtils
 		return UE::Color::EColorSpace::sRGB;
 	}
 
-	UE::Color::EColorSpace TranslateMPEGMatrixCoefficients(uint8 InMatrixCoefficients)
+	inline UE::Color::EColorSpace TranslateMPEGMatrixCoefficients(uint8 InMatrixCoefficients)
 	{
 		switch (InMatrixCoefficients)
 		{
@@ -68,7 +68,7 @@ namespace ElectraColorimetryUtils
 		return UE::Color::EColorSpace::sRGB;
 	}
 
-	UE::Color::EEncoding TranslateMPEGTransferCharacteristics(uint8 InTransferCharacteristics)
+	inline UE::Color::EEncoding TranslateMPEGTransferCharacteristics(uint8 InTransferCharacteristics)
 	{
 		switch (InTransferCharacteristics)
 		{
diff --git a/Engine/Plugins/Media/MediaIOFramework/Source/MediaIOCore/Private/MediaCapture.cpp b/Engine/Plugins/Media/MediaIOFramework/Source/MediaIOCore/Private/MediaCapture.cpp
--- a/Engine/Plugins/Media/MediaIOFramework/Source/MediaIOCore/Private/MediaCapture.cpp
+++ b/Engine/Plugins/Media/MediaIOFramework/Source/MediaIOCore/Private/MediaCapture.cpp
@@ -12,7 +12,7 @@
 #include "MediaCaptureHelper.h"
 #include "MediaCaptureRenderPass.h"
 #include "MediaCaptureSceneViewExtension.h"
-#include "MediaCaptureSources.h"
+#include "MediaCaptureSources.inl"
 #include "MediaCaptureSyncPointWatcher.h"
 #include "MediaIOFrameManager.h"
 #include "MediaIOCoreModule.h"
diff --git a/Engine/Plugins/Mutable/Source/MutableRuntime/Private/MuR/OpImageTransform.h b/Engine/Plugins/Mutable/Source/MutableRuntime/Private/MuR/OpImageTransform.h
--- a/Engine/Plugins/Mutable/Source/MutableRuntime/Private/MuR/OpImageTransform.h
+++ b/Engine/Plugins/Mutable/Source/MutableRuntime/Private/MuR/OpImageTransform.h
@@ -12,7 +12,7 @@
 
 namespace UE::Mutable::Private
 {
-	void ImageTransform(
+	inline void ImageTransform(
 			FImage* pDestImage, const FImage* pImage,
 			FTransform2f Transform, float MipFactor, EAddressMode AddressMode, bool bUseVectorImplCvar)
     {
diff --git a/Engine/Plugins/Mutable/Source/MutableRuntime/Private/MuR/OpMeshComputeNormals.h b/Engine/Plugins/Mutable/Source/MutableRuntime/Private/MuR/OpMeshComputeNormals.h
--- a/Engine/Plugins/Mutable/Source/MutableRuntime/Private/MuR/OpMeshComputeNormals.h
+++ b/Engine/Plugins/Mutable/Source/MutableRuntime/Private/MuR/OpMeshComputeNormals.h
@@ -18,7 +18,7 @@
 namespace UE::Mutable::Private
 {
 
-	float ComputeTangentBasisDeterminantSign(const FVector3f& N, const FVector3f& T, const FVector3f& B)
+	inline float ComputeTangentBasisDeterminantSign(const FVector3f& N, const FVector3f& T, const FVector3f& B)
 	{
 		const float BaseTangentBasisDeterminant =  
 				B.X*T.Y*N.Z + B.Z*T.X*N.Y + B.Y*T.Z*N.Y -
@@ -27,7 +27,7 @@ namespace UE::Mutable::Private
 		return BaseTangentBasisDeterminant < 0.0f ? -1.0f : 1.0f;
 	}
 
-	void OrthogonalizeTangentSpace(
+	inline void OrthogonalizeTangentSpace(
 			const FVector3f* NormalPtr, FVector3f* TangentPtr, FVector3f* BiNormalPtr, 
 			float TangentBasisDeterminantSign)	
 	{
@@ -44,7 +44,7 @@ namespace UE::Mutable::Private
 		}
 	}
 
-	void ComputeMeshNormals(FMesh& DestMesh)
+	inline void ComputeMeshNormals(FMesh& DestMesh)
 	{
 		MUTABLE_CPUPROFILER_SCOPE(ComputeMeshNormals);
 
diff --git a/Engine/Plugins/Mutable/Source/MutableRuntime/Private/MuR/OpMeshSmoothing.h b/Engine/Plugins/Mutable/Source/MutableRuntime/Private/MuR/OpMeshSmoothing.h
--- a/Engine/Plugins/Mutable/Source/MutableRuntime/Private/MuR/OpMeshSmoothing.h
+++ b/Engine/Plugins/Mutable/Source/MutableRuntime/Private/MuR/OpMeshSmoothing.h
@@ -20,7 +20,7 @@
 
 namespace UE::Mutable::Private
 {
-	TArray<int32> MakeUniqueVertexMap(TArrayView<const FVector3f> Vertices)
+	inline TArray<int32> MakeUniqueVertexMap(TArrayView<const FVector3f> Vertices)
 	{
 		MUTABLE_CPUPROFILER_SCOPE(MeshSmoothing_MakeUniqueVertexMap);
 
@@ -68,7 +68,7 @@ namespace UE::Mutable::Private
 		return UniqueVertexMap;
 	}
 
-	TArray<TArray<int32, TInlineAllocator<8>>> BuildVertexFaces(TArrayView<const uint32> Indices, const TArray<int32>& UniqueVertexMap)
+	inline TArray<TArray<int32, TInlineAllocator<8>>> BuildVertexFaces(TArrayView<const uint32> Indices, const TArray<int32>& UniqueVertexMap)
 	{
 		MUTABLE_CPUPROFILER_SCOPE(MeshSmoothing_BuildVertexFaces);
 		
@@ -105,7 +105,7 @@ namespace UE::Mutable::Private
 		return (static_cast<uint64>(E0) << (32 & TestMask)) | (static_cast<uint64>(E1) << (32 & (~TestMask)));
 	}
 
-	TMap<uint64, UE::Geometry::FIndex2i> BuildEdgesFaces(TArrayView<const uint32> Indices, const TArray<int32>& UniqueVertexMap)
+	inline TMap<uint64, UE::Geometry::FIndex2i> BuildEdgesFaces(TArrayView<const uint32> Indices, const TArray<int32>& UniqueVertexMap)
 	{
 		MUTABLE_CPUPROFILER_SCOPE(MeshSmoothing_BuildEdgesFaces);
 
@@ -179,7 +179,7 @@ namespace UE::Mutable::Private
 		return -1;
 	}
 
-	TTuple<TArray<int32>, TArray<int32>> BuildVertexRings(
+	inline TTuple<TArray<int32>, TArray<int32>> BuildVertexRings(
 		const TArrayView<const uint32> Indices,
 		const TArray<int32>& UniqueVertexMap,
 		const TArray<TArray<int32, TInlineAllocator<8>>>& VertexFaces,
diff --git a/Engine/Plugins/Online/OnlineSubsystemUtils/Source/OnlineSubsystemUtils/Private/Tests/TestVoice.cpp b/Engine/Plugins/Online/OnlineSubsystemUtils/Source/OnlineSubsystemUtils/Private/Tests/TestVoice.cpp
--- a/Engine/Plugins/Online/OnlineSubsystemUtils/Source/OnlineSubsystemUtils/Private/Tests/TestVoice.cpp
+++ b/Engine/Plugins/Online/OnlineSubsystemUtils/Source/OnlineSubsystemUtils/Private/Tests/TestVoice.cpp
@@ -3,7 +3,7 @@
 #include "Tests/TestVoice.h"
 #include "Components/AudioComponent.h"
 #include "Interfaces/VoiceCapture.h"
-#include "Tests/TestVoiceData.h"
+#include "Tests/TestVoiceData.inl"
 
 #include "Interfaces/VoiceCodec.h"
 #include "OnlineSubsystemUtils.h"
diff --git a/Engine/Plugins/Runtime/GameplayAbilities/Source/GameplayAbilitiesEditor/Private/SGameplayCueEditor_Picker.h b/Engine/Plugins/Runtime/GameplayAbilities/Source/GameplayAbilitiesEditor/Private/SGameplayCueEditor_Picker.h
--- a/Engine/Plugins/Runtime/GameplayAbilities/Source/GameplayAbilitiesEditor/Private/SGameplayCueEditor_Picker.h
+++ b/Engine/Plugins/Runtime/GameplayAbilities/Source/GameplayAbilitiesEditor/Private/SGameplayCueEditor_Picker.h
@@ -48,7 +48,10 @@ public:
 
 private:
 	/** Handler for when a class is picked in the class picker */
-	void OnClassPicked(UClass* InChosenClass);
+	void OnClassPicked(UClass* InChosenClass)
+	{
+		ChosenClass = InChosenClass;
+	}
 
 	/** Creates the default class widgets */
 	TSharedRef<ITableRow> GenerateListRow(UClass* InItem, const TSharedRef<STableViewBase>& OwnerTable);
@@ -78,264 +81,4 @@ private:
 	TArray<UClass*> DefaultClasses;
 
 	FString GameplayCueTag;
-};
-
-BEGIN_SLATE_FUNCTION_BUILD_OPTIMIZATION
-void SGameplayCuePickerDialog::Construct(const FArguments& InArgs)
-{
-	WeakParentWindow = InArgs._ParentWindow;
-	DefaultClasses = InArgs._DefaultClasses;
-	GameplayCueTag = InArgs._GameplayCueTag;
-
-	FLinearColor AssetColor = FLinearColor::White;
-
-	bPressedOk = false;
-	ChosenClass = NULL;
-
-	FString PathStr = SGameplayCueEditor::GetPathNameForGameplayCueTag(GameplayCueTag);
-	FGameplayCueEditorStrings Strings;
-	auto del = IGameplayAbilitiesEditorModule::Get().GetGameplayCueEditorStringsDelegate();
-	if (del.IsBound())
-	{
-		Strings = del.Execute();
-	}
-
-	ChildSlot
-	[
-		SNew(SBorder)
-		.Visibility(EVisibility::Visible)
-		.BorderImage(FAppStyle::GetBrush("Menu.Background"))
-		[
-			SNew(SBox)
-			.Visibility(EVisibility::Visible)
-			.Padding(2.f)
-			.WidthOverride(520.0f)
-			[
-				SNew(SVerticalBox)
-				+SVerticalBox::Slot()
-				.Padding(2.f, 2.f)
-				.AutoHeight()
-				[
-					SNew(SBorder)
-					.Visibility(EVisibility::Visible)
-					.BorderImage( FAppStyle::GetBrush("AssetThumbnail.AssetBackground") )
-					.BorderBackgroundColor(AssetColor.CopyWithNewOpacity(0.3f))
-					[
-						SNew(SExpandableArea)
-						.AreaTitle(NSLOCTEXT("SGameplayCuePickerDialog", "CommonClassesAreaTitle", "GameplayCue Notifies"))
-						.BodyContent()
-						[
-							SNew(SVerticalBox)
-							+SVerticalBox::Slot()
-							.Padding(2.f, 2.f)
-							.AutoHeight()
-							[
-								SNew(STextBlock)
-								
-								.Text(FText::FromString(Strings.GameplayCueNotifyDescription1))
-								.AutoWrapText(true)
-							]
-
-							+SVerticalBox::Slot()
-							.AutoHeight( )
-							[
-								SNew(SListView < UClass*  >)
-								.SelectionMode(ESelectionMode::None)
-								.ListItemsSource(&DefaultClasses)
-								.OnGenerateRow(this, &SGameplayCuePickerDialog::GenerateListRow)
-							]
-
-							+SVerticalBox::Slot()
-							.Padding(2.f, 2.f)
-							.AutoHeight()
-							[
-								SNew(STextBlock)
-								.Text(FText::FromString(FString::Printf(TEXT("This will create a new GameplayCue Notify here:"))))
-								.AutoWrapText(true)
-							]
-							+SVerticalBox::Slot()
-							.Padding(2.f, 2.f)
-							.AutoHeight()
-							[
-								SNew(STextBlock)
-								.Text(FText::FromString(PathStr))
-								.HighlightText(FText::FromString(PathStr))
-								.AutoWrapText(true)
-							]
-							+SVerticalBox::Slot()
-							.Padding(2.f, 2.f)
-							.AutoHeight()
-							[
-								SNew(STextBlock)
-								.Text(FText::FromString(Strings.GameplayCueNotifyDescription2))
-								.AutoWrapText(true)
-							]
-						]
-					]
-				]
-
-				+SVerticalBox::Slot()
-				.Padding(2.f, 2.f)
-				.AutoHeight()
-				[
-					SNew(SBorder)
-					.Visibility(EVisibility::Visible)
-					.BorderImage( FAppStyle::GetBrush("AssetThumbnail.AssetBackground") )
-					.BorderBackgroundColor(AssetColor.CopyWithNewOpacity(0.3f))
-					[
-						SNew(SExpandableArea)
-						.AreaTitle(NSLOCTEXT("SGameplayCuePickerDialogEvents", "CommonClassesAreaTitleEvents", "Custom BP Events"))
-						.BodyContent()
-						[
-							SNew(SVerticalBox)
-							+SVerticalBox::Slot()
-							.Padding(2.f, 2.f)
-							.AutoHeight()
-							[
-								SNew(STextBlock)
-								.Text(FText::FromString(Strings.GameplayCueEventDescription1))
-								.AutoWrapText(true)
-
-							]
-
-							+SVerticalBox::Slot()
-							.Padding(2.f, 2.f)
-							.AutoHeight()
-							[
-								SNew(STextBlock)
-								.Text(FText::FromString(Strings.GameplayCueEventDescription2))
-								.AutoWrapText(true)
-							]							
-						]
-					]
-				]
-			]
-		]
-	];
-}
-
-/** Spawns window for picking new GameplayCue handler/notify */
-bool SGameplayCuePickerDialog::PickGameplayCue(const FText& TitleText, const TArray<UClass*>& DefaultClasses, UClass*& OutChosenClass, FString InGameplayCueName)
-{
-	// Create the window to pick the class
-	TSharedRef<SWindow> PickerWindow = SNew(SWindow)
-		.Title(TitleText)
-		.SizingRule( ESizingRule::Autosized )
-		.ClientSize( FVector2D( 0.f, 600.f ))
-		.SupportsMaximize(false)
-		.SupportsMinimize(false);
-
-	TSharedRef<SGameplayCuePickerDialog> ClassPickerDialog = SNew(SGameplayCuePickerDialog)
-		.ParentWindow(PickerWindow)
-		.DefaultClasses(DefaultClasses)
-		.GameplayCueTag(InGameplayCueName);
-
-	PickerWindow->SetContent(ClassPickerDialog);
-
-	GEditor->EditorAddModalWindow(PickerWindow);
-
-	if (ClassPickerDialog->bPressedOk)
-	{
-		OutChosenClass = ClassPickerDialog->ChosenClass;
-		return true;
-	}
-	else
-	{
-		// Ok was not selected, NULL the class
-		OutChosenClass = NULL;
-		return false;
-	}
-}
-
-void SGameplayCuePickerDialog::OnClassPicked(UClass* InChosenClass)
-{
-	ChosenClass = InChosenClass;
-}
-
-/** Generates rows in the list of GameplayCueNotify classes to pick from */
-TSharedRef<ITableRow> SGameplayCuePickerDialog::GenerateListRow(UClass* ItemClass, const TSharedRef<STableViewBase>& OwnerTable)
-{	
-	const FSlateBrush* ItemBrush = FSlateIconFinder::FindIconBrushForClass(ItemClass);
-
-	return 
-	SNew(STableRow< UClass* >, OwnerTable)
-	[
-		SNew(SVerticalBox)
-		+SVerticalBox::Slot()
-		.MaxHeight(60.0f)
-		.Padding(10.0f, 6.0f, 0.0f, 4.0f)
-		[
-			SNew(SHorizontalBox)
-			+SHorizontalBox::Slot()
-			.FillWidth(0.65f)
-			[
-				SNew(SButton)
-				.OnClicked(this, &SGameplayCuePickerDialog::OnDefaultClassPicked, ItemClass)
-				.Content()
-				[
-					SNew(SHorizontalBox)
-					+SHorizontalBox::Slot()
-					.HAlign(HAlign_Center)
-					.VAlign(VAlign_Center)
-					.FillWidth(0.12f)
-					[
-						SNew(SImage)
-						.Image(ItemBrush)
-					]
-					+SHorizontalBox::Slot()
-					.VAlign(VAlign_Center)
-					.Padding(4.0f, 0.0f)
-					.FillWidth(0.8f)
-					[
-						SNew(STextBlock)
-						.Text(ItemClass->GetDisplayNameText())
-					]
-
-				]
-			]
-			+SHorizontalBox::Slot()
-			.Padding(10.0f, 0.0f)
-			[
-				SNew(STextBlock)
-				.Text(ItemClass->GetToolTipText(true))
-				.AutoWrapText(true)
-			]
-			+SHorizontalBox::Slot()
-			.AutoWidth()
-			[
-				FEditorClassUtils::GetDocumentationLinkWidget(ItemClass)
-			]
-		]
-	];
-}
-
-FReply SGameplayCuePickerDialog::OnDefaultClassPicked(UClass* InChosenClass)
-{
-	ChosenClass = InChosenClass;
-	bPressedOk = true;
-	if (WeakParentWindow.IsValid())
-	{
-		WeakParentWindow.Pin()->RequestDestroyWindow();
-	}
-	return FReply::Handled();
-}
-
-FReply SGameplayCuePickerDialog::OnClassPickerConfirmed()
-{
-	if (ChosenClass == NULL)
-	{
-		FMessageDialog::Open(EAppMsgType::Ok, NSLOCTEXT("EditorFactories", "MustChooseClassWarning", "You must choose a class."));
-	}
-	else
-	{
-		bPressedOk = true;
-
-		if (WeakParentWindow.IsValid())
-		{
-			WeakParentWindow.Pin()->RequestDestroyWindow();
-		}
-	}
-	return FReply::Handled();
-}
-
-END_SLATE_FUNCTION_BUILD_OPTIMIZATION
+};
\ No newline at end of file
diff --git a/Engine/Plugins/Runtime/Metasound/Source/MetasoundEngineTest/Private/EngineTestMetaSoundBuilder.h b/Engine/Plugins/Runtime/Metasound/Source/MetasoundEngineTest/Private/EngineTestMetaSoundBuilder.h
--- a/Engine/Plugins/Runtime/Metasound/Source/MetasoundEngineTest/Private/EngineTestMetaSoundBuilder.h
+++ b/Engine/Plugins/Runtime/Metasound/Source/MetasoundEngineTest/Private/EngineTestMetaSoundBuilder.h
@@ -23,7 +23,7 @@
 
 namespace EngineTestMetaSoundPatchBuilderPrivate
 {
-	UMetaSoundPatchBuilder& CreatePatchBuilderChecked(FAutomationTestBase& Test, FName BuilderName, const TArray<FName>& InterfaceNamesToAdd)
+	inline UMetaSoundPatchBuilder& CreatePatchBuilderChecked(FAutomationTestBase& Test, FName BuilderName, const TArray<FName>& InterfaceNamesToAdd)
 	{
 		EMetaSoundBuilderResult Result = EMetaSoundBuilderResult::Failed;
 		UMetaSoundPatchBuilder* Builder = UMetaSoundBuilderSubsystem::GetChecked().CreatePatchBuilder(BuilderName, Result);
@@ -39,7 +39,7 @@ namespace EngineTestMetaSoundPatchBuilderPrivate
 		return *Builder;
 	}
 
-	UMetaSoundPatch* CreatePatchFromBuilder(FAutomationTestBase& Test, const FString& PatchName, const TArray<FName>& InterfaceNamesToAdd)
+	inline UMetaSoundPatch* CreatePatchFromBuilder(FAutomationTestBase& Test, const FString& PatchName, const TArray<FName>& InterfaceNamesToAdd)
 	{
 		UMetaSoundPatchBuilder& Builder = CreatePatchBuilderChecked(Test, FName(PatchName + TEXT(" Builder")), InterfaceNamesToAdd);
 
@@ -52,7 +52,7 @@ namespace EngineTestMetaSoundPatchBuilderPrivate
 
 DEFINE_LATENT_AUTOMATION_COMMAND_THREE_PARAMETER(FMetaSoundSourceBuilderDisconnectInputLatentCommand, FAutomationTestBase&, Test, UMetaSoundBuilderBase*, Builder, FMetaSoundBuilderNodeInputHandle, InputToDisconnect);
 
-bool FMetaSoundSourceBuilderDisconnectInputLatentCommand::Update()
+inline bool FMetaSoundSourceBuilderDisconnectInputLatentCommand::Update()
 {
 	EMetaSoundBuilderResult Result = EMetaSoundBuilderResult::Failed;
 	if (Builder)
@@ -66,7 +66,7 @@ bool FMetaSoundSourceBuilderDisconnectInputLatentCommand::Update()
 
 DEFINE_LATENT_AUTOMATION_COMMAND_FOUR_PARAMETER(FMetaSoundSourceBuilderSetLiteralLatentCommand, FAutomationTestBase&, Test, UMetaSoundBuilderBase*, Builder, FMetaSoundBuilderNodeInputHandle, NodeInput, FMetasoundFrontendLiteral, NewValue);
 
-bool FMetaSoundSourceBuilderSetLiteralLatentCommand::Update()
+inline bool FMetaSoundSourceBuilderSetLiteralLatentCommand::Update()
 {
 	EMetaSoundBuilderResult Result = EMetaSoundBuilderResult::Failed;
 	if (Builder)
@@ -80,7 +80,7 @@ bool FMetaSoundSourceBuilderSetLiteralLatentCommand::Update()
 
 DEFINE_LATENT_AUTOMATION_COMMAND_THREE_PARAMETER(FMetaSoundSourceBuilderRemoveNodeDefaultLiteralLatentCommand, FAutomationTestBase&, Test, UMetaSoundBuilderBase*, Builder, FMetaSoundBuilderNodeInputHandle, NodeInput);
 
-bool FMetaSoundSourceBuilderRemoveNodeDefaultLiteralLatentCommand::Update()
+inline bool FMetaSoundSourceBuilderRemoveNodeDefaultLiteralLatentCommand::Update()
 {
 	EMetaSoundBuilderResult Result = EMetaSoundBuilderResult::Failed;
 	if (Builder)
@@ -94,7 +94,7 @@ bool FMetaSoundSourceBuilderRemoveNodeDefaultLiteralLatentCommand::Update()
 
 DEFINE_LATENT_AUTOMATION_COMMAND_THREE_PARAMETER(FMetaSoundSourceBuilderCreateAndConnectTriGeneratorNodeLatentCommand, FAutomationTestBase&, Test, UMetaSoundSourceBuilder*, Builder, FMetaSoundBuilderNodeInputHandle, AudioOutNodeInput);
 
-bool FMetaSoundSourceBuilderCreateAndConnectTriGeneratorNodeLatentCommand::Update()
+inline bool FMetaSoundSourceBuilderCreateAndConnectTriGeneratorNodeLatentCommand::Update()
 {
 	EMetaSoundBuilderResult Result = EMetaSoundBuilderResult::Failed;
 	if (Builder)
@@ -115,7 +115,7 @@ bool FMetaSoundSourceBuilderCreateAndConnectTriGeneratorNodeLatentCommand::Updat
 
 DEFINE_LATENT_AUTOMATION_COMMAND_ONE_PARAMETER(FBuilderRemoveFromRootLatentCommand, UMetaSoundBuilderBase*, Builder);
 
-bool FBuilderRemoveFromRootLatentCommand::Update()
+inline bool FBuilderRemoveFromRootLatentCommand::Update()
 {
 	if (Builder)
 	{
@@ -129,7 +129,7 @@ bool FBuilderRemoveFromRootLatentCommand::Update()
 
 IMPLEMENT_SIMPLE_AUTOMATION_TEST(FAudioMetaSoundBuilderInterfaceBindingConnectAndDisconnect, "Audio.Metasound.Builder.InterfaceBindingConnectAndDisconnect", EAutomationTestFlags::EditorContext | EAutomationTestFlags::EngineFilter)
 REGISTER_SIMPLE_AUTOMATION_TEST_TAGS(FAudioMetaSoundBuilderInterfaceBindingConnectAndDisconnect, "Audio.Metasound.Builder.InterfaceBindingConnectAndDisconnect", "[MustPass]")
-bool FAudioMetaSoundBuilderInterfaceBindingConnectAndDisconnect::RunTest(const FString& Parameters)
+inline bool FAudioMetaSoundBuilderInterfaceBindingConnectAndDisconnect::RunTest(const FString& Parameters)
 {
 	using namespace EngineTestMetaSoundPatchBuilderPrivate;
 	using namespace Metasound::Engine;
@@ -199,7 +199,7 @@ bool FAudioMetaSoundBuilderInterfaceBindingConnectAndDisconnect::RunTest(const F
 // This test creates a MetaSound patch, then adds and connects sin oscillator, attempts to retrieve a default input set, then clears it, and finally retrieves the class default
 IMPLEMENT_SIMPLE_AUTOMATION_TEST(FAudioMetaSoundNodeClassQueryFunctions, "Audio.Metasound.Builder.NodeClassQueryFunctions", EAutomationTestFlags::EditorContext | EAutomationTestFlags::EngineFilter)
 REGISTER_SIMPLE_AUTOMATION_TEST_TAGS(FAudioMetaSoundNodeClassQueryFunctions, "Audio.Metasound.Builder.NodeClassQueryFunctions", "[MustPass]")
-bool FAudioMetaSoundNodeClassQueryFunctions::RunTest(const FString& Parameters)
+inline bool FAudioMetaSoundNodeClassQueryFunctions::RunTest(const FString& Parameters)
 {
 	using namespace EngineTestMetaSoundPatchBuilderPrivate;
 	using namespace Metasound;
@@ -279,7 +279,7 @@ bool FAudioMetaSoundNodeClassQueryFunctions::RunTest(const FString& Parameters)
 
 // This test creates a MetaSound patch, adds a frequency input, connects a sin oscillator, and if the editor is loaded, and then attempts to inject template input nodes.
 IMPLEMENT_SIMPLE_AUTOMATION_TEST(FAudioMetaSoundInjectTemplateNodes, "Audio.Metasound.Builder.InjectInputTemplateNodes", EAutomationTestFlags::EditorContext | EAutomationTestFlags::EngineFilter)
-bool FAudioMetaSoundInjectTemplateNodes::RunTest(const FString& Parameters)
+inline bool FAudioMetaSoundInjectTemplateNodes::RunTest(const FString& Parameters)
 {
 	using namespace EngineTestMetaSoundPatchBuilderPrivate;
 	using namespace Metasound;
diff --git a/Engine/Plugins/VirtualProduction/EpicStageApp/Source/EpicStageApp/Private/EpicStageApp.cpp b/Engine/Plugins/VirtualProduction/EpicStageApp/Source/EpicStageApp/Private/EpicStageApp.cpp
--- a/Engine/Plugins/VirtualProduction/EpicStageApp/Source/EpicStageApp/Private/EpicStageApp.cpp
+++ b/Engine/Plugins/VirtualProduction/EpicStageApp/Source/EpicStageApp/Private/EpicStageApp.cpp
@@ -3,6 +3,8 @@
 #include "EpicStageApp.h"
 #include "IWebRemoteControlModule.h"
 
+IMPLEMENT_MODULE(FEpicStageAppModule, EpicStageApp);
+
 void FEpicStageAppModule::StartupModule()
 {
 	IWebRemoteControlModule& WebRemoteControl = FModuleManager::LoadModuleChecked<IWebRemoteControlModule>("WebRemoteControl");
diff --git a/Engine/Plugins/VirtualProduction/EpicStageApp/Source/EpicStageApp/Private/EpicStageApp.h b/Engine/Plugins/VirtualProduction/EpicStageApp/Source/EpicStageApp/Private/EpicStageApp.h
--- a/Engine/Plugins/VirtualProduction/EpicStageApp/Source/EpicStageApp/Private/EpicStageApp.h
+++ b/Engine/Plugins/VirtualProduction/EpicStageApp/Source/EpicStageApp/Private/EpicStageApp.h
@@ -26,5 +26,3 @@ private:
 	/** Delegate for when the websocket server starts. */
 	FDelegateHandle WebSocketServerStartedDelegate;
 };
-
-IMPLEMENT_MODULE(FEpicStageAppModule, EpicStageApp);
\ No newline at end of file
diff --git a/Engine/Plugins/VirtualProduction/RemoteControl/Source/RemoteControl/Private/RemoteControlModule.cpp b/Engine/Plugins/VirtualProduction/RemoteControl/Source/RemoteControl/Private/RemoteControlModule.cpp
--- a/Engine/Plugins/VirtualProduction/RemoteControl/Source/RemoteControl/Private/RemoteControlModule.cpp
+++ b/Engine/Plugins/VirtualProduction/RemoteControl/Source/RemoteControl/Private/RemoteControlModule.cpp
@@ -55,6 +55,8 @@
 DEFINE_LOG_CATEGORY(LogRemoteControl);
 #define LOCTEXT_NAMESPACE "RemoteControl"
 
+IMPLEMENT_MODULE(FRemoteControlModule, RemoteControl);
+
 struct FRCInterceptionPayload
 {
 	TArray<uint8> Payload;
diff --git a/Engine/Plugins/VirtualProduction/RemoteControl/Source/RemoteControl/Private/RemoteControlModule.h b/Engine/Plugins/VirtualProduction/RemoteControl/Source/RemoteControl/Private/RemoteControlModule.h
--- a/Engine/Plugins/VirtualProduction/RemoteControl/Source/RemoteControl/Private/RemoteControlModule.h
+++ b/Engine/Plugins/VirtualProduction/RemoteControl/Source/RemoteControl/Private/RemoteControlModule.h
@@ -321,4 +321,3 @@ private:
 
 PRAGMA_ENABLE_DEPRECATION_WARNINGS
 
-IMPLEMENT_MODULE(FRemoteControlModule, RemoteControl);
diff --git a/Engine/Plugins/VirtualProduction/RemoteControl/Source/WebRemoteControl/Private/RemoteControlDefaultPreprocessors.h b/Engine/Plugins/VirtualProduction/RemoteControl/Source/WebRemoteControl/Private/RemoteControlDefaultPreprocessors.h
--- a/Engine/Plugins/VirtualProduction/RemoteControl/Source/WebRemoteControl/Private/RemoteControlDefaultPreprocessors.h
+++ b/Engine/Plugins/VirtualProduction/RemoteControl/Source/WebRemoteControl/Private/RemoteControlDefaultPreprocessors.h
@@ -82,427 +82,33 @@ namespace UE::WebRemoteControl
 	using FRCPreprocessorHandler = TFunction<FPreprocessorResult(const FHttpServerRequest& Request)>;
 
 	/** Utility function to wrap a preprocessor handler to a http request handler than the HttpRouter can take. */
-	FHttpRequestHandler MakeHttpRequestHandler(FRCPreprocessorHandler Handler)
-	{
-		return FHttpRequestHandler::CreateLambda([WrappedHandler = MoveTemp(Handler)](const FHttpServerRequest& Request, const FHttpResultCallback& OnComplete)
-		{
-			FPreprocessorResult Result = WrappedHandler(Request);
-			if (Result.Result == EPreprocessorResult::RequestPassthrough)
-			{
-				return false; // Request not handled.
-			}
-			else
-			{
-				OnComplete(MoveTemp(Result.OptionalResponse));
-				return true; // Request handled.
-			}
-		});
-	}
+	FHttpRequestHandler MakeHttpRequestHandler(FRCPreprocessorHandler Handler);
 
 #if WITH_EDITOR
-	/** Notification prompt used when a remote client attempts to connect to unreal without a passphrase. */
-	TWeakPtr<SNotificationItem> NoPassphrasePrompt;
-	/** Message displayed in the notification prompt when no action was taken. */
-	FText NoActionTakenText = LOCTEXT("NoActionTaken", "No action was taken, further requests from this IP will be denied.");
-
 	/** Attempt saving the remote control config, checkouting it if needed.*/
-	void SaveRemoteControlConfig()
-	{
-		FString ConfigFilename = FPaths::ConvertRelativePathToFull(URemoteControlSettings::StaticClass()->GetConfigName());
-		
-		if (!FPlatformFileManager::Get().GetPlatformFile().IsReadOnly(*ConfigFilename))
-		{
-			GetMutableDefault<URemoteControlSettings>()->SaveConfig();
-		}
-		else
-		{
-			if (FPlatformFileManager::Get().GetPlatformFile().FileExists(*ConfigFilename))
-			{
-				if (ISourceControlModule::Get().IsEnabled())
-				{
-					if (SettingsHelpers::IsSourceControlled(ConfigFilename))
-					{
-						FText DisplayMessage = LOCTEXT("SaveConfigCheckoutMessage", "The configuration file for these settings is currently not checked out. Would you like to check it out from revision control?");
-						if (FMessageDialog::Open(EAppMsgType::YesNo, DisplayMessage) == EAppReturnType::Yes)
-						{
-							constexpr bool bForceSourceControlUpdate = true;
-							constexpr bool bShowErrorInNotification = true;
-							SettingsHelpers::CheckOutOrAddFile(ConfigFilename, bForceSourceControlUpdate, bShowErrorInNotification);
-							GetMutableDefault<URemoteControlSettings>()->SaveConfig();
-							return;
-						}
-					}
-				}
-
-				FText DisplayMessage = LOCTEXT("MakeConfigWritable", "The configuration file for these settings is currently read only. Would you like to make it writable?");
-				if (FMessageDialog::Open(EAppMsgType::YesNo, DisplayMessage) == EAppReturnType::Yes)
-				{
-					if (FPlatformFileManager::Get().GetPlatformFile().SetReadOnly(*ConfigFilename, false))
-					{
-						GetMutableDefault<URemoteControlSettings>()->SaveConfig();
-					}
-					else
-					{
-						FText ErrorMessage = FText::Format(LOCTEXT("FailedToMakeWritable", "Could not make {0} writable."), FText::FromString(ConfigFilename));;
-
-						FNotificationInfo MakeWritableNotification(MoveTemp(ErrorMessage));
-						MakeWritableNotification.ExpireDuration = 3.0f;
-						FSlateNotificationManager::Get().AddNotification(MakeWritableNotification);
-					}
-				}
-			}
-		}
-	}
+	void SaveRemoteControlConfig();
 
 	/** Add IP to the list of IPs that should be let through by RC. */
-	void AddIPToAllowlist(FString IPAddress)
-	{
-		if (TSharedPtr<SNotificationItem> Notification = NoPassphrasePrompt.Pin())
-		{
-			Notification->SetCompletionState(SNotificationItem::ECompletionState::CS_Pending);
-
-			GetMutableDefault<URemoteControlSettings>()->AllowClient(IPAddress);
-			SaveRemoteControlConfig();
-			
-			Notification->SetCompletionState(SNotificationItem::ECompletionState::CS_Success);
-			Notification->ExpireAndFadeout();
-		}
-	}
+	void AddIPToAllowlist(FString IPAddress);
 
 	/** Prompts the user to create a passphrase.  */
-	void CreatePassphrase()
-	{
-		if (TSharedPtr<SNotificationItem> Notification = NoPassphrasePrompt.Pin())
-		{
-			Notification->SetCompletionState(SNotificationItem::ECompletionState::CS_Pending);
-
-			TSharedPtr<SEditableTextBox> PassphraseIdentifierTextbox;
-			TSharedPtr<SEditableTextBox> PassphraseTextbox;
-
-			TSharedPtr<SWindow> Window = SNew(SWindow)
-				.Title(LOCTEXT("RemoteControlCreatePassphrase", "Create Passphrase"))
-				.SizingRule(ESizingRule::Autosized)
-				.SupportsMaximize(false)
-				.SupportsMinimize(false)
-				.FocusWhenFirstShown(true);
-
-			auto OnPassphraseEntered = [Window, Notification](const FText& PassphraseIdentifier, const FText& PassphraseText)
-			{
-				FRCPassphrase PassphraseStruct;
-				PassphraseStruct.Identifier = PassphraseIdentifier.ToString();
-				PassphraseStruct.Passphrase = FMD5::HashAnsiString(*PassphraseText.ToString());
-				GetMutableDefault<URemoteControlSettings>()->Passphrases.Add(MoveTemp(PassphraseStruct));
-				SaveRemoteControlConfig();
-
-				Notification->SetCompletionState(SNotificationItem::ECompletionState::CS_Success);
-				Notification->ExpireAndFadeout();
-
-				if (Window)
-				{
-					Window->RequestDestroyWindow();
-				}
-			};
-
-			Window->SetContent(
-				SNew(SBorder)
-				.Padding(4.f)
-				.VAlign(VAlign_Center)
-				.BorderImage(FAppStyle::GetBrush("ToolPanel.GroupBorder"))
-				[
-					SNew(SVerticalBox)
-					+ SVerticalBox::Slot()
-					.Padding(0.f, 5.f)
-					.AutoHeight()
-					[
-						SNew(STextBlock)
-						.Text(LOCTEXT("CreatePassphraseDescription", "Create a passphrase for the client to use."))
-					]
-					+ SVerticalBox::Slot()
-					.AutoHeight()
-					[
-						SAssignNew(PassphraseIdentifierTextbox, SEditableTextBox)
-						.HintText(LOCTEXT("PassphraseIdentifier", "Passphrase Identifier (Optional)"))
-					]
-					+ SVerticalBox::Slot()
-					.AutoHeight()
-					[
-						SAssignNew(PassphraseTextbox, SEditableTextBox)
-						.HintText(LOCTEXT("Passphrase", "Passphrase"))
-						.OnTextCommitted_Lambda([PassphraseIdentifierTextbox, OnPassphraseEntered](const FText& InText, ETextCommit::Type CommitType)
-						{
-							if (PassphraseIdentifierTextbox && !InText.IsEmpty() && CommitType == ETextCommit::OnEnter)
-							{
-								OnPassphraseEntered(PassphraseIdentifierTextbox->GetText(), InText);
-							}
-						})
-					]
-					+ SVerticalBox::Slot()
-					.HAlign(HAlign_Right)
-					.AutoHeight()
-					.Padding(0.0f, 2.0f, 0.0f, 0.0f)
-					[
-						SNew(SButton)
-						.Text(LOCTEXT("OK", "OK"))
-						.OnClicked_Lambda([PassphraseTextbox, PassphraseIdentifierTextbox, OnPassphraseEntered]()
-						{
-							if (PassphraseIdentifierTextbox && PassphraseTextbox && !PassphraseTextbox->GetText().IsEmpty())
-							{
-								OnPassphraseEntered(PassphraseIdentifierTextbox->GetText(), PassphraseTextbox->GetText());
-							}
-
-							return FReply::Handled();
-						})
-					]
-				]
-			);
-
-			Window->GetOnWindowClosedEvent().AddLambda([Notification](const TSharedRef<SWindow>&)
-			{
-				if (Notification->GetCompletionState() == SNotificationItem::ECompletionState::CS_Pending)
-				{
-					Notification->SetText(NoActionTakenText);
-					Notification->SetSubText(FText::GetEmpty());
-					Notification->SetCompletionState(SNotificationItem::ECompletionState::CS_Fail);
-					Notification->ExpireAndFadeout();
-				}
-			});
-
-			Window->GetOnWindowActivatedEvent().AddLambda([PassphraseTextbox]()
-			{
-				if (PassphraseTextbox)
-				{
-					constexpr uint32 UserIndex = 0;
-					FSlateApplication::Get().SetUserFocus(UserIndex, PassphraseTextbox, EFocusCause::SetDirectly);
-				}
-			});
-
-			GEditor->EditorAddModalWindow(Window.ToSharedRef());
-		}
-	}
+	void CreatePassphrase();
 
 	/** Disables remote passphrase enforcement by modifying a RC project setting. */
-	void DisableRemotePassphrases()
-	{
-		if (TSharedPtr<SNotificationItem> Notification = NoPassphrasePrompt.Pin())
-		{
-			Notification->SetCompletionState(SNotificationItem::ECompletionState::CS_Pending);
-			
-			if (FMessageDialog::Open(EAppMsgType::YesNo, LOCTEXT("DisableRemotePassphrases", "Are you sure you want to disable passphrase requirement for remote clients? This will allow anyone on your network to access the remote control servers and could open you up to vulnerabilities.")) == EAppReturnType::Yes)
-			{
-				GetMutableDefault<URemoteControlSettings>()->bEnforcePassphraseForRemoteClients = false;
-				SaveRemoteControlConfig();
-				
-				Notification->SetCompletionState(SNotificationItem::ECompletionState::CS_Success);
-                Notification->ExpireAndFadeout();
-			}
-			else
-			{
-				Notification->SetText(NoActionTakenText);
-				Notification->SetSubText(FText::GetEmpty());
-				Notification->SetCompletionState(SNotificationItem::ECompletionState::CS_Fail);
-				Notification->ExpireAndFadeout();
-			}
-		}
-	}
+	void DisableRemotePassphrases();
 #endif
 
 	// Notifies the editor if a client tries to access the RC server without a passphrase.
-	FPreprocessorResult RemotePassphraseEnforcementPreprocessor(const FHttpServerRequest& Request)
-	{
-		if (GetDefault<URemoteControlSettings>()->bRestrictServerAccess)
-		{
-			if (Request.PeerAddress)
-			{
-				constexpr bool bAppendPort = false;
-				FString PeerAddress = Request.PeerAddress->ToString(bAppendPort);
-
-				if (!GetDefault<URemoteControlSettings>()->bEnforcePassphraseForRemoteClients)
-				{
-					return FPreprocessorResult::Passthrough();
-				}
-
-				auto IsLocal = [](const FString& Address)
-				{
-					return Address.Contains(TEXT("127.0.0.1")) || Address.Contains(TEXT("localhost"));
-				};
-
-				bool bContainsForwardedAddress = false;
-				if (IsLocal(PeerAddress))
-				{
-					if (const TArray<FString>* ForwardedIP = Request.Headers.Find(WebRemoteControlInternalUtils::ForwardedIPHeader))
-					{
-						if (ForwardedIP->Num())
-						{
-							PeerAddress = ForwardedIP->Last();
-
-							if (GetDefault<URemoteControlSettings>()->IsClientAllowed(PeerAddress) || IsLocal(PeerAddress))
-							{
-								return FPreprocessorResult::Passthrough();
-							}
-							else
-							{
-								// Let the forwarded IP go through the rest of the validations if it's not 
-								bContainsForwardedAddress = true;
-							}
-						}
-						else
-						{
-							return FPreprocessorResult::Passthrough();
-						}
-					}
-					else
-					{
-						return FPreprocessorResult::Passthrough();
-					}
-				}
-
-				// We've already attempted to validate a forwarded address so just do it if it's a direct remote address.
-				if (!bContainsForwardedAddress && GetDefault<URemoteControlSettings>()->IsClientAllowed(PeerAddress))
-				{
-					return FPreprocessorResult::Passthrough();
-				}
-
-				if (GetDefault<URemoteControlSettings>()->Passphrases.Num())
-				{
-					const TArray<FString>* PassphraseHeader = Request.Headers.Find(WebRemoteControlInternalUtils::PassphraseHeader);
-					if (!PassphraseHeader || PassphraseHeader->Num() == 0)
-					{
-						return FPreprocessorResult::Deny(TEXT("Remote passphrase enforcement is enabled but no passphrase was specified!"));
-					}
-
-					if (!WebRemoteControlInternalUtils::CheckPassphrase(PassphraseHeader->Last()))
-					{
-						return FPreprocessorResult::Deny(WebRemoteControlInternalUtils::InvalidPassphraseError);
-					}
-					else
-					{
-						return FPreprocessorResult::Passthrough();
-					}
-				}
-				else
-				{
-					/* Prompt user, checkout settings if needed,
-					 * Either
-					 * 1: Explicitely add peer IP to Allowlist
-					 * 2: Create a passphrase, tell user he must enter that on the app he's using or put it in the header
-					 * 3: Disable remote passphrase enforcement. (Warn that this is dangerous)
-					 */
-#if WITH_EDITOR
-					if (GEditor)
-					{
-						if (!NoPassphrasePrompt.IsValid() || NoPassphrasePrompt.Pin()->GetCompletionState() == SNotificationItem::CS_None)
-						{
-							FNotificationInfo Info(LOCTEXT("NoPassphraseNotificationHeader", "Remote control request denied!"));
-							Info.SubText = FText::Format(LOCTEXT("NoPassphraseNotificationSubtext", "A remote control request was made by an external client ({0}) without a passphrase."), FText::FromString(PeerAddress));
-							Info.bFireAndForget = false;
-							Info.FadeInDuration = 0.5f;
-							Info.ButtonDetails.Add(FNotificationButtonInfo(LOCTEXT("AllowListIP", "AllowList IP"), LOCTEXT("AllowListTooltip", "Add this IP to the list of allowed IPs that can make remote control requests without a passphrase."), FSimpleDelegate::CreateStatic(&AddIPToAllowlist, PeerAddress)));
-							Info.ButtonDetails.Add(FNotificationButtonInfo(LOCTEXT("CreatePassphrase", "Create Passphrase"), LOCTEXT("CreatePassPhraseTooltip", "Create a passphrase for this client to use."), FSimpleDelegate::CreateStatic(&CreatePassphrase)));
-							Info.ButtonDetails.Add(FNotificationButtonInfo(LOCTEXT("DisableRemotePassphrase", "Disable Passphrases"), LOCTEXT("DisableRemotePassphrasesTooltip", "Disable the requirement for remote control requests coming from external clients to have a passphrase.\nWarning: This should only be done as a last resort since all clients on the network will be able to access your servers."), FSimpleDelegate::CreateStatic(&DisableRemotePassphrases)));
-							Info.WidthOverride = 450.f;
-							NoPassphrasePrompt = FSlateNotificationManager::Get().AddNotification(MoveTemp(Info));
-							NoPassphrasePrompt.Pin()->SetCompletionState(SNotificationItem::CS_Pending);
-							return FPreprocessorResult::Deny(TEXT("A passphrase for remote control is required. See editor for resolving this."));
-						}
-						else
-						{
-							return FPreprocessorResult::Deny(TEXT("A passphrase for remote control is required. See editor for resolving this."));
-						}
-					}
-					else
-#endif
-					{
-						// In -game or packaged, we won't prompt the user so just deny the request.
-						return FPreprocessorResult::Deny(TEXT("A passphrase for remote control is required but we can't create one in -game or packaged."));
-					}
-				}
-			}
-		}
-
-		return FPreprocessorResult::Passthrough();
-	}
+	FPreprocessorResult RemotePassphraseEnforcementPreprocessor(const FHttpServerRequest& Request);
 	
 	/** Checks whether a request has a valid passphrase when passphrases are enabled for this editor. */
-	FPreprocessorResult PassphrasePreprocessor(const FHttpServerRequest& Request)
-	{
-		TUniquePtr<FHttpServerResponse> Response = WebRemoteControlInternalUtils::CreateHttpResponse();
-
-		TArray<FString> ValueArray = {};
-		if (Request.Headers.Find(WebRemoteControlInternalUtils::PassphraseHeader))
-		{
-			ValueArray = Request.Headers[WebRemoteControlInternalUtils::PassphraseHeader];
-		}
-		
-		const FString Passphrase = !ValueArray.IsEmpty() ? ValueArray.Last() : "";
-
-		if (!WebRemoteControlInternalUtils::CheckPassphrase(Passphrase))
-		{
-			return FPreprocessorResult::Deny(WebRemoteControlInternalUtils::InvalidPassphraseError);
-		}
-
-		return FPreprocessorResult::Passthrough();
-	}
+	FPreprocessorResult PassphrasePreprocessor(const FHttpServerRequest& Request);
 
 	/** Checks whether an IP is in a valid range for making remote control requests. 
 	 * Also checks for a valid origin to block malicious requests from web browsers.
 	 * Can be controlled using the Allowed Origin and AllowedIP remote control settings.
 	 */
-	FPreprocessorResult IPValidationPreprocessor(const FHttpServerRequest& Request)
-	{
-		if (GetDefault<URemoteControlSettings>()->bRestrictServerAccess)
-		{
-			TUniquePtr<FHttpServerResponse> Response = WebRemoteControlInternalUtils::CreateHttpResponse();
-
-			FString OriginHeader;
-			if (const TArray<FString>* OriginHeaders = Request.Headers.Find(WebRemoteControlInternalUtils::OriginHeader))
-			{
-				if (OriginHeaders->Num())
-				{
-					OriginHeader = (*OriginHeaders)[0];
-				}
-			}
-
-			OriginHeader.RemoveSpacesInline();
-			OriginHeader.TrimStartAndEndInline();
-		
-			auto SimplifyAddress = [] (FString Address)
-			{
-				Address.RemoveFromStart(TEXT("https://www."));
-				Address.RemoveFromStart(TEXT("http://www."));
-				Address.RemoveFromStart(TEXT("https://"));
-				Address.RemoveFromStart(TEXT("http://"));
-				Address.RemoveFromEnd(TEXT("/"));
-				return Address;
-			};
-
-			const FString SimplifiedOrigin = SimplifyAddress(OriginHeader);
-			const FWildcardString SimplifiedAllowedOrigin = SimplifyAddress(GetDefault<URemoteControlSettings>()->AllowedOrigin);
-			if (!SimplifiedOrigin.IsEmpty() && GetDefault<URemoteControlSettings>()->AllowedOrigin != TEXT("*"))
-			{
-				if (!SimplifiedAllowedOrigin.IsMatch(SimplifiedOrigin))
-				{
-					return FPreprocessorResult::Deny(FString::Printf(TEXT("Client origin %s does not respect the allowed origin set in Remote Control Settings."), *OriginHeader));
-				}
-			}
-
-			if (Request.PeerAddress)
-			{
-				constexpr bool bAppendPort = false;
-				FString ClientIP = Request.PeerAddress->ToString(bAppendPort);
-				
-				// Allow requests from localhost
-				if (ClientIP != TEXT("localhost") && ClientIP != TEXT("127.0.0.1"))
-				{
-					if (!GetDefault<URemoteControlSettings>()->IsClientAllowed(ClientIP))
-					{
-						return FPreprocessorResult::Deny(FString::Printf(TEXT("Client IP %s does not respect the allowed IP set in Remote Control Settings."), *ClientIP));
-					}
-				}
-			}
-		}
-
-		return FPreprocessorResult::Passthrough();
-	}
+	FPreprocessorResult IPValidationPreprocessor(const FHttpServerRequest& Request);
 }
 
 #undef LOCTEXT_NAMESPACE /* WebRemoteControl */
diff --git a/Engine/Source/Developer/Datasmith/DatasmithExporterUI/Private/DatasmithExporterUIModule.cpp b/Engine/Source/Developer/Datasmith/DatasmithExporterUI/Private/DatasmithExporterUIModule.cpp
--- a/Engine/Source/Developer/Datasmith/DatasmithExporterUI/Private/DatasmithExporterUIModule.cpp
+++ b/Engine/Source/Developer/Datasmith/DatasmithExporterUI/Private/DatasmithExporterUIModule.cpp
@@ -7,6 +7,8 @@
 
 #include <atomic>
 
+IMPLEMENT_MODULE( FDatasmithExporterUIModule, DatasmithExporterUI );
+
 std::atomic<class FDatasmithExporterUIModule*> GExporterUIModule(nullptr);
 
 IDatasmithExporterUIModule* IDatasmithExporterUIModule::Get()
diff --git a/Engine/Source/Developer/Datasmith/DatasmithExporterUI/Private/DatasmithExporterUIModule.h b/Engine/Source/Developer/Datasmith/DatasmithExporterUI/Private/DatasmithExporterUIModule.h
--- a/Engine/Source/Developer/Datasmith/DatasmithExporterUI/Private/DatasmithExporterUIModule.h
+++ b/Engine/Source/Developer/Datasmith/DatasmithExporterUI/Private/DatasmithExporterUIModule.h
@@ -22,6 +22,4 @@ public:
 
 private:
 	TUniquePtr<IDirectLinkUI> DirectLinkUI;
-};
-
-IMPLEMENT_MODULE( FDatasmithExporterUIModule, DatasmithExporterUI );
\ No newline at end of file
+};
\ No newline at end of file
diff --git a/Engine/Source/Developer/ShaderCompilerCommon/Private/ISAParser/QuadModeDetect.cpp b/Engine/Source/Developer/ShaderCompilerCommon/Private/ISAParser/QuadModeDetect.cpp
--- a/Engine/Source/Developer/ShaderCompilerCommon/Private/ISAParser/QuadModeDetect.cpp
+++ b/Engine/Source/Developer/ShaderCompilerCommon/Private/ISAParser/QuadModeDetect.cpp
@@ -1,8 +1,8 @@
 // Copyright Epic Games, Inc. All Rights Reserved.
 
 #include "ISAParser/ISAParser.h"
-#include "RDNA1_ISA.h"
-#include "RDNA2_ISA.h"
+#include "RDNA1_ISA.inl"
+#include "RDNA2_ISA.inl"
 
 struct FDebugParserRDNA1
 {
diff --git a/Engine/Source/Developer/ShaderCompilerCommon/Private/ISAParser/RDNA_Common.h b/Engine/Source/Developer/ShaderCompilerCommon/Private/ISAParser/RDNA_Common.h
--- a/Engine/Source/Developer/ShaderCompilerCommon/Private/ISAParser/RDNA_Common.h
+++ b/Engine/Source/Developer/ShaderCompilerCommon/Private/ISAParser/RDNA_Common.h
@@ -340,7 +340,7 @@ struct FDPP16
 	uint32 ROW_MASK : 4;
 };
 
-uint32 GetAdvanceAmount(EInstructionType InstructionType)
+inline uint32 GetAdvanceAmount(EInstructionType InstructionType)
 {
 	switch (InstructionType)
 	{
@@ -369,7 +369,7 @@ uint32 GetAdvanceAmount(EInstructionType InstructionType)
 	return 0xffffffffu; // Force termination
 }
 
-bool IsSOP2(const uint32* Code, uint32& AdvanceAmount)
+inline bool IsSOP2(const uint32* Code, uint32& AdvanceAmount)
 {
 	const FInstSOP2& SOP2 = *reinterpret_cast<const FInstSOP2*>(Code);
 	const bool bMatch = (SOP2.ID == SOP2Header);
@@ -385,7 +385,7 @@ bool IsSOP2(const uint32* Code, uint32& AdvanceAmount)
 	return bMatch;
 }
 
-bool IsSOPK(const uint32* Code, uint32& AdvanceAmount)
+inline bool IsSOPK(const uint32* Code, uint32& AdvanceAmount)
 {
 	const FInstSOPK& SOPK = *reinterpret_cast<const FInstSOPK*>(Code);
 	const bool bMatch = (SOPK.ID == SOPKHeader);
@@ -399,7 +399,7 @@ bool IsSOPK(const uint32* Code, uint32& AdvanceAmount)
 	return bMatch;
 }
 
-bool IsSOP1(const uint32* Code, uint32& AdvanceAmount)
+inline bool IsSOP1(const uint32* Code, uint32& AdvanceAmount)
 {
 	const FInstSOP1& SOP1 = *reinterpret_cast<const FInstSOP1*>(Code);
 	const bool bMatch = (SOP1.ID == SOP1Header);
@@ -415,7 +415,7 @@ bool IsSOP1(const uint32* Code, uint32& AdvanceAmount)
 	return bMatch;
 }
 
-bool IsSOPC(const uint32* Code, uint32& AdvanceAmount)
+inline bool IsSOPC(const uint32* Code, uint32& AdvanceAmount)
 {
 	const FInstSOPC& SOPC = *reinterpret_cast<const FInstSOPC*>(Code);
 	const bool bMatch = (SOPC.ID == SOPCHeader);
@@ -431,7 +431,7 @@ bool IsSOPC(const uint32* Code, uint32& AdvanceAmount)
 	return bMatch;
 }
 
-bool IsSOPP(const uint32* Code, uint32& AdvanceAmount)
+inline bool IsSOPP(const uint32* Code, uint32& AdvanceAmount)
 {
 	const FInstSOPP& SOPP = *reinterpret_cast<const FInstSOPP*>(Code);
 	const bool bMatch = (SOPP.ID == SOPPHeader);
@@ -442,7 +442,7 @@ bool IsSOPP(const uint32* Code, uint32& AdvanceAmount)
 	return bMatch;
 }
 
-bool IsSMEM(const uint32* Code, uint32& AdvanceAmount)
+inline bool IsSMEM(const uint32* Code, uint32& AdvanceAmount)
 {
 	const FInstSMEM& SMEM = *reinterpret_cast<const FInstSMEM*>(Code);
 	const bool bMatch = (SMEM.ID == SMEMHeader);
@@ -453,7 +453,7 @@ bool IsSMEM(const uint32* Code, uint32& AdvanceAmount)
 	return bMatch;
 }
 
-bool IsVOP2(const uint32* Code, uint32& AdvanceAmount)
+inline bool IsVOP2(const uint32* Code, uint32& AdvanceAmount)
 {
 	const FInstVOP2& VOP2 = *reinterpret_cast<const FInstVOP2*>(Code);
 	const bool bMatch = (VOP2.ID == VOP2Header);
@@ -468,7 +468,7 @@ bool IsVOP2(const uint32* Code, uint32& AdvanceAmount)
 	return bMatch;
 }
 
-bool IsVOP3(const uint32* Code, uint32& AdvanceAmount)
+inline bool IsVOP3(const uint32* Code, uint32& AdvanceAmount)
 {
 	const FInstVOP3A& VOP3 = *reinterpret_cast<const FInstVOP3A*>(Code);
 	// VOP3A and VOP3B share the same header and bit range for the op code
@@ -486,7 +486,7 @@ bool IsVOP3(const uint32* Code, uint32& AdvanceAmount)
 	return bMatch;
 }
 
-bool IsVOP1(const uint32* Code, uint32& AdvanceAmount)
+inline bool IsVOP1(const uint32* Code, uint32& AdvanceAmount)
 {
 	const FInstVOP1& VOP1 = *reinterpret_cast<const FInstVOP1*>(Code);
 	const bool bMatch = (VOP1.ID == VOP1Header);
@@ -501,7 +501,7 @@ bool IsVOP1(const uint32* Code, uint32& AdvanceAmount)
 	return bMatch;
 }
 
-bool IsVOPC(const uint32* Code, uint32& AdvanceAmount)
+inline bool IsVOPC(const uint32* Code, uint32& AdvanceAmount)
 {
 	const FInstVOPC& VOPC = *reinterpret_cast<const FInstVOPC*>(Code);
 	const bool bMatch = (VOPC.ID == VOPCHeader);
@@ -516,7 +516,7 @@ bool IsVOPC(const uint32* Code, uint32& AdvanceAmount)
 	return bMatch;
 }
 
-bool IsVOP3P(const uint32* Code, uint32& AdvanceAmount)
+inline bool IsVOP3P(const uint32* Code, uint32& AdvanceAmount)
 {
 	const FInstVOP3P& VOP3P = *reinterpret_cast<const FInstVOP3P*>(Code);
 	const bool bMatch = (VOP3P.ID == VOP3PHeader);
@@ -527,7 +527,7 @@ bool IsVOP3P(const uint32* Code, uint32& AdvanceAmount)
 	return bMatch;
 }
 
-bool IsVINTERP(const uint32* Code, uint32& AdvanceAmount)
+inline bool IsVINTERP(const uint32* Code, uint32& AdvanceAmount)
 {
 	const FInstVINTERP& VINTERP = *reinterpret_cast<const FInstVINTERP*>(Code);
 	const bool bMatch = (VINTERP.ID == VINTERPHeader);
@@ -538,7 +538,7 @@ bool IsVINTERP(const uint32* Code, uint32& AdvanceAmount)
 	return bMatch;
 }
 
-bool IsLDSGDS(const uint32* Code, uint32& AdvanceAmount)
+inline bool IsLDSGDS(const uint32* Code, uint32& AdvanceAmount)
 {
 	const FInstLDSGDS& LDSGDS = *reinterpret_cast<const FInstLDSGDS*>(Code);
 	const bool bMatch = (LDSGDS.ID == LDSGDSHeader);
@@ -549,7 +549,7 @@ bool IsLDSGDS(const uint32* Code, uint32& AdvanceAmount)
 	return bMatch;
 }
 
-bool IsMUBUF(const uint32* Code, uint32& AdvanceAmount)
+inline bool IsMUBUF(const uint32* Code, uint32& AdvanceAmount)
 {
 	const FInstMUBUF& MUBUF = *reinterpret_cast<const FInstMUBUF*>(Code);
 	const bool bMatch = (MUBUF.ID == MUBUFHeader);
@@ -560,7 +560,7 @@ bool IsMUBUF(const uint32* Code, uint32& AdvanceAmount)
 	return bMatch;
 }
 
-bool IsMTBUF(const uint32* Code, uint32& AdvanceAmount)
+inline bool IsMTBUF(const uint32* Code, uint32& AdvanceAmount)
 {
 	const FInstMTBUF& MTBUF = *reinterpret_cast<const FInstMTBUF*>(Code);
 	const bool bMatch = (MTBUF.ID == MTBUFHeader);
@@ -571,7 +571,7 @@ bool IsMTBUF(const uint32* Code, uint32& AdvanceAmount)
 	return bMatch;
 }
 
-bool IsMIMG(const uint32* Code, uint32& AdvanceAmount)
+inline bool IsMIMG(const uint32* Code, uint32& AdvanceAmount)
 {
 	const FInstMIMG& MIMG = *reinterpret_cast<const FInstMIMG*>(Code);
 	const bool bMatch = (MIMG.ID == MIMGHeader);
@@ -583,7 +583,7 @@ bool IsMIMG(const uint32* Code, uint32& AdvanceAmount)
 	return bMatch;
 }
 
-bool IsEXPORT(const uint32* Code, uint32& AdvanceAmount)
+inline bool IsEXPORT(const uint32* Code, uint32& AdvanceAmount)
 {
 	const FInstEXPORT& EXPORT = *reinterpret_cast<const FInstEXPORT*>(Code);
 	const bool bMatch = (EXPORT.ID == EXPORTHeader);
@@ -594,7 +594,7 @@ bool IsEXPORT(const uint32* Code, uint32& AdvanceAmount)
 	return bMatch;
 }
 
-bool IsFSG(const uint32* Code, uint32& AdvanceAmount)
+inline bool IsFSG(const uint32* Code, uint32& AdvanceAmount)
 {
 	const FInstFSG& FSG = *reinterpret_cast<const FInstFSG*>(Code);
 	const bool bMatch = (FSG.ID == FSGHeader);
@@ -605,7 +605,7 @@ bool IsFSG(const uint32* Code, uint32& AdvanceAmount)
 	return bMatch;
 }
 
-EInstructionType DecodeInstructionType(const uint32* Code, uint32& AdvanceAmount)
+inline EInstructionType DecodeInstructionType(const uint32* Code, uint32& AdvanceAmount)
 {
 	if (IsSOP1(Code, AdvanceAmount)) return EInstructionType::SOP1;
 	if (IsSOPC(Code, AdvanceAmount)) return EInstructionType::SOPC;
@@ -628,7 +628,7 @@ EInstructionType DecodeInstructionType(const uint32* Code, uint32& AdvanceAmount
 	return EInstructionType::UNKNOWN;
 }
 
-const char* ToString(EInstructionType Inst)
+inline const char* ToString(EInstructionType Inst)
 {
 #define OP_TO_STRING_CASE(x) case EInstructionType::x: return #x
 
diff --git a/Engine/Source/Developer/WidgetRegistration/Public/WidgetRegistrationModule.h b/Engine/Source/Developer/WidgetRegistration/Public/WidgetRegistrationModule.h
--- a/Engine/Source/Developer/WidgetRegistration/Public/WidgetRegistrationModule.h
+++ b/Engine/Source/Developer/WidgetRegistration/Public/WidgetRegistrationModule.h
@@ -20,7 +20,4 @@ public:
 	
 	virtual void ShutdownModule() override;
 
-};
-
-
-IMPLEMENT_MODULE(FWidgetRegistrationModule, WidgetRegistration);
+};
\ No newline at end of file
diff --git a/Engine/Source/Editor/AnimationModifiers/Private/AnimationModifiersModule.cpp b/Engine/Source/Editor/AnimationModifiers/Private/AnimationModifiersModule.cpp
--- a/Engine/Source/Editor/AnimationModifiers/Private/AnimationModifiersModule.cpp
+++ b/Engine/Source/Editor/AnimationModifiers/Private/AnimationModifiersModule.cpp
@@ -53,6 +53,8 @@ class UFactory;
 
 IMPLEMENT_MODULE(FAnimationModifiersModule, AnimationModifiers);
 
+const FName FAnimationModifiersTabSummoner::AnimationModifiersName = TEXT("AnimationModifiers");
+
 void FAnimationModifiersModule::StartupModule()
 {
 	// Register class/struct customizations
diff --git a/Engine/Source/Editor/AnimationModifiers/Private/AnimationModifiersTabSummoner.h b/Engine/Source/Editor/AnimationModifiers/Private/AnimationModifiersTabSummoner.h
--- a/Engine/Source/Editor/AnimationModifiers/Private/AnimationModifiersTabSummoner.h
+++ b/Engine/Source/Editor/AnimationModifiers/Private/AnimationModifiersTabSummoner.h
@@ -38,6 +38,4 @@ public:
 	}
 };
 
-const FName FAnimationModifiersTabSummoner::AnimationModifiersName = TEXT("AnimationModifiers");
-
 #undef LOCTEXT_NAMESPACE // AnimationModifiersModule
\ No newline at end of file
diff --git a/Engine/Source/Editor/Kismet/Internal/Reflection/FunctionUtilsPrivate.h b/Engine/Source/Editor/Kismet/Internal/Reflection/FunctionUtilsPrivate.h
--- a/Engine/Source/Editor/Kismet/Internal/Reflection/FunctionUtilsPrivate.h
+++ b/Engine/Source/Editor/Kismet/Internal/Reflection/FunctionUtilsPrivate.h
@@ -126,13 +126,13 @@ namespace UE::Private
 		return Property && Type->IsChildOf(Property->PropertyClass);
 	}
 	template<>
-	bool DoesReferencePropertyMatch<FClassProperty>(const FClassProperty* Property, UClass* Type)
+	inline bool DoesReferencePropertyMatch<FClassProperty>(const FClassProperty* Property, UClass* Type)
 	{
 		check(Type);
 		return Property && Type->IsChildOf(Property->MetaClass);
 	}
 	template<>
-	bool DoesReferencePropertyMatch<FSoftClassProperty>(const FSoftClassProperty* Property, UClass* Type)
+	inline bool DoesReferencePropertyMatch<FSoftClassProperty>(const FSoftClassProperty* Property, UClass* Type)
 	{
 		check(Type);
 		return Property && Type->IsChildOf(Property->MetaClass);
diff --git a/Engine/Source/Editor/PlacementMode/Private/PlacementModeModule.cpp b/Engine/Source/Editor/PlacementMode/Private/PlacementModeModule.cpp
--- a/Engine/Source/Editor/PlacementMode/Private/PlacementModeModule.cpp
+++ b/Engine/Source/Editor/PlacementMode/Private/PlacementModeModule.cpp
@@ -59,6 +59,8 @@
 
 #define LOCTEXT_NAMESPACE "PlacementMode"
 
+IMPLEMENT_MODULE(FPlacementModeModule, PlacementMode);
+
 namespace PlacementModeModuleLocals
 {
 	FActorPlacementInfo MakePlacementInfo(const FPlaceableItem& Item)
diff --git a/Engine/Source/Editor/PlacementMode/Public/PlacementModeModule.h b/Engine/Source/Editor/PlacementMode/Public/PlacementModeModule.h
--- a/Engine/Source/Editor/PlacementMode/Public/PlacementModeModule.h
+++ b/Engine/Source/Editor/PlacementMode/Public/PlacementModeModule.h
@@ -161,5 +161,3 @@ private:
 	// need to store extra data to be able to recreate the placeable item in "recently placed"
 	TMap<FActorPlacementInfo, TWeakPtr<FPlaceableItem>> ManuallyCreatedPlaceableItems;
 };
-
-IMPLEMENT_MODULE(FPlacementModeModule, PlacementMode);
\ No newline at end of file
diff --git a/Engine/Source/Editor/SparseVolumeTexture/Private/OpenVDBGridAdapter.h b/Engine/Source/Editor/SparseVolumeTexture/Private/OpenVDBGridAdapter.h
--- a/Engine/Source/Editor/SparseVolumeTexture/Private/OpenVDBGridAdapter.h
+++ b/Engine/Source/Editor/SparseVolumeTexture/Private/OpenVDBGridAdapter.h
@@ -41,19 +41,19 @@ float GetOpenVDBValueComponent(const ValueType& Value, uint32 ComponentIndex)
 }
 
 template<>
-float GetOpenVDBValueComponent<FOpenVDBHalf>(const FOpenVDBHalf& Value, uint32 ComponentIndex)
+inline float GetOpenVDBValueComponent<FOpenVDBHalf>(const FOpenVDBHalf& Value, uint32 ComponentIndex)
 {
 	return (float)Value;
 }
 
 template<>
-float GetOpenVDBValueComponent<float>(const float& Value, uint32 ComponentIndex)
+inline float GetOpenVDBValueComponent<float>(const float& Value, uint32 ComponentIndex)
 {
 	return Value;
 }
 
 template<>
-float GetOpenVDBValueComponent<double>(const double& Value, uint32 ComponentIndex)
+inline float GetOpenVDBValueComponent<double>(const double& Value, uint32 ComponentIndex)
 {
 	return (float)Value;
 }
@@ -65,19 +65,19 @@ ValueType GetOpenVDBMaxComponent(const ValueType& ValueA, const ValueType& Value
 }
 
 template<>
-FOpenVDBHalf GetOpenVDBMaxComponent<FOpenVDBHalf>(const FOpenVDBHalf& ValueA, const FOpenVDBHalf& ValueB)
+inline FOpenVDBHalf GetOpenVDBMaxComponent<FOpenVDBHalf>(const FOpenVDBHalf& ValueA, const FOpenVDBHalf& ValueB)
 {
 	return FOpenVDBHalf(FMath::Max((float)ValueA, (float)ValueB));
 }
 
 template<>
-float GetOpenVDBMaxComponent<float>(const float& ValueA, const float& ValueB)
+inline float GetOpenVDBMaxComponent<float>(const float& ValueA, const float& ValueB)
 {
 	return FMath::Max(ValueA, ValueB);
 }
 
 template<>
-double GetOpenVDBMaxComponent<double>(const double& ValueA, const double& ValueB)
+inline double GetOpenVDBMaxComponent<double>(const double& ValueA, const double& ValueB)
 {
 	return FMath::Max(ValueA, ValueB);
 }
@@ -89,19 +89,19 @@ ValueType GetOpenVDBMinComponent(const ValueType& ValueA, const ValueType& Value
 }
 
 template<>
-FOpenVDBHalf GetOpenVDBMinComponent<FOpenVDBHalf>(const FOpenVDBHalf& ValueA, const FOpenVDBHalf& ValueB)
+inline FOpenVDBHalf GetOpenVDBMinComponent<FOpenVDBHalf>(const FOpenVDBHalf& ValueA, const FOpenVDBHalf& ValueB)
 {
 	return FOpenVDBHalf(FMath::Min((float)ValueA, (float)ValueB));
 }
 
 template<>
-float GetOpenVDBMinComponent<float>(const float& ValueA, const float& ValueB)
+inline float GetOpenVDBMinComponent<float>(const float& ValueA, const float& ValueB)
 {
 	return FMath::Min(ValueA, ValueB);
 }
 
 template<>
-double GetOpenVDBMinComponent<double>(const double& ValueA, const double& ValueB)
+inline double GetOpenVDBMinComponent<double>(const double& ValueA, const double& ValueB)
 {
 	return FMath::Min(ValueA, ValueB);
 }
@@ -220,7 +220,7 @@ TSharedPtr<IOpenVDBGridAdapterBase> CreateOpenVDBGridAdapterInternal(openvdb::Gr
 }
 
 // Creates an adapter suitable for the type of the given grid or nullptr if the type is not supported.
-TSharedPtr<IOpenVDBGridAdapterBase> CreateOpenVDBGridAdapter(openvdb::GridBase::Ptr Grid)
+inline TSharedPtr<IOpenVDBGridAdapterBase> CreateOpenVDBGridAdapter(openvdb::GridBase::Ptr Grid)
 {
 	if (Grid->isType<FOpenVDBHalf1Grid>())
 	{
diff --git a/Engine/Source/Runtime/AVEncoder/Private/Decoders/vdecmpeg4/M4VlcDecoder.cpp b/Engine/Source/Runtime/AVEncoder/Private/Decoders/vdecmpeg4/M4VlcDecoder.cpp
--- a/Engine/Source/Runtime/AVEncoder/Private/Decoders/vdecmpeg4/M4VlcDecoder.cpp
+++ b/Engine/Source/Runtime/AVEncoder/Private/Decoders/vdecmpeg4/M4VlcDecoder.cpp
@@ -2,7 +2,7 @@
 
 #include "M4VlcDecoder.h"
 #include "M4Global.h"
-#include "M4VlcData.h"
+#include "M4VlcData.inl"
 #include "M4Decoder.h"
 
 namespace vdecmpeg4
diff --git a/Engine/Source/Runtime/AssetRegistry/Internal/AssetRegistry/LinearBlockAllocator.h b/Engine/Source/Runtime/AssetRegistry/Internal/AssetRegistry/LinearBlockAllocator.h
--- a/Engine/Source/Runtime/AssetRegistry/Internal/AssetRegistry/LinearBlockAllocator.h
+++ b/Engine/Source/Runtime/AssetRegistry/Internal/AssetRegistry/LinearBlockAllocator.h
@@ -153,7 +153,7 @@ inline void* FLinearBlockAllocator::TryAllocateFromBlock(SIZE_T Size, uint32 Ali
 	return AlignedPtr;
 }
 
-bool FLinearBlockAllocator::RequiresCustomAllocation(SIZE_T Size, uint32 Alignment) const
+inline bool FLinearBlockAllocator::RequiresCustomAllocation(SIZE_T Size, uint32 Alignment) const
 {
 	using namespace UE::LinearBlockAllocator::Private;
 
@@ -169,7 +169,7 @@ bool FLinearBlockAllocator::RequiresCustomAllocation(SIZE_T Size, uint32 Alignme
 	return RequiredBlockSize > BlockSize;
 }
 
-UE::LinearBlockAllocator::Private::FBlock* FLinearBlockAllocator::GetNewBlock()
+inline UE::LinearBlockAllocator::Private::FBlock* FLinearBlockAllocator::GetNewBlock()
 {
 	using namespace UE::LinearBlockAllocator::Private;
 
@@ -181,7 +181,7 @@ UE::LinearBlockAllocator::Private::FBlock* FLinearBlockAllocator::GetNewBlock()
 	return NewBlock;
 }
 
-void* FLinearBlockAllocator::GetCustomAllocation(SIZE_T Size, uint32 Alignment)
+inline void* FLinearBlockAllocator::GetCustomAllocation(SIZE_T Size, uint32 Alignment)
 {
 	using namespace UE::LinearBlockAllocator::Private;
 
diff --git a/Engine/Source/Runtime/Experimental/Chaos/Public/Chaos/Complex.h b/Engine/Source/Runtime/Experimental/Chaos/Public/Chaos/Complex.h
--- a/Engine/Source/Runtime/Experimental/Chaos/Public/Chaos/Complex.h
+++ b/Engine/Source/Runtime/Experimental/Chaos/Public/Chaos/Complex.h
@@ -29,7 +29,7 @@ class FComplex
 	FReal MImaginary;
 };
 
-FComplex operator*(const FReal Other, const FComplex Complex)
+inline FComplex operator*(const FReal Other, const FComplex Complex)
 {
 	return Complex * Other;
 }
diff --git a/Engine/Source/Runtime/Launch/Public/RequiredProgramMainCPPInclude.h b/Engine/Source/Runtime/Launch/Public/RequiredProgramMainCPPInclude.h
--- a/Engine/Source/Runtime/Launch/Public/RequiredProgramMainCPPInclude.h
+++ b/Engine/Source/Runtime/Launch/Public/RequiredProgramMainCPPInclude.h
@@ -4,6 +4,9 @@
 
 #include "CoreMinimal.h"
 #include "LaunchEngineLoop.h"
+#include "Misc/DirectHeaderCompiling.h"
 
+#if !UE_DIRECT_HEADER_COMPILING(RequiredProgramMainCPPInclude)
 // this is highly sketchy, but we need some stuff from launchengineloop.cpp
 #include "Runtime/Launch/Private/LaunchEngineLoop.cpp"
+#endif
diff --git a/Engine/Source/Runtime/Online/BuildPatchServices/Private/Tests/Fake/HttpManager.fake.h b/Engine/Source/Runtime/Online/BuildPatchServices/Private/Tests/Fake/HttpManager.fake.h
--- a/Engine/Source/Runtime/Online/BuildPatchServices/Private/Tests/Fake/HttpManager.fake.h
+++ b/Engine/Source/Runtime/Online/BuildPatchServices/Private/Tests/Fake/HttpManager.fake.h
@@ -50,40 +50,40 @@ namespace BuildPatchServices
 		TMap<FString, TArray<uint8>> DataServed;
 	};
 
-	bool FFakeHttpRequest::ProcessRequest()
+	inline bool FFakeHttpRequest::ProcessRequest()
 	{
 		++RxProcessRequest;
 		return FakeHttpModule->OnProcessRequest(this);
 	}
 
-	int32 FFakeHttpResponse::GetResponseCode() const
+	inline int32 FFakeHttpResponse::GetResponseCode() const
 	{
 		return static_cast<int32>(Code);
 	}
 
-	const TArray<uint8>& FFakeHttpResponse::GetContent() const
+	inline const TArray<uint8>& FFakeHttpResponse::GetContent() const
 	{
 		return Data;
 	}
 
-	FString FFakeHttpResponse::GetContentAsString() const
+	inline FString FFakeHttpResponse::GetContentAsString() const
 	{
 		TArray<uint8> ZeroTerminatedPayload(GetContent());
 		ZeroTerminatedPayload.Add(0);
 		return UTF8_TO_TCHAR(ZeroTerminatedPayload.GetData());
 	}
 
-	uint64 FFakeHttpResponse::GetContentLength() const
+	inline uint64 FFakeHttpResponse::GetContentLength() const
 	{
 		return Data.Num();
 	}
 
-	FFakeHttpManager::FFakeHttpManager(FTSTicker& Ticker)
+	inline FFakeHttpManager::FFakeHttpManager(FTSTicker& Ticker)
 	{
 		Ticker.AddTicker(FTickerDelegate::CreateRaw(this, &FFakeHttpManager::Tick));
 	}
 
-	TSharedRef<IHttpRequest, ESPMode::ThreadSafe> FFakeHttpManager::CreateRequest()
+	inline TSharedRef<IHttpRequest, ESPMode::ThreadSafe> FFakeHttpManager::CreateRequest()
 	{
 		++RxCreateRequest;
 		NewRequests.Emplace(new FFakeHttpRequest());
@@ -91,7 +91,7 @@ namespace BuildPatchServices
 		return NewRequests.Last();
 	}
 
-	bool FFakeHttpManager::Tick(float Delta)
+	inline bool FFakeHttpManager::Tick(float Delta)
 	{
 		// Complete progressed.
 		for (const TSharedRef<FFakeHttpRequest, ESPMode::ThreadSafe>& ProgressedRequest : ProgressedRequests)
@@ -128,7 +128,7 @@ namespace BuildPatchServices
 		return true;
 	}
 
-	bool FFakeHttpManager::OnProcessRequest(FFakeHttpRequest* FakeHttpRequest)
+	inline bool FFakeHttpManager::OnProcessRequest(FFakeHttpRequest* FakeHttpRequest)
 	{
 		int32 Index = NewRequests.IndexOfByPredicate([FakeHttpRequest](const TSharedRef<FFakeHttpRequest, ESPMode::ThreadSafe>& Element)
 		{
diff --git a/Engine/Source/Runtime/Online/BuildPatchServices/Private/Tests/Unit/ProcessTimer.spec.cpp b/Engine/Source/Runtime/Online/BuildPatchServices/Private/Tests/Unit/ProcessTimer.spec.cpp
--- a/Engine/Source/Runtime/Online/BuildPatchServices/Private/Tests/Unit/ProcessTimer.spec.cpp
+++ b/Engine/Source/Runtime/Online/BuildPatchServices/Private/Tests/Unit/ProcessTimer.spec.cpp
@@ -2,7 +2,7 @@
 
 #include "Misc/AutomationTest.h"
 #include "Tests/TestHelpers.h"
-#include "Tests/Mock/CyclesProvider.mock.h"
+#include "Tests/Mock/CyclesProvider.mock.inl"
 #include "Core/ProcessTimer.h"
 
 #if WITH_DEV_AUTOMATION_TESTS
diff --git a/Engine/Source/Runtime/Renderer/Internal/PostProcess/ColorGradingScopes.h b/Engine/Source/Runtime/Renderer/Internal/PostProcess/ColorGradingScopes.h
--- a/Engine/Source/Runtime/Renderer/Internal/PostProcess/ColorGradingScopes.h
+++ b/Engine/Source/Runtime/Renderer/Internal/PostProcess/ColorGradingScopes.h
@@ -38,7 +38,7 @@ namespace UE::ColorParade
 		FRDGTexture* Output, FIntRect OutputViewRect,
 		const FDrawParameters& Parameters);
 
-	void Draw(
+	inline void Draw(
 		FRDGBuilder& GraphBuilder,
 		FRDGTextureSRV* InputSRV, FIntRect InputViewRect,
 		FRDGTexture* Output, FIntRect OutputViewRect)
@@ -76,7 +76,7 @@ namespace UE::VectorScope
 		FRDGTexture* Output, FIntRect OutputViewRect,
 		const FDrawParameters& Parameters);
 
-	void Draw(
+	inline void Draw(
 		FRDGBuilder& GraphBuilder,
 		FRDGTextureSRV* InputSRV, FIntRect InputViewRect,
 		FRDGTexture* Output, FIntRect OutputViewRect)
@@ -108,7 +108,7 @@ namespace UE::ColorHistogram
 		FRDGTexture* Output, FIntRect OutputViewRect,
 		const FDrawParameters& Parameters);
 
-	void Draw(
+	inline void Draw(
 		FRDGBuilder& GraphBuilder,
 		FRDGTextureSRV* InputSRV, FIntRect InputViewRect,
 		FRDGTexture* Output, FIntRect OutputViewRect)
diff --git a/Engine/Source/Runtime/Renderer/Private/Shadows/bend_sss_cpu.h b/Engine/Source/Runtime/Renderer/Private/Shadows/bend_sss_cpu.h
--- a/Engine/Source/Runtime/Renderer/Private/Shadows/bend_sss_cpu.h
+++ b/Engine/Source/Runtime/Renderer/Private/Shadows/bend_sss_cpu.h
@@ -70,7 +70,7 @@ namespace Bend
 	// 
 	// inWaveSize:				Wavefront size of the compiled compute shader (currently only tested with 64)
 	//
-	DispatchList BuildDispatchList(float inLightProjection[4], int inViewportSize[2], int inMinRenderBounds[2], int inMaxRenderBounds[2], bool inExpandedZRange = false, int inWaveSize = 64)
+	inline DispatchList BuildDispatchList(float inLightProjection[4], int inViewportSize[2], int inMinRenderBounds[2], int inMaxRenderBounds[2], bool inExpandedZRange = false, int inWaveSize = 64)
 	{
 		DispatchList result = {};
 
diff --git a/Engine/Source/Runtime/Renderer/Private/VT/UniquePageList.h b/Engine/Source/Runtime/Renderer/Private/VT/UniquePageList.h
--- a/Engine/Source/Runtime/Renderer/Private/VT/UniquePageList.h
+++ b/Engine/Source/Runtime/Renderer/Private/VT/UniquePageList.h
@@ -36,13 +36,13 @@ private:
 	uint16 Counts[ MaxUniquePages ];
 };
 
-FUniquePageList::FUniquePageList()
+inline FUniquePageList::FUniquePageList()
 	: bInitialized( false )
 	, NumPages( 0 )
 	, MaxNumCollisions( 0 )
 {}
 
-void FUniquePageList::Initialize()
+inline void FUniquePageList::Initialize()
 {
 	if (!bInitialized)
 	{
@@ -51,7 +51,7 @@ void FUniquePageList::Initialize()
 	}
 }
 
-void FUniquePageList::Add( uint32 Page, uint32 Count )
+inline void FUniquePageList::Add( uint32 Page, uint32 Count )
 {
 	uint32 HashIndex = MurmurFinalize32(Page) & (HashSize - 1u);
 	uint32 NumCollisions = 0u;
@@ -83,7 +83,7 @@ void FUniquePageList::Add( uint32 Page, uint32 Count )
 #endif // DO_GUARD_SLOW
 }
 
-void FUniquePageList::MergePages(const FUniquePageList* RESTRICT Other)
+inline void FUniquePageList::MergePages(const FUniquePageList* RESTRICT Other)
 {
 	for (uint32 Index = 0u; Index < Other->NumPages; ++Index)
 	{
diff --git a/Engine/Source/Runtime/Renderer/Private/VT/UniqueRequestList.h b/Engine/Source/Runtime/Renderer/Private/VT/UniqueRequestList.h
--- a/Engine/Source/Runtime/Renderer/Private/VT/UniqueRequestList.h
+++ b/Engine/Source/Runtime/Renderer/Private/VT/UniqueRequestList.h
@@ -187,7 +187,7 @@ private:
 	uint32 NumAdaptiveAllocationRequests;
 };
 
-uint16 FUniqueRequestList::AddLoadRequest(const FVirtualTextureLocalTileRequest& TileRequest, uint8 GroupMask, uint16 Count, bool bStreamingRequest)
+inline uint16 FUniqueRequestList::AddLoadRequest(const FVirtualTextureLocalTileRequest& TileRequest, uint8 GroupMask, uint16 Count, bool bStreamingRequest)
 {
 	const uint16 Hash = TileRequest.GetMurmurHash();
 	check(GroupMask != 0u);
@@ -218,7 +218,7 @@ uint16 FUniqueRequestList::AddLoadRequest(const FVirtualTextureLocalTileRequest&
 	return 0xffff;
 }
 
-uint16 FUniqueRequestList::LockLoadRequest(const FVirtualTextureLocalTileRequest& TileRequest, uint8 GroupMask, bool bStreamingRequest)
+inline uint16 FUniqueRequestList::LockLoadRequest(const FVirtualTextureLocalTileRequest& TileRequest, uint8 GroupMask, bool bStreamingRequest)
 {
 	const uint16 Hash = TileRequest.GetMurmurHash();
 	check(GroupMask != 0u);
@@ -258,7 +258,7 @@ uint16 FUniqueRequestList::LockLoadRequest(const FVirtualTextureLocalTileRequest
 	return 0xffff;
 }
 
-void FUniqueRequestList::AddMappingRequest(uint16 LoadRequestIndex, uint8 ProducerPhysicalGroupIndex, uint8 SpaceID, uint8 PageTableLayerIndex, uint32 MaxLevel, uint32 vAddress, uint8 vLevel, uint8 Local_vLevel)
+inline void FUniqueRequestList::AddMappingRequest(uint16 LoadRequestIndex, uint8 ProducerPhysicalGroupIndex, uint8 SpaceID, uint8 PageTableLayerIndex, uint32 MaxLevel, uint32 vAddress, uint8 vLevel, uint8 Local_vLevel)
 {
 	check(LoadRequestIndex < NumLoadRequests);
 	const FMappingRequest Request(LoadRequestIndex, ProducerPhysicalGroupIndex, SpaceID, PageTableLayerIndex, MaxLevel, vAddress, vLevel, Local_vLevel);
@@ -280,13 +280,13 @@ void FUniqueRequestList::AddMappingRequest(uint16 LoadRequestIndex, uint8 Produc
 	}
 }
 
-void FUniqueRequestList::AddDirectMappingRequest(uint8 InSpaceID, uint16 InPhysicalSpaceID, uint8 InPageTableLayerIndex, uint32 InMaxLevel, uint32 InAddress, uint8 InLevel, uint8 InLocalLevel, uint16 InPhysicalAddress)
+inline void FUniqueRequestList::AddDirectMappingRequest(uint8 InSpaceID, uint16 InPhysicalSpaceID, uint8 InPageTableLayerIndex, uint32 InMaxLevel, uint32 InAddress, uint8 InLevel, uint8 InLocalLevel, uint16 InPhysicalAddress)
 {
 	const FDirectMappingRequest Request(InSpaceID, InPhysicalSpaceID, InPageTableLayerIndex, InMaxLevel, InAddress, InLevel, InLocalLevel, InPhysicalAddress);
 	AddDirectMappingRequest(Request);
 }
 
-void FUniqueRequestList::AddDirectMappingRequest(const FDirectMappingRequest& Request)
+inline void FUniqueRequestList::AddDirectMappingRequest(const FDirectMappingRequest& Request)
 {
 	const uint16 Hash = Murmur32({ Request.PackedValue[0], Request.PackedValue[1], Request.PackedValue[2] });
 	for (uint16 Index = DirectMappingRequestHash.First(Hash); DirectMappingRequestHash.IsValid(Index); Index = DirectMappingRequestHash.Next(Index))
@@ -305,7 +305,7 @@ void FUniqueRequestList::AddDirectMappingRequest(const FDirectMappingRequest& Re
 	}
 }
 
-void FUniqueRequestList::AddContinuousUpdateRequest(const FVirtualTextureLocalTileRequest& Request)
+inline void FUniqueRequestList::AddContinuousUpdateRequest(const FVirtualTextureLocalTileRequest& Request)
 {
 	const uint16 Hash = Request.GetMurmurHash();
 	for (uint16 Index = ContinuousUpdateRequestHash.First(Hash); ContinuousUpdateRequestHash.IsValid(Index); Index = ContinuousUpdateRequestHash.Next(Index))
@@ -325,7 +325,7 @@ void FUniqueRequestList::AddContinuousUpdateRequest(const FVirtualTextureLocalTi
 	}
 }
 
-void FUniqueRequestList::AddAdaptiveAllocationRequest(uint32 Request)
+inline void FUniqueRequestList::AddAdaptiveAllocationRequest(uint32 Request)
 {
 	if (NumAdaptiveAllocationRequests < AdaptiveAllocationRequestCapacity)
 	{
@@ -333,7 +333,7 @@ void FUniqueRequestList::AddAdaptiveAllocationRequest(uint32 Request)
 	}
 }
 
-void FUniqueRequestList::MergeRequests(const FUniqueRequestList* RESTRICT Other, FConcurrentLinearBulkObjectAllocator& Allocator)
+inline void FUniqueRequestList::MergeRequests(const FUniqueRequestList* RESTRICT Other, FConcurrentLinearBulkObjectAllocator& Allocator)
 {
 	uint16* LoadRequestIndexRemap = Allocator.MallocArray<uint16>(Other->NumLoadRequests);
 
@@ -376,7 +376,7 @@ void FUniqueRequestList::MergeRequests(const FUniqueRequestList* RESTRICT Other,
 	}
 }
 
-void FUniqueRequestList::SortRequests(FVirtualTextureProducerCollection& Producers, FConcurrentLinearBulkObjectAllocator& Allocator, uint32 MaxNonStreamingLoadRequests, uint32 MaxStreamingLoadRequests, bool bUseCombinedLimit, bool bSortByPriority)
+inline void FUniqueRequestList::SortRequests(FVirtualTextureProducerCollection& Producers, FConcurrentLinearBulkObjectAllocator& Allocator, uint32 MaxNonStreamingLoadRequests, uint32 MaxStreamingLoadRequests, bool bUseCombinedLimit, bool bSortByPriority)
 {
 	if (bUseCombinedLimit)
 	{
diff --git a/Engine/Source/Runtime/Slate/Private/Framework/Docking/SPanelDrawerArea.cpp b/Engine/Source/Runtime/Slate/Private/Framework/Docking/SPanelDrawerArea.cpp
--- a/Engine/Source/Runtime/Slate/Private/Framework/Docking/SPanelDrawerArea.cpp
+++ b/Engine/Source/Runtime/Slate/Private/Framework/Docking/SPanelDrawerArea.cpp
@@ -3,7 +3,6 @@
 #include "Framework/Docking/SPanelDrawerArea.h"
 
 #include "Containers/Array.h"
-#include "Framework/Docking/DockingUtilsPrivate.h"
 #include "Framework/Docking/STabPanelDrawer.h"
 #include "Layout/Geometry.h"
 #include "Layout/WidgetPath.h"
@@ -21,6 +20,14 @@ namespace UE::Slate::Private
 constexpr float SSplitterHandleWith = 5.f;
 constexpr float AnimationLength = 0.15f;
 
+// Workaround for the web browser issues on some platform
+TAutoConsoleVariable<bool> CVarNoAnimationOnTabForegroundedEvent(
+	TEXT("Slate.NoAnimationOnTabForegroundedEvent"),
+	false,
+	TEXT("Tell the docking system to broadcast the event in a way that doesn't account for the animations")
+);
+
+
 void SPanelDrawerArea::Construct(const FArguments& InArgs, const TSharedRef<SWidget>& InContent)
 {
 	AreaContent = InContent;
diff --git a/Engine/Source/Runtime/TypedElementRuntime/Private/Elements/Framework/TypedElementSelectionSetLibrary.h b/Engine/Source/Runtime/TypedElementRuntime/Private/Elements/Framework/TypedElementSelectionSetLibrary.h
--- a/Engine/Source/Runtime/TypedElementRuntime/Private/Elements/Framework/TypedElementSelectionSetLibrary.h
+++ b/Engine/Source/Runtime/TypedElementRuntime/Private/Elements/Framework/TypedElementSelectionSetLibrary.h
@@ -13,7 +13,7 @@
 namespace UE::TypedElementFramework::Private
 {
 	// Return true if the selection set is invalid
-	bool CheckForInvalidSelectionSet(UTypedElementSelectionSet* SelectionSet)
+	inline bool CheckForInvalidSelectionSet(UTypedElementSelectionSet* SelectionSet)
 	{
 		if (!SelectionSet)
 		{
