Remote Control: fixed issue where RC is causing keyframe icons to disappear in details panel because it's holding a shared ptr to the original details panel created.
New details panel created for the first level editor details are never found by sequencer module because the original details is still in the list

#Synced-CL 51020850

diff --git a/Engine/Plugins/VirtualProduction/RemoteControl/Source/RemoteControlUI/Private/RemoteControlUIModule.cpp b/Engine/Plugins/VirtualProduction/RemoteControl/Source/RemoteControlUI/Private/RemoteControlUIModule.cpp
--- a/Engine/Plugins/VirtualProduction/RemoteControl/Source/RemoteControlUI/Private/RemoteControlUIModule.cpp
+++ b/Engine/Plugins/VirtualProduction/RemoteControl/Source/RemoteControlUI/Private/RemoteControlUIModule.cpp
@@ -360,30 +360,10 @@ TSharedRef<SRemoteControlPanel> FRemoteControlUIModule::CreateRemoteControlPanel
 
 	// NOTE : Reregister the module with the detail panel when this panel is created.
 
-	if (!SharedDetailsPanel.IsValid())
+	constexpr bool bResolveIfInvalid = true;
+	if (TSharedPtr<IDetailsView> DetailsPanel = GetDetailsView(bResolveIfInvalid, Preset->GetDetailsTabIdentifierOverrides()))
 	{
-		FPropertyEditorModule& PropertyEditor = FModuleManager::LoadModuleChecked<FPropertyEditorModule>("PropertyEditor");
-		TArray<FName> DetailsTabIdentifiers = Preset->GetDetailsTabIdentifierOverrides();
-
-		if (DetailsTabIdentifiers.Num() == 0)
-		{
-			DetailsTabIdentifiers.Append(DetailsTabIdentifiers_LevelEditor, DetailsTabIdentifiers_LevelEditor_Count);
-		}
-
-		for (const FName& DetailsTabIdentifier : DetailsTabIdentifiers)
-		{
-			SharedDetailsPanel = PropertyEditor.FindDetailView(DetailsTabIdentifier);
-
-			if (SharedDetailsPanel.IsValid())
-			{
-				RemoteControlUIModule::UpdateDetailViewExtensionWidth(SharedDetailsPanel, /*bInOnOpen*/ true);
-				break;
-			}
-		}
-	}
-	else
-	{
-		RemoteControlUIModule::UpdateDetailViewExtensionWidth(SharedDetailsPanel, /*bInOnOpen*/ true);
+		RemoteControlUIModule::UpdateDetailViewExtensionWidth(DetailsPanel, /*bInOnOpen*/ true);
 	}
 
 	OnRemoteControlPresetOpened().Broadcast(Preset);
@@ -445,6 +425,38 @@ void FRemoteControlUIModule::UnregisterEvents()
 	FEditorDelegates::PostUndoRedo.RemoveAll(this);
 }
 
+TSharedPtr<IDetailsView> FRemoteControlUIModule::GetDetailsView(bool bInResolveIfInvalid, TConstArrayView<FName> InDetailsTabIdentifierOverrides) const
+{
+	if (TSharedPtr<IDetailsView> ExistingDetailsView = DetailsPanelWeak.Pin())
+	{
+		return ExistingDetailsView;
+	}
+
+	if (!bInResolveIfInvalid)
+	{
+		// Skip resolve if saved details panel is destroyed/null
+		return nullptr;
+	}
+
+	FPropertyEditorModule& PropertyEditor = FModuleManager::LoadModuleChecked<FPropertyEditorModule>("PropertyEditor");
+
+	if (InDetailsTabIdentifierOverrides.IsEmpty())
+	{
+		InDetailsTabIdentifierOverrides = MakeArrayView(DetailsTabIdentifiers_LevelEditor, DetailsTabIdentifiers_LevelEditor_Count);
+	}
+
+	for (FName DetailsTabIdentifier : InDetailsTabIdentifierOverrides)
+	{
+		if (TSharedPtr<IDetailsView> DetailsPanel = PropertyEditor.FindDetailView(DetailsTabIdentifier))
+		{
+			DetailsPanelWeak = DetailsPanel;
+			return DetailsPanel;
+		}
+	}
+
+	return nullptr;
+}
+
 void FRemoteControlUIModule::ExtendPropertyRowContextMenu() const
 {
 	UToolMenus* Menus = UToolMenus::Get();
@@ -1187,9 +1199,9 @@ void FRemoteControlUIModule::UnregisterWidgetFactoryForType(UScriptStruct* Remot
 
 void FRemoteControlUIModule::HighlightPropertyInDetailsPanel(const FPropertyPath& Path) const
 {
-	if (SharedDetailsPanel)
+	if (TSharedPtr<IDetailsView> DetailsView = GetDetailsView())
 	{
-		SharedDetailsPanel->HighlightProperty(Path);
+		DetailsView->HighlightProperty(Path);
 	}
 }
 
@@ -1507,11 +1519,10 @@ UMeshComponent* FRemoteControlUIModule::GetSelectedMeshComponentToBeModified(UOb
 			SelectedObjects = DetailsView->GetSelectedObjects();
 		}
 	}
-	else if (SharedDetailsPanel.IsValid()) // Fallback to global detail panel reference if Detail Tree Node is invalid.
+	else if (TSharedPtr<IDetailsView> DetailsView = GetDetailsView()) // Fallback to global detail panel reference if Detail Tree Node is invalid.
 	{
-		SelectedActors = SharedDetailsPanel->GetSelectedActors();
-
-		SelectedObjects = SharedDetailsPanel->GetSelectedObjects();
+		SelectedActors = DetailsView->GetSelectedActors();
+		SelectedObjects = DetailsView->GetSelectedObjects();
 	}
 
 	UMeshComponent* MeshComponentToBeModified = nullptr;
@@ -1589,9 +1600,9 @@ void FRemoteControlUIModule::RefreshPanels()
 			DetailsView->ForceRefresh();
 		}
 	}
-	else if (SharedDetailsPanel.IsValid()) // Fallback to global detail panel reference if Detail Tree Node is invalid.
+	else if (TSharedPtr<IDetailsView> DetailsView = GetDetailsView()) // Fallback to global detail panel reference if Detail Tree Node is invalid.
 	{
-		SharedDetailsPanel->ForceRefresh();
+		DetailsView->ForceRefresh();
 	}
 
 	if (TSharedPtr<SRemoteControlPanel> Panel = GetPanelForObject(nullptr))
@@ -1811,10 +1822,12 @@ TSharedPtr<SRCPanelTreeNode> FRemoteControlUIModule::GenerateEntityWidget(const
 
 void FRemoteControlUIModule::UnregisterRemoteControlPanel(SRemoteControlPanel* Panel)
 {
-	if (SharedDetailsPanel.IsValid())
+	// When unregistering remote control panel, skip this step if the original details view is already destroyed (do not resolve to a new one)
+	constexpr bool bResolveIfInvalid = false;
+	if (TSharedPtr<IDetailsView> DetailsView = GetDetailsView(bResolveIfInvalid))
 	{
 		// reset column min width to standard size
-		RemoteControlUIModule::UpdateDetailViewExtensionWidth(SharedDetailsPanel, /*bInOnOpen*/ false);
+		RemoteControlUIModule::UpdateDetailViewExtensionWidth(DetailsView, /*bInOnOpen*/ false);
 	}
 
 	if (!Panel)
diff --git a/Engine/Plugins/VirtualProduction/RemoteControl/Source/RemoteControlUI/Private/RemoteControlUIModule.h b/Engine/Plugins/VirtualProduction/RemoteControl/Source/RemoteControlUI/Private/RemoteControlUIModule.h
--- a/Engine/Plugins/VirtualProduction/RemoteControl/Source/RemoteControlUI/Private/RemoteControlUIModule.h
+++ b/Engine/Plugins/VirtualProduction/RemoteControl/Source/RemoteControlUI/Private/RemoteControlUIModule.h
@@ -152,6 +152,9 @@ private:
 	void RegisterEvents();
 	void UnregisterEvents();
 
+	/** Retrieves the saved details view. If invalid, optionally resolve by finding a new details panel to store in DetailsPanelWeak and return that */
+	TSharedPtr<IDetailsView> GetDetailsView(bool bInResolveIfInvalid=true, TConstArrayView<FName> InDetailsTabIdentifierOverrides = {}) const;
+
 	/** Extend context row menu */
 	void ExtendPropertyRowContextMenu() const;
 
@@ -289,8 +292,8 @@ private:
 	/** Holds a weak ptr to the owner tree node of the active details panel. */
 	TWeakPtr<IDetailTreeNode> WeakDetailsTreeNode;
 
-	/** Holds a shared ptr to the global details panel. */
-	TSharedPtr<IDetailsView> SharedDetailsPanel;
+	/** Holds a weak ptr to the global details panel. */
+	mutable TWeakPtr<IDetailsView> DetailsPanelWeak;
 
 	/** Delegate called to gather extensions added externally to the panel. */
 	FOnGenerateExtensions ExtensionsGenerator;
