Remote Control: fix web socket function call to validate the function prior to invoking it.
+ exposed functions to the preset are also validated too

#Synced-CL 51441789

diff --git a/Engine/Plugins/VirtualProduction/RemoteControl/Source/WebRemoteControl/Private/WebRemoteControl.cpp b/Engine/Plugins/VirtualProduction/RemoteControl/Source/WebRemoteControl/Private/WebRemoteControl.cpp
--- a/Engine/Plugins/VirtualProduction/RemoteControl/Source/WebRemoteControl/Private/WebRemoteControl.cpp
+++ b/Engine/Plugins/VirtualProduction/RemoteControl/Source/WebRemoteControl/Private/WebRemoteControl.cpp
@@ -1224,6 +1224,17 @@ bool FWebRemoteControlModule::HandlePresetCallFunctionRoute(const FHttpServerReq
 				Call.TransactionMode = CallRequest.GenerateTransaction ? ERCTransactionMode::AUTOMATIC : ERCTransactionMode::NONE;
 				Call.ParamStruct = FStructOnScope(FunctionArgs.GetStruct(), FunctionArgs.GetStructMemory());
 
+				FString ErrorText;
+				if (!WebRemoteControlInternalUtils::ValidateFunctionCall(Call, &ErrorText))
+				{
+					const FString FormattedMsg = FString::Printf(TEXT("Web Remote Call deserialization error: %s"), *ErrorText);
+					UE_LOG(LogRemoteControl, Error, TEXT("%s"), *FormattedMsg);
+					IRemoteControlModule::BroadcastError(FormattedMsg);
+					WebRemoteControlInternalUtils::CreateUTF8ErrorMessage(ErrorText, Response->Body);
+					OnComplete(MoveTemp(Response));
+					return true; // handled
+				}
+
 				// Invoke call with replication payload
 				bSuccess &= IRemoteControlModule::Get().InvokeCall(Call, ERCPayloadType::Json, FunctionPayload);
 
diff --git a/Engine/Plugins/VirtualProduction/RemoteControl/Source/WebRemoteControl/Private/WebSocketMessageHandler.cpp b/Engine/Plugins/VirtualProduction/RemoteControl/Source/WebRemoteControl/Private/WebSocketMessageHandler.cpp
--- a/Engine/Plugins/VirtualProduction/RemoteControl/Source/WebRemoteControl/Private/WebSocketMessageHandler.cpp
+++ b/Engine/Plugins/VirtualProduction/RemoteControl/Source/WebRemoteControl/Private/WebSocketMessageHandler.cpp
@@ -765,6 +765,13 @@ void FWebSocketMessageHandler::HandleWebSocketFunctionCall(const FRemoteControlW
 		}
 	}
 
+	FString ValidateFunctionError;
+	if (!WebRemoteControlInternalUtils::ValidateFunctionCall(Call, &ValidateFunctionError))
+	{
+		UE_LOG(LogRemoteControl, Warning, TEXT("WebSocket Function call failed. Error: %s"), *ValidateFunctionError);
+		return;
+	}
+
 	if (!Call.IsValid())
 	{
 		return;
