Remote Control: fix PostInitProperties to account for async loading thread.

#Synced-CL 51764840

diff --git a/Engine/Plugins/VirtualProduction/RemoteControl/Source/RemoteControl/Private/RemoteControlPreset.cpp b/Engine/Plugins/VirtualProduction/RemoteControl/Source/RemoteControl/Private/RemoteControlPreset.cpp
--- a/Engine/Plugins/VirtualProduction/RemoteControl/Source/RemoteControl/Private/RemoteControlPreset.cpp
+++ b/Engine/Plugins/VirtualProduction/RemoteControl/Source/RemoteControl/Private/RemoteControlPreset.cpp
@@ -686,9 +686,15 @@ void URemoteControlPreset::PostInitProperties()
 	Super::PostInitProperties();
 	if (!HasAnyFlags(RF_ClassDefaultObject | RF_ArchetypeObject))
 	{
-		RegisterDelegates();
-
-		OnPostInitPropertiesRemoteControlPreset.Broadcast(this);
+		if (IsInAsyncLoadingThread())
+		{
+			bPendingPostInitNotify = true;
+		}
+		else
+		{
+			RegisterDelegates();
+			OnPostInitPropertiesRemoteControlPreset.Broadcast(this);
+		}
 	}
 }
 
@@ -707,6 +713,11 @@ void URemoteControlPreset::PostLoad()
 {
 	Super::PostLoad();
 
+	if (bPendingPostInitNotify)
+	{
+		OnPostInitPropertiesRemoteControlPreset.Broadcast(this);
+		bPendingPostInitNotify = false;
+	}
 	OnPostLoadRemoteControlPreset.Broadcast(this);
 
 	RegisterDelegates();
diff --git a/Engine/Plugins/VirtualProduction/RemoteControl/Source/RemoteControl/Public/RemoteControlPreset.h b/Engine/Plugins/VirtualProduction/RemoteControl/Source/RemoteControl/Public/RemoteControlPreset.h
--- a/Engine/Plugins/VirtualProduction/RemoteControl/Source/RemoteControl/Public/RemoteControlPreset.h
+++ b/Engine/Plugins/VirtualProduction/RemoteControl/Source/RemoteControl/Public/RemoteControlPreset.h
@@ -1069,6 +1069,9 @@ private:
 	/** Used for OnObjectPropertyChanged to avoid being called by itself. */
 	bool bReentryGuard = false;
 
+	/** Whether OnPostInitPropertiesRemoteControlProperties delegate is pending broadcast because it wasn't broadcast in URemoteControlPreset::PostInitProperties() */
+	bool bPendingPostInitNotify = false;
+
 	/** Transient list of the currently bound entity ids and the number of times they are bound. */
 	TMap<FGuid, int32> BoundEntityIds;
 
