Remote Control: removed redundant costly pre/post edit change calls in SetAssetByPath behavior.
And removed SetStaticMeshValue as SetPropertyValue already correctly calls UStaticMeshComponent::SetStaticMesh, while also supporting other static mesh properties.

#Synced-CL 52400047

diff --git a/Engine/Plugins/VirtualProduction/RemoteControl/Source/RemoteControl/Private/RemoteControlPropertyHandle.cpp b/Engine/Plugins/VirtualProduction/RemoteControl/Source/RemoteControl/Private/RemoteControlPropertyHandle.cpp
--- a/Engine/Plugins/VirtualProduction/RemoteControl/Source/RemoteControl/Private/RemoteControlPropertyHandle.cpp
+++ b/Engine/Plugins/VirtualProduction/RemoteControl/Source/RemoteControl/Private/RemoteControlPropertyHandle.cpp
@@ -1878,11 +1878,6 @@ bool FRemoteControlPropertyHandleObject::SetValue(UObject* InValue)
 	{
 		return false;
 	}
-	
-	if (InValue->IsA(UStaticMesh::StaticClass()))
-	{
-		return SetStaticMeshValue(InValue);
-	}
 	return SetPropertyValue(*this, InValue);
 }
 
@@ -1930,34 +1925,6 @@ bool FRemoteControlPropertyHandleObject::GetValueInArray(UObject*& OutValue, int
 	return false;
 }
 
-bool FRemoteControlPropertyHandleObject::SetStaticMeshValue(UObject* InValue) const
-{
-	if (!InValue)
-	{
-		return false;
-	}
-
-	if (!InValue->IsA(UStaticMesh::StaticClass()))
-	{
-		return false;
-	}
-	
-	bool bSuccess = false;
-	
-	if (const TSharedPtr<FRemoteControlProperty>& RemoteControlPropertyPtr = GetRCProperty())
-	{
-		for (UObject* Object : RemoteControlPropertyPtr->GetBoundObjects())
-		{							
-			if (UStaticMeshComponent* StaticMeshComponent = Cast<UStaticMeshComponent>(Object))
-			{
-				bSuccess = StaticMeshComponent->SetStaticMesh(Cast<UStaticMesh>(InValue));
-			}			
-		}
-	}
-
-	return bSuccess;
-}
-
 bool FRemoteControlPropertyHandleObject::SetMaterialValue(UObject* InValue, int32 InIndex) const
 {
 	if (!InValue || InIndex < 0)
diff --git a/Engine/Plugins/VirtualProduction/RemoteControl/Source/RemoteControl/Private/RemoteControlPropertyHandle.h b/Engine/Plugins/VirtualProduction/RemoteControl/Source/RemoteControl/Private/RemoteControlPropertyHandle.h
--- a/Engine/Plugins/VirtualProduction/RemoteControl/Source/RemoteControl/Private/RemoteControlPropertyHandle.h
+++ b/Engine/Plugins/VirtualProduction/RemoteControl/Source/RemoteControl/Private/RemoteControlPropertyHandle.h
@@ -619,9 +619,6 @@ public:
 	virtual bool GetValueInArray(UObject*& OutValue, int32 InIndex) const override;
 	//~ End IRemoteControlPropertyHandle Interface
 
-	/** Set a Static Mesh value using the handled Object */
-	bool SetStaticMeshValue(UObject* InValue) const;
-
 	/** Set a Material value using the handled Object */
 	bool SetMaterialValue(UObject* InValue, int32 InIndex) const;
 };
diff --git a/Engine/Plugins/VirtualProduction/RemoteControl/Source/RemoteControlLogic/Private/Action/Path/RCSetAssetByPathActionNew.cpp b/Engine/Plugins/VirtualProduction/RemoteControl/Source/RemoteControlLogic/Private/Action/Path/RCSetAssetByPathActionNew.cpp
--- a/Engine/Plugins/VirtualProduction/RemoteControl/Source/RemoteControlLogic/Private/Action/Path/RCSetAssetByPathActionNew.cpp
+++ b/Engine/Plugins/VirtualProduction/RemoteControl/Source/RemoteControlLogic/Private/Action/Path/RCSetAssetByPathActionNew.cpp
@@ -102,23 +102,7 @@ bool URCSetAssetByPathActionNew::UpdateInternal()
 			return FText::FromString(FString("Path Behavior attempts to set Asset to %s") + *MainObjectRef.GetAssetPathString());
 		});
 
-#if !WITH_EDITOR
 	return SetInternalAsset(Object);
-#else
-	FEditPropertyChain PropertyChain;
-	PropertyChain.AddHead(Property);
-	RemoteControlProperty->GetBoundObject()->PreEditChange(PropertyChain);
-
-	if (SetInternalAsset(Object))
-	{
-		FPropertyChangedEvent ChangedEvent(Property, EPropertyChangeType::ValueSet);
-		RemoteControlProperty->GetBoundObject()->PostEditChangeProperty(ChangedEvent);
-
-		return true;
-	}
-
-	return false;
-#endif
 }
 
 bool URCSetAssetByPathActionNew::UpdateExternal()
